Horde: Improvements to workspace management/patching.

- Patching workspaces now reuses existing data in workspaces.
- Added commands to HordeCmd for manipulating layers in managed workspaces.
- Added tests for patching.

#rnx

diff --git a/Engine/Source/Programs/Horde/HordeCmd/Commands/Profile/ProfileChunking.cs b/Engine/Source/Programs/Horde/HordeCmd/Commands/Profile/ProfileChunking.cs
new file mode 100644
--- /dev/null
+++ b/Engine/Source/Programs/Horde/HordeCmd/Commands/Profile/ProfileChunking.cs
@@ -0,0 +1,168 @@
+// Copyright Epic Games, Inc. All Rights Reserved.
+
+using System.ComponentModel;
+using System.Diagnostics;
+using System.Threading.Channels;
+using EpicGames.Core;
+using EpicGames.Horde.Storage;
+using EpicGames.Horde.Storage.Bundles;
+using EpicGames.Horde.Storage.Nodes;
+using Microsoft.Extensions.Logging;
+
+namespace Horde.Commands.Profile
+{
+	[Command("profile", "chunking", "Profiles chunking over a source directory")]
+	class ProfileChunking : StorageCommandBase
+	{
+		record class ChunkInfo(long Offset, int Length, IoHash Hash);
+
+		[CommandLine("-Input=", Required = true)]
+		[Description("Input file or directory")]
+		public string Input { get; set; } = null!;
+
+		[CommandLine("-TraceDir=", Required = true)]
+		[Description("Directory to write traces to")]
+		public DirectoryReference TraceDir { get; set; } = null!;
+
+		[CommandLine("-Verify")]
+		public bool Verify { get; set; }
+
+		public ProfileChunking(HttpStorageClient storageClient, BundleCache bundleCache)
+			: base(storageClient, bundleCache)
+		{
+		}
+
+		public override async Task<int> ExecuteAsync(ILogger logger)
+		{
+			DirectoryReference baseDir = new DirectoryReference(Input);
+			List<FileReference> files = DirectoryReference.EnumerateFiles(baseDir, "*", SearchOption.AllDirectories).ToList();
+
+			// Create the bundle
+			logger.LogInformation("Creating archive from files in {BaseDir}...", baseDir);
+
+			List<ChunkInfo>? referenceChunks = null;
+			for (int idx = 0; ; idx++)
+			{
+				logger.LogInformation("PARALLEL ATTEMPT {Idx}", idx);
+				logger.LogInformation("");
+
+				int suffix = Math.Min(idx, 2);
+				List<ChunkInfo> chunks = await ChunkDataAsync(FileReference.Combine(TraceDir, $"output-{suffix}.txt"), files, false, logger);
+
+				logger.LogInformation("-----");
+
+				if (Verify)
+				{
+					if (referenceChunks == null)
+					{
+						logger.LogInformation("CREATING REFERENCE DATA");
+						logger.LogInformation("");
+						referenceChunks = await ChunkDataAsync(FileReference.Combine(TraceDir, "reference.txt"), files, true, logger);
+					}
+
+					if (!Enumerable.SequenceEqual(chunks, referenceChunks))
+					{
+						logger.LogError("MISMATCH");
+						break;
+					}
+				}
+			}
+
+			return 0;
+		}
+
+		static async Task ReadFilesAsync(List<FileReference> files, ChannelWriter<ChunkerSource> writer, CancellationToken cancellationToken)
+		{
+			foreach (FileReference file in files)
+			{
+				FileInfo fileInfo = file.ToFileInfo();
+				await writer.WriteAsync(new FileChunkerSource(fileInfo, null), cancellationToken);
+			}
+			writer.Complete();
+		}
+
+		class ChunkReader
+		{
+			public int _filesRead;
+			public int _chunksRead;
+			public long _bytesRead;
+
+			readonly ILogger _logger;
+
+			public List<ChunkInfo> _chunks = new List<ChunkInfo>();
+
+			public ChunkReader(ILogger logger)
+				=> _logger = logger;
+
+			public async Task ReadChunksAsync(ChannelReader<ChunkerOutput> reader, CancellationToken cancellationToken)
+			{
+				long lastBytesRead = 0;
+				while (await reader.WaitToReadAsync(cancellationToken))
+				{
+					if (reader.TryRead(out ChunkerOutput? output))
+					{
+						Interlocked.Increment(ref _filesRead);
+						while (await output.MoveNextAsync(cancellationToken))
+						{
+							_chunks.Add(new ChunkInfo(_bytesRead, (int)output.Data.Length, IoHash.Compute(output.Data.Span)));
+
+							Interlocked.Increment(ref _chunksRead);
+							long bytesRead = Interlocked.Add(ref _bytesRead, output.Data.Length);
+
+							if (bytesRead > lastBytesRead + (1024 * 1024 * 256))
+							{
+								_logger.LogInformation("Read {Size:n0}", bytesRead);
+								lastBytesRead = bytesRead;
+							}
+						}
+					}
+				}
+			}
+		}
+
+		static async Task<List<ChunkInfo>> ChunkDataAsync(FileReference file, List<FileReference> files, bool useSerialPipeline, ILogger logger)
+		{
+			byte[] buffer = new byte[2 * 1024 * 1024];
+			Stopwatch timer = Stopwatch.StartNew();
+
+			await using (AsyncPipeline pipeline = new AsyncPipeline(CancellationToken.None))
+			{
+				LeafChunkedDataNodeOptions options = LeafChunkedDataNodeOptions.Default;// new LeafChunkedDataNodeOptions(4 * 1024, 8 * 1024, 16 * 1024);
+
+				ContentChunker cdc;
+				if (useSerialPipeline)
+				{
+					cdc = new SerialBuzHashChunker(pipeline, buffer, options);
+				}
+				else
+				{
+					cdc = new ParallelBuzHashChunker(pipeline, buffer, options);
+				}
+
+				_ = pipeline.AddTask(ctx => ReadFilesAsync(files, cdc.SourceWriter, ctx));
+
+				ChunkReader reader = new ChunkReader(logger);
+				_ = pipeline.AddTask(ctx => reader.ReadChunksAsync(cdc.OutputReader, ctx));
+
+				await pipeline.WaitForCompletionAsync();
+
+				logger.LogInformation("Files: {Files:n0}", reader._filesRead);
+				logger.LogInformation("Chunks: {Chunks:n0}", reader._chunksRead);
+				logger.LogInformation("Bytes: {Bytes:n0}", reader._bytesRead);
+				logger.LogInformation("Speed: {Rate:n0} bytes/sec", reader._bytesRead / timer.Elapsed.TotalSeconds);
+				logger.LogInformation("Time: {Time:n0} sec", timer.Elapsed.TotalSeconds);
+
+				using (StreamWriter writer = new StreamWriter(file.FullName))
+				{
+					foreach (ChunkInfo chunkInfo in reader._chunks)
+					{
+						await writer.WriteLineAsync($"{chunkInfo.Offset},{chunkInfo.Length},{chunkInfo.Hash}");
+					}
+				}
+
+				return reader._chunks;
+			}
+		}
+	}
+}
+
diff --git a/Engine/Source/Programs/Horde/HordeCmd/Commands/Vcs/VcsCheckout.cs b/Engine/Source/Programs/Horde/HordeCmd/Commands/Vcs/VcsCheckout.cs
--- a/Engine/Source/Programs/Horde/HordeCmd/Commands/Vcs/VcsCheckout.cs
+++ b/Engine/Source/Programs/Horde/HordeCmd/Commands/Vcs/VcsCheckout.cs
@@ -117,7 +117,7 @@ namespace Horde.Commands.Vcs
 		static async Task<FileState> CheckoutFileAsync(FileEntry fileRef, FileInfo fileInfo, ILogger logger)
 		{
 			logger.LogInformation("Updating {File} to {Hash}", fileInfo, fileRef.Hash);
-			await fileRef.CopyToFileAsync(fileInfo, CancellationToken.None);
+			await fileRef.Target.CopyToFileAsync(fileInfo, CancellationToken.None);
 			fileInfo.Refresh();
 			return new FileState(fileInfo, fileRef.Hash);
 		}
diff --git a/Engine/Source/Programs/Horde/HordeCmd/Commands/Workspace/WorkspaceClean.cs b/Engine/Source/Programs/Horde/HordeCmd/Commands/Workspace/WorkspaceClean.cs
new file mode 100644
--- /dev/null
+++ b/Engine/Source/Programs/Horde/HordeCmd/Commands/Workspace/WorkspaceClean.cs
@@ -0,0 +1,16 @@
+// Copyright Epic Games, Inc. All Rights Reserved.
+
+using EpicGames.Core;
+using EpicGames.Horde.Storage;
+
+namespace Horde.Commands.Workspace
+{
+	using Workspace = EpicGames.Horde.Storage.Workspace;
+
+	[Command("workspace", "clean", "Removes all untracked files from the workspace")]
+	class WorkspaceClean : WorkspaceStatusBase
+	{
+		protected override Task ExecuteInternalAsync(Workspace workspace, FileFilter? filter, WorkspaceStatusFlags flags, CancellationToken cancellationToken)
+			=> workspace.CleanAsync(filter, flags, cancellationToken: cancellationToken);
+	}
+}
diff --git a/Engine/Source/Programs/Horde/HordeCmd/Commands/Workspace/WorkspaceDrop.cs b/Engine/Source/Programs/Horde/HordeCmd/Commands/Workspace/WorkspaceDrop.cs
--- a/Engine/Source/Programs/Horde/HordeCmd/Commands/Workspace/WorkspaceDrop.cs
+++ b/Engine/Source/Programs/Horde/HordeCmd/Commands/Workspace/WorkspaceDrop.cs
@@ -3,7 +3,6 @@
 using System.ComponentModel;
 using EpicGames.Core;
 using EpicGames.Horde.Storage;
-using EpicGames.Horde.Storage.Nodes;
 using Microsoft.Extensions.Logging;
 
 namespace Horde.Commands.Workspace
@@ -33,7 +32,7 @@ namespace Horde.Commands.Workspace
 				return 1;
 			}
 
-			await workspace.SyncAsync(LayerId, "", new DirectoryNode(), cancellationToken);
+			await workspace.SyncAsync(LayerId, null, cancellationToken: cancellationToken);
 			workspace.RemoveLayer(LayerId);
 			await workspace.SaveAsync(cancellationToken);
 
diff --git a/Engine/Source/Programs/Horde/HordeCmd/Commands/Workspace/WorkspaceForget.cs b/Engine/Source/Programs/Horde/HordeCmd/Commands/Workspace/WorkspaceForget.cs
new file mode 100644
--- /dev/null
+++ b/Engine/Source/Programs/Horde/HordeCmd/Commands/Workspace/WorkspaceForget.cs
@@ -0,0 +1,45 @@
+// Copyright Epic Games, Inc. All Rights Reserved.
+
+using System.ComponentModel;
+using EpicGames.Core;
+using EpicGames.Horde.Storage;
+using Microsoft.Extensions.Logging;
+
+namespace Horde.Commands.Workspace
+{
+	using Workspace = EpicGames.Horde.Storage.Workspace;
+
+	[Command("workspace", "forget", "Removes files in the workspace from tracking", AcceptsPositionalArguments = true)]
+	class WorkspaceForget : Command
+	{
+		[CommandLine("-Root=")]
+		[Description("Root directory for the workspace.")]
+		public DirectoryReference? RootDir { get; set; }
+
+		[CommandLine("-Layer=")]
+		[Description("Layer to remove files from")]
+		public WorkspaceLayerId LayerId { get; set; } = WorkspaceLayerId.Default;
+
+		[CommandLine("-Path=", Positional = true, Required = true)]
+		[Description("Filter for the files to remove")]
+		List<string> Paths { get; set; } = new List<string>();
+
+		public override async Task<int> ExecuteAsync(ILogger logger)
+		{
+			CancellationToken cancellationToken = CancellationToken.None;
+			RootDir ??= DirectoryReference.GetCurrentDirectory();
+
+			Workspace? workspace = await Workspace.TryFindAndOpenAsync(RootDir, logger, cancellationToken);
+			if (workspace == null)
+			{
+				logger.LogError("No workspace has been initialized in {RootDir}. Use 'workspace init' to create a new workspace.", RootDir);
+				return 1;
+			}
+
+			FileFilter filter = WorkspaceTrack.CreateFilter(RootDir, Paths);
+			await workspace.ForgetAsync(LayerId, filter, cancellationToken);
+
+			return 0;
+		}
+	}
+}
diff --git a/Engine/Source/Programs/Horde/HordeCmd/Commands/Workspace/WorkspaceLayers.cs b/Engine/Source/Programs/Horde/HordeCmd/Commands/Workspace/WorkspaceLayers.cs
new file mode 100644
--- /dev/null
+++ b/Engine/Source/Programs/Horde/HordeCmd/Commands/Workspace/WorkspaceLayers.cs
@@ -0,0 +1,41 @@
+// Copyright Epic Games, Inc. All Rights Reserved.
+
+using System.ComponentModel;
+using EpicGames.Core;
+using EpicGames.Horde.Storage;
+using Microsoft.Extensions.Logging;
+
+namespace Horde.Commands.Workspace
+{
+	using Workspace = EpicGames.Horde.Storage.Workspace;
+
+	[Command("workspace", "layers", "Lists all layers in the workspace")]
+	class WorkspaceLayers : Command
+	{
+		[CommandLine("-Root=")]
+		[Description("Root directory for the workspace.")]
+		public DirectoryReference? RootDir { get; set; }
+
+		public override async Task<int> ExecuteAsync(ILogger logger)
+		{
+			CancellationToken cancellationToken = CancellationToken.None;
+			RootDir ??= DirectoryReference.GetCurrentDirectory();
+
+			Workspace? workspace = await Workspace.TryFindAndOpenAsync(RootDir, logger, cancellationToken);
+			if (workspace == null)
+			{
+				logger.LogError("No workspace has been initialized in {RootDir}. Use 'workspace init' to create a new workspace.", RootDir);
+				return 1;
+			}
+
+			logger.LogInformation("Layers:");
+			logger.LogInformation("  {LayerId}", WorkspaceLayerId.Default);
+			foreach (WorkspaceLayerId layerId in workspace.CustomLayers)
+			{
+				logger.LogInformation("  {LayerId}", layerId);
+			}
+
+			return 0;
+		}
+	}
+}
diff --git a/Engine/Source/Programs/Horde/HordeCmd/Commands/Workspace/WorkspaceStatus.cs b/Engine/Source/Programs/Horde/HordeCmd/Commands/Workspace/WorkspaceStatus.cs
new file mode 100644
--- /dev/null
+++ b/Engine/Source/Programs/Horde/HordeCmd/Commands/Workspace/WorkspaceStatus.cs
@@ -0,0 +1,126 @@
+// Copyright Epic Games, Inc. All Rights Reserved.
+
+using System.ComponentModel;
+using EpicGames.Core;
+using EpicGames.Horde.Storage;
+using Microsoft.Extensions.Logging;
+
+namespace Horde.Commands.Workspace
+{
+	using Workspace = EpicGames.Horde.Storage.Workspace;
+
+	[Command("workspace", "status", "Prints the current status of the workspace", AcceptsPositionalArguments = true)]
+	class WorkspaceStatus : WorkspaceStatusBase
+	{
+		protected override Task ExecuteInternalAsync(Workspace workspace, FileFilter? filter, WorkspaceStatusFlags flags, CancellationToken cancellationToken)
+			=> workspace.StatusAsync(filter, flags, cancellationToken);
+	}
+
+	abstract class WorkspaceStatusBase : Command
+	{
+		[CommandLine("-Root=")]
+		[Description("Root directory for the workspace.")]
+		public DirectoryReference? RootDir { get; set; }
+
+		[CommandLine("-Path=", Positional = true)]
+		[Description("Path of the files to query")]
+		public List<string> Paths { get; set; } = new List<string>();
+
+		[CommandLine("-All")]
+		[Description("Show all files")]
+		public bool All { get; set; }
+
+		[CommandLine("-Added")]
+		[Description("Show added files")]
+		public bool Added { get; set; }
+
+		[CommandLine("-Modified")]
+		[Description("Show modified files")]
+		public bool Modified { get; set; }
+
+		[CommandLine("-Unmodified")]
+		[Description("Show all unmodified files")]
+		public bool Unmodified { get; set; }
+
+		[CommandLine("-Removed")]
+		[Description("Show removed files")]
+		public bool Removed { get; set; }
+
+		[CommandLine("-IgnoreAdded")]
+		[Description("Ignore added files")]
+		public bool IgnoreAdded { get; set; }
+
+		[CommandLine("-IgnoreModified")]
+		[Description("Ignore modified files")]
+		public bool IgnoreModified { get; set; }
+
+		[CommandLine("-IgnoreUnmodified")]
+		[Description("Ignore unmodified files")]
+		public bool IgnoreUnmodified { get; set; }
+
+		[CommandLine("-IgnoreRemoved")]
+		[Description("Ignore removed files")]
+		public bool IgnoreRemoved { get; set; }
+
+		public override async Task<int> ExecuteAsync(ILogger logger)
+		{
+			CancellationToken cancellationToken = CancellationToken.None;
+			RootDir ??= DirectoryReference.GetCurrentDirectory();
+
+			Workspace? workspace = await Workspace.TryFindAndOpenAsync(RootDir, logger, cancellationToken);
+			if (workspace == null)
+			{
+				logger.LogError("No workspace has been initialized in {RootDir}. Use 'workspace init' to create a new workspace.", RootDir);
+				return 1;
+			}
+
+			WorkspaceStatusFlags flags = WorkspaceStatusFlags.None;
+			if (All)
+			{
+				flags |= WorkspaceStatusFlags.FileAdded | WorkspaceStatusFlags.FileModified | WorkspaceStatusFlags.FileRemoved | WorkspaceStatusFlags.FileUnmodified | WorkspaceStatusFlags.DirectoryAdded | WorkspaceStatusFlags.DirectoryRemoved;
+			}
+
+			if (Added)
+			{
+				flags |= WorkspaceStatusFlags.FileAdded | WorkspaceStatusFlags.DirectoryAdded;
+			}
+			if (Modified)
+			{
+				flags |= WorkspaceStatusFlags.FileModified;
+			}
+			if (Removed)
+			{
+				flags |= WorkspaceStatusFlags.FileRemoved | WorkspaceStatusFlags.DirectoryRemoved;
+			}
+			if (Unmodified)
+			{
+				flags |= WorkspaceStatusFlags.FileUnmodified;
+			}
+
+			if (flags == WorkspaceStatusFlags.None)
+			{
+				flags = WorkspaceStatusFlags.Default;
+			}
+
+			if (IgnoreAdded)
+			{
+				flags &= ~WorkspaceStatusFlags.FileAdded | WorkspaceStatusFlags.DirectoryAdded;
+			}
+			if (IgnoreModified)
+			{
+				flags &= ~WorkspaceStatusFlags.FileModified;
+			}
+			if (IgnoreRemoved)
+			{
+				flags &= ~WorkspaceStatusFlags.FileRemoved | WorkspaceStatusFlags.DirectoryRemoved;
+			}
+
+			FileFilter filter = WorkspaceTrack.CreateFilter(RootDir, Paths);
+			await ExecuteInternalAsync(workspace, filter, flags, cancellationToken);
+
+			return 0;
+		}
+
+		protected abstract Task ExecuteInternalAsync(Workspace workspace, FileFilter? filter, WorkspaceStatusFlags flags, CancellationToken cancellationToken);
+	}
+}
diff --git a/Engine/Source/Programs/Horde/HordeCmd/Commands/Workspace/WorkspaceSync.cs b/Engine/Source/Programs/Horde/HordeCmd/Commands/Workspace/WorkspaceSync.cs
--- a/Engine/Source/Programs/Horde/HordeCmd/Commands/Workspace/WorkspaceSync.cs
+++ b/Engine/Source/Programs/Horde/HordeCmd/Commands/Workspace/WorkspaceSync.cs
@@ -20,9 +20,9 @@ namespace Horde.Commands.Workspace
 		[Description("Root directory for the managed workspace.")]
 		public DirectoryReference? RootDir { get; set; }
 
-		[CommandLine("-Dir=")]
+		[CommandLine("-Target=")]
 		[Description("Directory to extract files to. Must be under the workspace root.")]
-		public DirectoryReference? OutputDir { get; set; }
+		public DirectoryReference? TargetDir { get; set; }
 
 		[CommandLine("-Artifact=")]
 		[Description("Artifact to extract into the workspace")]
@@ -71,32 +71,32 @@ namespace Horde.Commands.Workspace
 				logger.LogInformation("Downloading artifact {Id}: {Description}", artifactId, artifact.Description);
 
 				IStorageNamespace store = _hordeClient.GetStorageNamespace(artifact.Id);
-				IBlobRef handle = await store.ReadRefAsync<DirectoryNode>(new RefName("default"));
+				IBlobRef<DirectoryNode> handle = await store.ReadRefAsync<DirectoryNode>(new RefName("default"));
 				return await ExecuteInternalAsync(store, handle, logger);
 			}
 			else
 			{
 				IStorageNamespace store = _hordeClient.GetStorageNamespace(new NamespaceId(Namespace));
 
-				IBlobRef handle;
+				IBlobRef<DirectoryNode> handle;
 				if (Ref != null)
 				{
-					handle = await store.ReadRefAsync(new RefName(Ref));
+					handle = await store.ReadRefAsync<DirectoryNode>(new RefName(Ref));
 				}
 				else if (Node != null)
 				{
-					handle = store.CreateBlobRef(new BlobLocator(Node));
+					handle = store.CreateBlobRef<DirectoryNode>(new BlobLocator(Node));
 				}
 				else
 				{
-					throw new CommandLineArgumentException("Either -File=... or -Ref=... must be specified");
+					throw new CommandLineArgumentException("Missing -Artifact=..., -File=... or -Ref=... option");
 				}
 
 				return await ExecuteInternalAsync(store, handle, logger);
 			}
 		}
 
-		async Task<int> ExecuteInternalAsync(IStorageNamespace store, IBlobRef handle, ILogger logger)
+		async Task<int> ExecuteInternalAsync(IStorageNamespace store, IBlobRef<DirectoryNode> handle, ILogger logger)
 		{
 			CancellationToken cancellationToken = CancellationToken.None;
 
@@ -107,7 +107,7 @@ namespace Horde.Commands.Workspace
 			}
 			else
 			{
-				DirectoryReference workspaceDir = OutputDir ?? DirectoryReference.GetCurrentDirectory();
+				DirectoryReference workspaceDir = TargetDir ?? DirectoryReference.GetCurrentDirectory();
 
 				workspace = await Workspace.TryFindAndOpenAsync(workspaceDir, logger, cancellationToken);
 				if (workspace == null)
@@ -118,25 +118,24 @@ namespace Horde.Commands.Workspace
 			}
 
 			string mountPath = "";
-			if (OutputDir != null)
+			if (TargetDir != null)
 			{
-				if (!OutputDir.IsUnderDirectory(workspace.RootDir))
+				if (!TargetDir.IsUnderDirectory(workspace.RootDir))
 				{
-					logger.LogError("Output directory ({OutputDir}) is not under workspace root ({RootDir})", OutputDir, workspace.RootDir);
+					logger.LogError("Output directory ({OutputDir}) is not under workspace root ({RootDir})", TargetDir, workspace.RootDir);
 					return 1;
 				}
-				mountPath = OutputDir.MakeRelativeTo(workspace.RootDir);
+				mountPath = TargetDir.MakeRelativeTo(workspace.RootDir);
 			}
 
 			Stopwatch timer = Stopwatch.StartNew();
 
 			logger.LogInformation("Refreshing metadata...");
-			await workspace.RefreshAsync(cancellationToken);
+			await workspace.ReconcileAsync(cancellationToken);
 
 			logger.LogInformation("Syncing into layer '{LayerId}'...", LayerId);
 
-			DirectoryNode contents = await handle.ReadBlobAsync<DirectoryNode>();
-			await workspace.SyncAsync(LayerId, mountPath, contents, cancellationToken);
+			await workspace.SyncAsync(LayerId, mountPath, handle, cancellationToken: cancellationToken);
 			await workspace.SaveAsync(cancellationToken);
 
 			logger.LogInformation("Elapsed: {Time}s", timer.Elapsed.TotalSeconds);
diff --git a/Engine/Source/Programs/Horde/HordeCmd/Commands/Workspace/WorkspaceTrack.cs b/Engine/Source/Programs/Horde/HordeCmd/Commands/Workspace/WorkspaceTrack.cs
new file mode 100644
--- /dev/null
+++ b/Engine/Source/Programs/Horde/HordeCmd/Commands/Workspace/WorkspaceTrack.cs
@@ -0,0 +1,73 @@
+// Copyright Epic Games, Inc. All Rights Reserved.
+
+using System.ComponentModel;
+using EpicGames.Core;
+using EpicGames.Horde.Storage;
+using Microsoft.Extensions.Logging;
+
+namespace Horde.Commands.Workspace
+{
+	using Workspace = EpicGames.Horde.Storage.Workspace;
+
+	[Command("workspace", "track", "Tracks files in their current state with the workspace", AcceptsPositionalArguments = true)]
+	class WorkspaceTrack : Command
+	{
+		[CommandLine("-Root=")]
+		[Description("Root directory for the workspace.")]
+		public DirectoryReference? RootDir { get; set; }
+
+		[CommandLine("-Layer=")]
+		[Description("Layer to add new files to")]
+		public WorkspaceLayerId LayerId { get; set; } = WorkspaceLayerId.Default;
+
+		[CommandLine("-Path=", Positional = true)]
+		[Description("List of files to reconcile")]
+		public List<string> Paths { get; set; } = new List<string>();
+
+		public override async Task<int> ExecuteAsync(ILogger logger)
+		{
+			CancellationToken cancellationToken = CancellationToken.None;
+			RootDir ??= DirectoryReference.GetCurrentDirectory();
+
+			Workspace? workspace = await Workspace.TryFindAndOpenAsync(RootDir, logger, cancellationToken);
+			if (workspace == null)
+			{
+				logger.LogError("No workspace has been initialized in {RootDir}. Use 'workspace init' to create a new workspace.", RootDir);
+				return 1;
+			}
+
+			workspace.AddLayer(LayerId);
+
+			FileFilter filter = CreateFilter(RootDir, Paths);
+			await workspace.TrackAsync(LayerId, filter, cancellationToken);
+
+			return 0;
+		}
+
+		public static FileFilter CreateFilter(DirectoryReference rootDir, List<string> paths)
+		{
+			if (paths.Count == 0)
+			{
+				return new FileFilter(FileFilterType.Include);
+			}
+
+			FileFilter filter = new FileFilter();
+			foreach (string path in paths)
+			{
+				string trimPath = path.TrimStart('/');
+				try
+				{
+					DirectoryReference subDir = DirectoryReference.Combine(rootDir, path);
+					if (DirectoryReference.Exists(subDir))
+					{
+						trimPath = $"{trimPath}/...";
+					}
+				}
+				catch { }
+
+				filter.AddRule($"/{trimPath}", FileFilterType.Include);
+			}
+			return filter;
+		}
+	}
+}
diff --git a/Engine/Source/Programs/Shared/EpicGames.Core/ChannelWithRefCount.cs b/Engine/Source/Programs/Shared/EpicGames.Core/ChannelWithRefCount.cs
new file mode 100644
--- /dev/null
+++ b/Engine/Source/Programs/Shared/EpicGames.Core/ChannelWithRefCount.cs
@@ -0,0 +1,63 @@
+// Copyright Epic Games, Inc. All Rights Reserved.
+
+using System;
+using System.Threading;
+using System.Threading.Channels;
+
+namespace EpicGames.Core
+{
+	/// <summary>
+	/// Implements a <see cref="Channel{T}"/> which signals completion according to a reference count.
+	/// </summary>
+	public class ChannelWithRefCount<T> : Channel<T>
+	{
+		int _count = 1;
+
+		/// <summary>
+		/// Called when the channel is completed
+		/// </summary>
+		public event Action? OnComplete;
+
+		/// <summary>
+		/// Constructor
+		/// </summary>
+		internal ChannelWithRefCount(Channel<T> inner)
+		{
+			Reader = inner.Reader;
+			Writer = inner.Writer;
+		}
+
+		/// <summary>
+		/// Adds a new reference to the channel
+		/// </summary>
+		public ChannelWithRefCount<T> AddRef()
+		{
+			Interlocked.Increment(ref _count);
+			return this;
+		}
+
+		/// <summary>
+		/// Releases a reference to the channel
+		/// </summary>
+		public void Release()
+		{
+			if (Interlocked.Decrement(ref _count) == 0)
+			{
+				OnComplete?.Invoke();
+				Writer.Complete();
+			}
+		}
+	}
+
+	/// <summary>
+	/// Extension methods for <see cref="Channel{T}"/>
+	/// </summary>
+	public static class ChannelExtensions
+	{
+		/// <summary>
+		/// Creates a channel whose completion state is controlled by a reference count
+		/// </summary>
+		public static ChannelWithRefCount<T> WithRefCount<T>(this Channel<T> channel)
+			=> new ChannelWithRefCount<T>(channel);
+	}
+}
diff --git a/Engine/Source/Programs/Shared/EpicGames.Core/CommandHost.cs b/Engine/Source/Programs/Shared/EpicGames.Core/CommandHost.cs
--- a/Engine/Source/Programs/Shared/EpicGames.Core/CommandHost.cs
+++ b/Engine/Source/Programs/Shared/EpicGames.Core/CommandHost.cs
@@ -57,6 +57,11 @@ namespace EpicGames.Core
 		/// </summary>
 		public bool Advertise { get; }
 
+		/// <summary>
+		/// Whether the command accepts positional arguments
+		/// </summary>
+		public bool AcceptsPositionalArguments { get; }
+
 		/// <summary>
 		/// Create a command instance
 		/// </summary>
@@ -85,6 +90,11 @@ namespace EpicGames.Core
 		/// </summary>
 		public bool Advertise { get; set; } = true;
 
+		/// <summary>
+		/// Whether the command takes positional arguments
+		/// </summary>
+		public bool AcceptsPositionalArguments { get; set; } = false;
+
 		/// <summary>
 		/// Constructor
 		/// </summary>
@@ -139,14 +149,16 @@ namespace EpicGames.Core
 		public IReadOnlyList<string> Names { get; }
 		public string Description { get; }
 		public bool Advertise { get; }
+		public bool AcceptsPositionalArguments { get; }
 
 		public Type Type { get; }
 
-		public CommandFactory(string[] names, string description, bool advertise, Type type)
+		public CommandFactory(string[] names, string description, bool advertise, bool acceptsPositionalArguments, Type type)
 		{
 			Names = names;
 			Description = description;
 			Advertise = advertise;
+			AcceptsPositionalArguments = acceptsPositionalArguments;
 			Type = type;
 		}
 
@@ -175,7 +187,7 @@ namespace EpicGames.Core
 					if (attribute != null)
 					{
 						services.AddTransient(type);
-						services.AddTransient(typeof(ICommandFactory), sp => new CommandFactory(attribute.Names, attribute.Description, attribute.Advertise, type));
+						services.AddTransient(typeof(ICommandFactory), sp => new CommandFactory(attribute.Names, attribute.Description, attribute.Advertise, attribute.AcceptsPositionalArguments, type));
 					}
 				}
 			}
@@ -233,8 +245,14 @@ namespace EpicGames.Core
 			{
 				foreach (ICommandFactory factory in commandFactories)
 				{
-					if (factory.Names.SequenceEqual(positionalArgs, StringComparer.OrdinalIgnoreCase))
+					IEnumerable<string> positionalArgsEnum = positionalArgs;
+					if (factory.AcceptsPositionalArguments && factory.Names.Count < positionalArgs.Length)
+					{
+						positionalArgsEnum = positionalArgsEnum.Take(factory.Names.Count);
+					}
+					if (factory.Names.SequenceEqual(positionalArgsEnum, StringComparer.OrdinalIgnoreCase))
 					{
+						args.RemovePositionalArguments(factory.Names.Count);
 						command = factory.CreateInstance(serviceProvider);
 						commandFactory = factory;
 						break;
diff --git a/Engine/Source/Programs/Shared/EpicGames.Core/CommandLineArguments.cs b/Engine/Source/Programs/Shared/EpicGames.Core/CommandLineArguments.cs
--- a/Engine/Source/Programs/Shared/EpicGames.Core/CommandLineArguments.cs
+++ b/Engine/Source/Programs/Shared/EpicGames.Core/CommandLineArguments.cs
@@ -340,6 +340,15 @@ namespace EpicGames.Core
 			return positionalArguments;
 		}
 
+		/// <summary>
+		/// Consume positional arguments from command line
+		/// </summary>
+		/// <param name="count">Number of arguments to remove</param>
+		public void RemovePositionalArguments(int count)
+		{
+			_positionalArgumentIndices.RemoveRange(0, count);
+		}
+
 		/// <summary>
 		/// Gets the value specified by an argument with the given prefix. Throws an exception if the argument was not specified.
 		/// </summary>
@@ -742,6 +751,20 @@ namespace EpicGames.Core
 				// Loop through all the attributes for different command line options that can modify it
 				foreach(CommandLineAttribute attribute in target.Attributes)
 				{
+					// Add in any positional arguments
+					if (attribute.Positional)
+					{
+						foreach (int argIdx in _positionalArgumentIndices)
+						{
+							string arg = _arguments[argIdx];
+							if (ApplyArgument(targetObject, target, arg, arg, assignedArgument, logger))
+							{
+								MarkAsUsed(argIdx);
+								assignedArgument = arg;
+							}
+						}
+					}
+
 					// Get the appropriate prefix for this attribute
 					string prefix = GetArgumentPrefix(target, attribute);
 
diff --git a/Engine/Source/Programs/Shared/EpicGames.Core/CommandLineAttribute.cs b/Engine/Source/Programs/Shared/EpicGames.Core/CommandLineAttribute.cs
--- a/Engine/Source/Programs/Shared/EpicGames.Core/CommandLineAttribute.cs
+++ b/Engine/Source/Programs/Shared/EpicGames.Core/CommandLineAttribute.cs
@@ -32,6 +32,11 @@ namespace EpicGames.Core
 		/// </summary>
 		public bool MarkUsed { get; set; } = true;
 
+		/// <summary>
+		/// Accepts positional arguments (ie. without a preceding hyphen)
+		/// </summary>
+		public bool Positional { get; set; }
+
 		/// <summary>
 		/// For collection types, specifies the separator character between multiple arguments
 		/// </summary>
diff --git a/Engine/Source/Programs/Shared/EpicGames.Core/FileFilter.cs b/Engine/Source/Programs/Shared/EpicGames.Core/FileFilter.cs
--- a/Engine/Source/Programs/Shared/EpicGames.Core/FileFilter.cs
+++ b/Engine/Source/Programs/Shared/EpicGames.Core/FileFilter.cs
@@ -2,6 +2,8 @@
 
 using System;
 using System.Collections.Generic;
+using System.ComponentModel;
+using System.Globalization;
 using System.IO;
 using System.Linq;
 using System.Text.RegularExpressions;
@@ -29,6 +31,7 @@ namespace EpicGames.Core
 	/// The rules are merged into a tree with nodes for each path fragments and rule numbers in each leaf, allowing tree traversal in an arbitrary order 
 	/// to how they are applied.
 	/// </summary>
+	[TypeConverter(typeof(FileFilterConverter))]
 	public class FileFilter
 	{
 		/// <summary>
@@ -84,6 +87,37 @@ namespace EpicGames.Core
 			}
 		}
 
+		/// <summary>
+		/// Parse a filter pattern, determining the default rule type from the initial rule
+		/// </summary>
+		/// <param name="text">Sequence of patterns to parse. Entries may be separated by semicolons.</param>
+		public static FileFilter? Parse(string? text)
+		{
+			FileFilter? filter = null;
+			if (text != null)
+			{
+				string[] rules = text.Split(';', StringSplitOptions.RemoveEmptyEntries | StringSplitOptions.TrimEntries);
+				filter = Parse(rules);
+			}
+			return filter;
+		}
+
+		/// <summary>
+		/// Parse a filter pattern, determining the default rule type from the initial rule
+		/// </summary>
+		/// <param name="rules">Array of rules</param>
+		public static FileFilter? Parse(IReadOnlyList<string>? rules)
+		{
+			FileFilter? filter = null;
+			if (rules != null && rules.Count > 0)
+			{
+				FileFilterType defaultType = rules[0].StartsWith("-", StringComparison.Ordinal) ? FileFilterType.Include : FileFilterType.Exclude;
+				filter = new FileFilter(defaultType);
+				filter.AddRules(rules);
+			}
+			return filter;
+		}
+
 		/// <summary>
 		/// Adds an include or exclude rule to the filter
 		/// </summary>
@@ -108,24 +142,24 @@ namespace EpicGames.Core
 		public void AddRule(string rule, params string[] allowTags)
 		{
 			string cleanRule = rule.Trim();
-			if(cleanRule.StartsWith('{'))
+			if (cleanRule.StartsWith('{'))
 			{
 				// Find the end of the condition
 				int conditionEnd = cleanRule.IndexOf('}', StringComparison.Ordinal);
-				if(conditionEnd == -1)
+				if (conditionEnd == -1)
 				{
 					throw new Exception(String.Format("Missing closing parenthesis in rule: {0}", cleanRule));
 				}
 
 				// Check there's a matching tag
 				string[] ruleTags = cleanRule.Substring(1, conditionEnd - 1).Split(',').Select(x => x.Trim()).ToArray();
-				if(!ruleTags.Any(x => allowTags.Contains(x)))
+				if (!ruleTags.Any(x => allowTags.Contains(x)))
 				{
 					return;
 				}
 
 				// Strip the condition from the rule
-				cleanRule = cleanRule.Substring(conditionEnd + 1).TrimStart(); 
+				cleanRule = cleanRule.Substring(conditionEnd + 1).TrimStart();
 			}
 			AddRule(cleanRule);
 		}
@@ -136,7 +170,7 @@ namespace EpicGames.Core
 		/// <param name="rules">List of patterns to match.</param>
 		public void AddRules(IEnumerable<string> rules)
 		{
-			foreach(string rule in rules)
+			foreach (string rule in rules)
 			{
 				AddRule(rule);
 			}
@@ -149,7 +183,7 @@ namespace EpicGames.Core
 		/// <param name="type">Whether the rules are include or exclude rules</param>
 		public void AddRules(IEnumerable<string> rules, FileFilterType type)
 		{
-			foreach(string rule in rules)
+			foreach (string rule in rules)
 			{
 				AddRule(rule, type);
 			}
@@ -163,7 +197,7 @@ namespace EpicGames.Core
 		/// <param name="tags">Lookup for variables to test against</param>
 		public void AddRules(IEnumerable<string> rules, params string[] tags)
 		{
-			foreach(string rule in rules)
+			foreach (string rule in rules)
 			{
 				AddRule(rule, tags);
 			}
@@ -203,16 +237,16 @@ namespace EpicGames.Core
 		public void ReadRulesFromFile(FileReference fileName, string sectionName, params string[] allowTags)
 		{
 			bool bInSection = false;
-			foreach(string line in File.ReadAllLines(fileName.FullName))
+			foreach (string line in File.ReadAllLines(fileName.FullName))
 			{
 				string trimLine = line.Trim();
-				if(!trimLine.StartsWith(';') && trimLine.Length > 0)
+				if (!trimLine.StartsWith(';') && trimLine.Length > 0)
 				{
-					if(trimLine.StartsWith('['))
+					if (trimLine.StartsWith('['))
 					{
 						bInSection = (trimLine == "[" + sectionName + "]");
 					}
-					else if(bInSection)
+					else if (bInSection)
 					{
 						AddRule(line, allowTags);
 					}
@@ -295,7 +329,7 @@ namespace EpicGames.Core
 			{
 				normalizedPattern = normalizedPattern.Substring(1);
 			}
-			else if(!normalizedPattern.Contains('/', StringComparison.Ordinal) && !normalizedPattern.StartsWith("...", StringComparison.Ordinal))
+			else if (!normalizedPattern.Contains('/', StringComparison.Ordinal) && !normalizedPattern.StartsWith("...", StringComparison.Ordinal))
 			{
 				normalizedPattern = ".../" + normalizedPattern;
 			}
@@ -415,11 +449,20 @@ namespace EpicGames.Core
 		/// <param name="bIgnoreSymlinks">Whether to ignore symlinks in the output</param>
 		/// <returns>List of files that pass the filter</returns>
 		public List<FileReference> ApplyToDirectory(DirectoryReference directoryName, bool bIgnoreSymlinks)
+			=> ApplyToDirectory(directoryName.ToDirectoryInfo(), bIgnoreSymlinks).ConvertAll(x => new FileReference(x));
+
+		/// <summary>
+		/// Finds a list of files within a given directory which match the filter.
+		/// </summary>
+		/// <param name="directoryInfo">File to match</param>
+		/// <param name="bIgnoreSymlinks">Whether to ignore symlinks in the output</param>
+		/// <returns>List of files that pass the filter</returns>
+		public List<FileInfo> ApplyToDirectory(DirectoryInfo directoryInfo, bool bIgnoreSymlinks)
 		{
-			List<FileReference> matchingFileNames = [];
-			if (DirectoryReference.Exists(directoryName))
+			List<FileInfo> matchingFileNames = [];
+			if (directoryInfo.Exists)
 			{
-				FindMatchesFromDirectory(new DirectoryInfo(directoryName.FullName), "", bIgnoreSymlinks, matchingFileNames);
+				FindMatchesFromDirectory(directoryInfo, "", bIgnoreSymlinks, matchingFileNames);
 			}
 			return matchingFileNames;
 		}
@@ -432,11 +475,21 @@ namespace EpicGames.Core
 		/// <param name="bIgnoreSymlinks">Whether to ignore symlinks in the output</param>
 		/// <returns>List of files that pass the filter</returns>
 		public List<FileReference> ApplyToDirectory(DirectoryReference directoryName, string prefixPath, bool bIgnoreSymlinks)
+			=> ApplyToDirectory(directoryName.ToDirectoryInfo(), prefixPath, bIgnoreSymlinks).ConvertAll(x => new FileReference(x));
+
+		/// <summary>
+		/// Finds a list of files within a given directory which match the filter.
+		/// </summary>
+		/// <param name="directoryInfo">File to match</param>
+		/// <param name="prefixPath"></param>
+		/// <param name="bIgnoreSymlinks">Whether to ignore symlinks in the output</param>
+		/// <returns>List of files that pass the filter</returns>
+		public List<FileInfo> ApplyToDirectory(DirectoryInfo directoryInfo, string prefixPath, bool bIgnoreSymlinks)
 		{
-			List<FileReference> matchingFileNames = [];
-			if (DirectoryReference.Exists(directoryName))
+			List<FileInfo> matchingFileNames = [];
+			if (directoryInfo.Exists)
 			{
-				FindMatchesFromDirectory(new DirectoryInfo(directoryName.FullName), NormalizePath(prefixPath).TrimEnd('/') + "/", bIgnoreSymlinks, matchingFileNames);
+				FindMatchesFromDirectory(directoryInfo, NormalizePath(prefixPath).TrimEnd('/') + "/", bIgnoreSymlinks, matchingFileNames);
 			}
 			return matchingFileNames;
 		}
@@ -449,9 +502,9 @@ namespace EpicGames.Core
 		public static int FindWildcardIndex(string pattern)
 		{
 			int result = -1;
-			for(int idx = 0; idx < pattern.Length; idx++)
+			for (int idx = 0; idx < pattern.Length; idx++)
 			{
-				if(pattern[idx] == '?' || pattern[idx] == '*' || (pattern[idx] == '.' && idx + 2 < pattern.Length && pattern[idx + 1] == '.' && pattern[idx + 2] == '.'))
+				if (pattern[idx] == '?' || pattern[idx] == '*' || (pattern[idx] == '.' && idx + 2 < pattern.Length && pattern[idx + 1] == '.' && pattern[idx + 2] == '.'))
 				{
 					result = idx;
 					break;
@@ -467,9 +520,9 @@ namespace EpicGames.Core
 		/// <returns>List of files matching the wildcard</returns>
 		public static List<FileReference> ResolveWildcard(string pattern)
 		{
-			for(int idx = FindWildcardIndex(pattern); idx > 0; idx--)
+			for (int idx = FindWildcardIndex(pattern); idx > 0; idx--)
 			{
-				if(pattern[idx] == '/' || pattern[idx] == '\\')
+				if (pattern[idx] == '/' || pattern[idx] == '\\')
 				{
 					return ResolveWildcard(new DirectoryReference(pattern.Substring(0, idx)), pattern.Substring(idx + 1));
 				}
@@ -498,14 +551,14 @@ namespace EpicGames.Core
 		/// <param name="bIgnoreSymlinks">Whether to ignore symlinks in the output</param>
 		/// <param name="matchingFileNames">Receives a list of matching files</param>
 		/// <returns>True if the file passes the filter</returns>
-		void FindMatchesFromDirectory(DirectoryInfo currentDirectory, string namePrefix, bool bIgnoreSymlinks, List<FileReference> matchingFileNames)
+		void FindMatchesFromDirectory(DirectoryInfo currentDirectory, string namePrefix, bool bIgnoreSymlinks, List<FileInfo> matchingFileNames)
 		{
 			foreach (FileInfo nextFile in currentDirectory.EnumerateFiles())
 			{
 				string fileName = namePrefix + nextFile.Name;
 				if (Matches(fileName) && (!bIgnoreSymlinks || !nextFile.Attributes.HasFlag(FileAttributes.ReparsePoint)))
 				{
-					matchingFileNames.Add(new FileReference(nextFile));
+					matchingFileNames.Add(nextFile);
 				}
 			}
 			foreach (DirectoryInfo nextDirectory in currentDirectory.EnumerateDirectories())
@@ -603,4 +656,24 @@ namespace EpicGames.Core
 			return bestRuleNumber;
 		}
 	}
+
+	/// <summary>
+	/// Type converter for FileFilter objects. Allows parsing filters from command line arguments.
+	/// </summary>
+	class FileFilterConverter : TypeConverter
+	{
+		/// <inheritdoc/>
+		public override bool CanConvertFrom(ITypeDescriptorContext? context, Type sourceType)
+			=> sourceType == typeof(string);
+
+		/// <inheritdoc/>
+		public override object? ConvertFrom(ITypeDescriptorContext? context, CultureInfo? culture, object value)
+		{
+			if (value is string str)
+			{
+				return FileFilter.Parse(str);
+			}
+			return base.ConvertFrom(context, culture, value);
+		}
+	}
 }
diff --git a/Engine/Source/Programs/Shared/EpicGames.Horde.Tests/Storage/WorkspaceTests.cs b/Engine/Source/Programs/Shared/EpicGames.Horde.Tests/Storage/WorkspaceTests.cs
new file mode 100644
--- /dev/null
+++ b/Engine/Source/Programs/Shared/EpicGames.Horde.Tests/Storage/WorkspaceTests.cs
@@ -0,0 +1,272 @@
+// Copyright Epic Games, Inc. All Rights Reserved.
+
+using System;
+using System.IO;
+using System.Runtime.CompilerServices;
+using System.Text;
+using System.Threading;
+using System.Threading.Tasks;
+using EpicGames.Core;
+using EpicGames.Horde.Storage;
+using EpicGames.Horde.Storage.Nodes;
+using Microsoft.Extensions.Logging.Abstractions;
+using Microsoft.VisualStudio.TestTools.UnitTesting;
+
+namespace EpicGames.Horde.Tests.Storage
+{
+	[TestClass]
+	public class WorkspaceTests
+	{
+		static ReadOnlyMemory<byte> HelloBytes { get; } = Encoding.UTF8.GetBytes("hello");
+		static ReadOnlyMemory<byte> WorldBytes { get; } = Encoding.UTF8.GetBytes("world");
+
+		static async Task<IHashedBlobRef<DirectoryNode>> CreateBasicTreeAsync(IBlobWriter blobWriter, CancellationToken cancellationToken)
+		{
+			DirectoryNode rootDir = new DirectoryNode();
+			rootDir.AddFile("readme.txt", FileEntryFlags.None, await ChunkedData.CreateAsync(blobWriter, HelloBytes, cancellationToken));
+			return await blobWriter.WriteBlobAsync(rootDir, cancellationToken);
+		}
+
+		static async Task AssertContentsEqualAsync(IHashedBlobRef<DirectoryNode> dirRefA, IHashedBlobRef<DirectoryNode> dirRefB)
+		{
+			DirectoryNode dirA = await dirRefA.ReadBlobAsync();
+			DirectoryNode dirB = await dirRefB.ReadBlobAsync();
+			dirA.DeleteDirectory(".horde");
+			dirB.DeleteDirectory(".horde");
+
+			Assert.AreEqual(dirA.Directories.Count, dirB.Directories.Count);
+			foreach (DirectoryEntry dirEntryA in dirA.Directories)
+			{
+				DirectoryEntry dirEntryB = dirB.GetDirectoryEntry(dirEntryA.Name);
+				Assert.AreEqual(dirEntryA.Name, dirEntryB.Name, false);
+				Assert.AreEqual(dirEntryA.Length, dirEntryB.Length);
+				await AssertContentsEqualAsync(dirEntryA.Handle, dirEntryB.Handle);
+			}
+
+			Assert.AreEqual(dirA.Files.Count, dirB.Files.Count);
+			foreach (FileEntry fileEntryA in dirA.Files)
+			{
+				FileEntry fileEntryB = dirB.GetFileEntry(fileEntryA.Name);
+				Assert.AreEqual(fileEntryA.Name, fileEntryB.Name, false);
+				Assert.AreEqual(fileEntryA.Length, fileEntryB.Length);
+
+				byte[] dataA = await fileEntryA.ReadAllBytesAsync();
+				byte[] dataB = await fileEntryB.ReadAllBytesAsync();
+
+				Assert.IsTrue(MemoryExtensions.SequenceEqual<byte>(dataA, dataB));
+			}
+		}
+
+		static DirectoryInfo CreateTempDirectory([CallerMemberName] string? callerName = null)
+		{
+			DirectoryInfo tempDir = new DirectoryInfo($"temp/{callerName ?? Guid.NewGuid().ToString()}");
+			tempDir.Create();
+			FileUtils.ForceDeleteDirectoryContents(tempDir);
+			return tempDir;
+		}
+
+		[TestMethod]
+		public async Task MaterializeManuallyAsync()
+		{
+			await using MemoryBlobWriter blobWriter = new MemoryBlobWriter(new BlobSerializerOptions());
+			IHashedBlobRef<DirectoryNode> input = await CreateBasicTreeAsync(blobWriter, default);
+
+			DirectoryInfo tempDir = CreateTempDirectory();
+			await input.ExtractAsync(tempDir, new ExtractOptions(), NullLogger.Instance, default);
+
+			IHashedBlobRef<DirectoryNode> output = await blobWriter.WriteFilesAsync(new DirectoryReference(tempDir), null, null, default);
+			await AssertContentsEqualAsync(input, output);
+		}
+
+		[TestMethod]
+		public async Task SyncTestAsync()
+		{
+			DirectoryInfo tempDir = CreateTempDirectory();
+			Workspace ws = await Workspace.CreateAsync(new DirectoryReference(tempDir), NullLogger.Instance);
+			await using MemoryBlobWriter blobWriter = new MemoryBlobWriter(new BlobSerializerOptions());
+
+			IHashedBlobRef<DirectoryNode> initialState;
+			{
+				DirectoryNode rootDir = new DirectoryNode();
+				rootDir.AddFile("readme.txt", FileEntryFlags.None, await ChunkedData.CreateAsync(blobWriter, HelloBytes));
+				initialState = await blobWriter.WriteBlobAsync(rootDir);
+			}
+
+			WorkspaceSyncStats initialStats = new WorkspaceSyncStats();
+			await ws.SyncAsync(WorkspaceLayerId.Default, initialState, null, stats: initialStats);
+
+			WorkspaceSyncStats updateStats = new WorkspaceSyncStats();
+			await ws.SyncAsync(WorkspaceLayerId.Default, initialState, null, stats: updateStats);
+			Assert.AreEqual(1, updateStats.NumFilesKept);
+			Assert.AreEqual(0, updateStats.NumBytesReused);
+			Assert.AreEqual(0, updateStats.NumBytesDownloaded);
+
+			IHashedBlobRef<DirectoryNode> output = await blobWriter.WriteFilesAsync(new DirectoryReference(tempDir), null, null, default);
+			await AssertContentsEqualAsync(initialState, output);
+		}
+
+		[TestMethod]
+		public async Task MoveTestAsync()
+		{
+			DirectoryInfo tempDir = CreateTempDirectory();
+			Workspace ws = await Workspace.CreateAsync(new DirectoryReference(tempDir), NullLogger.Instance);
+			await using MemoryBlobWriter blobWriter = new MemoryBlobWriter(new BlobSerializerOptions());
+
+			IHashedBlobRef<DirectoryNode> initialState;
+			{
+				DirectoryNode rootDir = new DirectoryNode();
+				rootDir.AddFile("readme.txt", FileEntryFlags.None, await ChunkedData.CreateAsync(blobWriter, HelloBytes));
+				initialState = await blobWriter.WriteBlobAsync(rootDir);
+			}
+			await ws.SyncAsync(WorkspaceLayerId.Default, initialState, null);
+
+			IHashedBlobRef<DirectoryNode> updatedState;
+			{
+				DirectoryNode rootDir = new DirectoryNode();
+				rootDir.AddFile("file2.txt", FileEntryFlags.None, await ChunkedData.CreateAsync(blobWriter, HelloBytes));
+				updatedState = await blobWriter.WriteBlobAsync(rootDir);
+			}
+			await ws.SyncAsync(WorkspaceLayerId.Default, updatedState, null);
+
+			IHashedBlobRef<DirectoryNode> output = await blobWriter.WriteFilesAsync(new DirectoryReference(tempDir), null, null, default);
+			await AssertContentsEqualAsync(updatedState, output);
+		}
+
+		[TestMethod]
+		public async Task CopyTestAsync()
+		{
+			DirectoryInfo tempDir = CreateTempDirectory();
+			Workspace ws = await Workspace.CreateAsync(new DirectoryReference(tempDir), NullLogger.Instance);
+			await using MemoryBlobWriter blobWriter = new MemoryBlobWriter(new BlobSerializerOptions());
+
+			IHashedBlobRef<DirectoryNode> initialState;
+			{
+				DirectoryNode rootDir = new DirectoryNode();
+				rootDir.AddFile("readme.txt", FileEntryFlags.None, await ChunkedData.CreateAsync(blobWriter, HelloBytes));
+				initialState = await blobWriter.WriteBlobAsync(rootDir);
+			}
+			await ws.SyncAsync(WorkspaceLayerId.Default, initialState, null);
+
+			IHashedBlobRef<DirectoryNode> updatedState;
+			{
+				DirectoryNode rootDir = new DirectoryNode();
+				rootDir.AddFile("file2.txt", FileEntryFlags.None, await ChunkedData.CreateAsync(blobWriter, HelloBytes));
+				rootDir.AddFile("file3.txt", FileEntryFlags.None, await ChunkedData.CreateAsync(blobWriter, HelloBytes));
+				updatedState = await blobWriter.WriteBlobAsync(rootDir);
+			}
+			await ws.SyncAsync(WorkspaceLayerId.Default, updatedState, null);
+
+			IHashedBlobRef<DirectoryNode> output = await blobWriter.WriteFilesAsync(new DirectoryReference(tempDir), null, null, default);
+			await AssertContentsEqualAsync(updatedState, output);
+		}
+
+		[TestMethod]
+		public async Task LayerTestAsync()
+		{
+			DirectoryInfo tempDir = CreateTempDirectory();
+			Workspace ws = await Workspace.CreateAsync(new DirectoryReference(tempDir), NullLogger.Instance);
+			await using MemoryBlobWriter blobWriter = new MemoryBlobWriter(new BlobSerializerOptions());
+
+			IHashedBlobRef<DirectoryNode> emptyState;
+			{
+				DirectoryNode rootDir = new DirectoryNode();
+				emptyState = await blobWriter.WriteBlobAsync(rootDir);
+			}
+
+			IHashedBlobRef<DirectoryNode> initialState;
+			{
+				DirectoryNode subDirA = new DirectoryNode();
+				subDirA.AddFile("hello.txt", FileEntryFlags.None, await ChunkedData.CreateAsync(blobWriter, HelloBytes));
+				subDirA.AddFile("world.txt", FileEntryFlags.None, await ChunkedData.CreateAsync(blobWriter, WorldBytes));
+				IHashedBlobRef<DirectoryNode> subDirRefA = await blobWriter.WriteBlobAsync(subDirA);
+
+				DirectoryNode rootDir = new DirectoryNode();
+				rootDir.AddDirectory(new DirectoryEntry("A", subDirA.Length, subDirRefA));
+				initialState = await blobWriter.WriteBlobAsync(rootDir);
+			}
+
+			WorkspaceLayerId layerA = new WorkspaceLayerId("layer-a");
+			ws.AddLayer(layerA);
+
+			WorkspaceLayerId layerB = new WorkspaceLayerId("layer-b");
+			ws.AddLayer(layerB);
+
+			// Add the files to layer A
+			WorkspaceSyncStats stats = new WorkspaceSyncStats();
+			await ws.SyncAsync(layerA, initialState, stats: stats);
+			Assert.AreEqual(stats.NumFilesKept, 0);
+			Assert.AreEqual(stats.NumBytesDownloaded, 10);
+
+			// Add the same files to layer B
+			stats = new WorkspaceSyncStats();
+			await ws.SyncAsync(layerB, initialState, stats: stats);
+			Assert.AreEqual(stats.NumFilesKept, 2);
+			Assert.AreEqual(stats.NumBytesDownloaded, 0);
+
+			// Remove layer B
+			stats = new WorkspaceSyncStats();
+			await ws.SyncAsync(layerB, null, stats: stats);
+			Assert.AreEqual(stats.NumFilesKept, 0);
+			Assert.AreEqual(stats.NumBytesDownloaded, 0);
+
+			// Check it's still the same
+			IHashedBlobRef<DirectoryNode> output = await blobWriter.WriteFilesAsync(new DirectoryReference(tempDir), null, null, default);
+			await AssertContentsEqualAsync(initialState, output);
+
+			// Remove layer A
+			stats = new WorkspaceSyncStats();
+			await ws.SyncAsync(layerA, null, stats: stats);
+			Assert.AreEqual(stats.NumFilesKept, 0);
+			Assert.AreEqual(stats.NumBytesDownloaded, 0);
+
+			// Check it's empty
+			output = await blobWriter.WriteFilesAsync(new DirectoryReference(tempDir), null, null, default);
+			await AssertContentsEqualAsync(emptyState, output);
+		}
+
+		[TestMethod]
+		public async Task SwapTestAsync()
+		{
+			DirectoryInfo tempDir = CreateTempDirectory();
+			Workspace ws = await Workspace.CreateAsync(new DirectoryReference(tempDir), NullLogger.Instance);
+			await using MemoryBlobWriter blobWriter = new MemoryBlobWriter(new BlobSerializerOptions());
+
+			IHashedBlobRef<DirectoryNode> initialState;
+			{
+				DirectoryNode subDirA = new DirectoryNode();
+				subDirA.AddFile("hello.txt", FileEntryFlags.None, await ChunkedData.CreateAsync(blobWriter, HelloBytes));
+				subDirA.AddFile("world.txt", FileEntryFlags.None, await ChunkedData.CreateAsync(blobWriter, WorldBytes));
+				IHashedBlobRef<DirectoryNode> subDirRefA = await blobWriter.WriteBlobAsync(subDirA);
+
+				DirectoryNode rootDir = new DirectoryNode();
+				rootDir.AddDirectory(new DirectoryEntry("A", subDirA.Length, subDirRefA));
+				rootDir.AddFile(new FileEntry("B", FileEntryFlags.None, await ChunkedData.CreateAsync(blobWriter, HelloBytes)));
+				initialState = await blobWriter.WriteBlobAsync(rootDir);
+			}
+			await ws.SyncAsync(WorkspaceLayerId.Default, initialState, null);
+
+			IHashedBlobRef<DirectoryNode> updatedState;
+			{
+				DirectoryNode subDirB = new DirectoryNode();
+				subDirB.AddFile("hello.txt", FileEntryFlags.None, await ChunkedData.CreateAsync(blobWriter, HelloBytes));
+				subDirB.AddFile("world.txt", FileEntryFlags.None, await ChunkedData.CreateAsync(blobWriter, WorldBytes));
+				IHashedBlobRef<DirectoryNode> subDirRefB = await blobWriter.WriteBlobAsync(subDirB);
+
+				DirectoryNode rootDir = new DirectoryNode();
+				rootDir.AddFile(new FileEntry("A", FileEntryFlags.None, await ChunkedData.CreateAsync(blobWriter, HelloBytes)));
+				rootDir.AddDirectory(new DirectoryEntry("B", subDirB.Length, subDirRefB));
+				updatedState = await blobWriter.WriteBlobAsync(rootDir);
+			}
+
+			WorkspaceSyncStats syncStats = new WorkspaceSyncStats();
+			await ws.SyncAsync(WorkspaceLayerId.Default, updatedState, stats: syncStats);
+
+			IHashedBlobRef<DirectoryNode> output = await blobWriter.WriteFilesAsync(new DirectoryReference(tempDir), null, null, default);
+			await AssertContentsEqualAsync(updatedState, output);
+
+			Assert.AreEqual(3, syncStats.NumFilesMoved);
+			Assert.AreEqual(0, syncStats.NumBytesReused);
+			Assert.AreEqual(0, syncStats.NumBytesDownloaded);
+		}
+	}
+}
diff --git a/Engine/Source/Programs/Shared/EpicGames.Horde/Commits/CommitTag.cs b/Engine/Source/Programs/Shared/EpicGames.Horde/Commits/CommitTag.cs
--- a/Engine/Source/Programs/Shared/EpicGames.Horde/Commits/CommitTag.cs
+++ b/Engine/Source/Programs/Shared/EpicGames.Horde/Commits/CommitTag.cs
@@ -36,7 +36,10 @@ namespace EpicGames.Horde.Commits
 		public CommitTag(string text)
 		{
 			Text = text;
-			StringId.ValidateArgument(new Utf8String(text), nameof(text));
+			if (!StringId.TryParse(text, out _, out string? errorMessage))
+			{
+				throw new ArgumentException(errorMessage, nameof(text));
+			}
 		}
 
 		/// <summary>
diff --git a/Engine/Source/Programs/Shared/EpicGames.Horde/Storage/ContentChunker.cs b/Engine/Source/Programs/Shared/EpicGames.Horde/Storage/ContentChunker.cs
--- a/Engine/Source/Programs/Shared/EpicGames.Horde/Storage/ContentChunker.cs
+++ b/Engine/Source/Programs/Shared/EpicGames.Horde/Storage/ContentChunker.cs
@@ -528,6 +528,13 @@ namespace EpicGames.Horde.Storage
 		public override ChannelReader<ChunkerOutput> OutputReader
 			=> _outputChannel.Reader;
 
+		/// <summary>
+		/// Constructor
+		/// </summary>
+		public ParallelBuzHashChunker(AsyncPipeline pipeline, LeafChunkedDataNodeOptions options)
+			: this(pipeline, new byte[2 * 1024 * 1024], options)
+		{ }
+
 		/// <summary>
 		/// Constructor
 		/// </summary>
diff --git a/Engine/Source/Programs/Shared/EpicGames.Horde/Storage/Workspace.cs b/Engine/Source/Programs/Shared/EpicGames.Horde/Storage/Workspace.cs
--- a/Engine/Source/Programs/Shared/EpicGames.Horde/Storage/Workspace.cs
+++ b/Engine/Source/Programs/Shared/EpicGames.Horde/Storage/Workspace.cs
@@ -1,6 +1,7 @@
 // Copyright Epic Games, Inc. All Rights Reserved.
 
 using System;
+using System.Collections.Concurrent;
 using System.Collections.Generic;
 using System.Diagnostics;
 using System.Diagnostics.CodeAnalysis;
@@ -12,90 +13,306 @@ using System.Threading.Channels;
 using System.Threading.Tasks;
 using EpicGames.Core;
 using EpicGames.Horde.Storage.Nodes;
+using EpicGames.Serialization;
 using Microsoft.Extensions.Logging;
 
 namespace EpicGames.Horde.Storage
 {
+	/// <summary>
+	/// Base exception type for workspace errors
+	/// </summary>
+	public class WorkspaceException : Exception
+	{
+		/// <summary>
+		/// Constructor
+		/// </summary>
+		public WorkspaceException(string message, Exception? innerException = null)
+			: base(message, innerException) { }
+	}
+
 	/// <summary>
 	/// Provides functionality to extract and patch data in a local workspace
 	/// </summary>
 	public class Workspace
 	{
+		// Flags for the layers that a file or directory belongs to. Some of these are transient states; others are used
+		// to indicate user defined layers.
+		[Flags]
+		enum LayerFlags : ulong
+		{
+			// Item is in the cache layer
+			Cache = 1,
+
+			// Item has pending actions; state may be inconsistent with disk.
+			Pending = 2,
+
+			// The default layer for user operations. Must be the last item in the enum; user layers will be allocated flags beyond this.
+			Default = 4,
+
+			// Custom layers
+			Custom1 = Default << 1,
+			Custom2 = Default << 2,
+			Custom3 = Default << 3,
+			Custom4 = Default << 4,
+			Custom5 = Default << 5,
+			Custom6 = Default << 6,
+			Custom7 = Default << 7,
+			Custom8 = Default << 8,
+		}
+
+		// Actions to be performed to an item during a sync
+		[Flags]
+		enum PendingFlags
+		{
+			// Item is not modified
+			None = 0,
+
+			// Remove this file from the layer
+			RemoveFromLayer = 2,
+
+			// Delete this file from disk
+			RemoveFromDisk = 4,
+
+			// Move into a new location
+			Move = 8,
+
+			// Use this as a source for chunks, Harvest for new files
+			Copy = 16,
+
+			// This item overlaps with the location of another output, and must be moved to the cache.
+			ForceRenameOrRemove = 32,
+		}
+
+		// Maps a user defined layer id to a flag
+		[DebuggerDisplay("{Id}")]
+		class LayerState
+		{
+			public WorkspaceLayerId Id { get; }
+			public LayerFlags Flag { get; }
+
+			public LayerState(WorkspaceLayerId id, LayerFlags flag)
+			{
+				Id = id;
+				Flag = flag;
+			}
+
+			public LayerState(IMemoryReader reader)
+			{
+				Id = new WorkspaceLayerId(new StringId(reader.ReadUtf8String()));
+				Flag = (LayerFlags)reader.ReadUnsignedVarInt();
+			}
+
+			public void Write(IMemoryWriter writer)
+			{
+				writer.WriteUtf8String(Id.Id.Text);
+				writer.WriteUnsignedVarInt((ulong)Flag);
+			}
+		}
+
 		// Tracked state of a directory
+		[DebuggerDisplay("{GetPath()}")]
 		class DirectoryState
 		{
-			public DirectoryState? Parent { get; }
+			// Concrete implementation of FileState which can access private members in DirectoryState
+			sealed class FileStateImpl : FileState
+			{
+				public FileStateImpl(DirectoryState parent, string name, IoHash hash, LayerFlags layers)
+					: base(parent, name, hash, layers)
+				{ }
+
+				public FileStateImpl(DirectoryState parent, IMemoryReader reader)
+					: base(parent, reader)
+				{ }
+
+				public override void MoveTo(DirectoryState newParent, string newName, LayerFlags newLayers)
+				{
+					Unlink();
+
+					Parent = newParent;
+					Name = newName;
+					Layers = newLayers;
+
+					newParent.AddFileInternal(this);
+				}
+
+				public override void Unlink()
+				{
+					if (Parent != null)
+					{
+						Parent._files.Remove(Name);
+						Parent = null!;
+					}
+				}
+			}
+
+			public DirectoryState? Parent { get; private set; }
 			public string Name { get; }
-			public List<DirectoryState> Directories { get; } = new List<DirectoryState>();
-			public List<FileState> Files { get; } = new List<FileState>();
+			public IReadOnlyCollection<DirectoryState> Directories => _directories.Values;
+			public IReadOnlyCollection<FileState> Files => _files.Values;
+			public LayerFlags Layers { get; set; }
+			public PendingFlags Pending => _pendingFlags;
 
-			public ulong LayerFlags { get; set; }
+			PendingFlags _pendingFlags;
+
+			readonly Dictionary<string, DirectoryState> _directories;
+			readonly Dictionary<string, FileState> _files;
 
 			public DirectoryState(DirectoryState? parent, string name)
 			{
 				Parent = parent;
 				Name = name;
+
+				_directories = new Dictionary<string, DirectoryState>(DirectoryReference.Comparer);
+				_files = new Dictionary<string, FileState>(FileReference.Comparer);
 			}
 
-			public void AddFile(FileState fileState)
+			public DirectoryState(DirectoryState? parent, IMemoryReader reader)
 			{
-				int index = Files.BinarySearch(x => x.Name, fileState.Name, StringComparer.Ordinal);
-				if (index >= 0)
+				Parent = parent;
+				Name = reader.ReadString();
+
+				int numDirectories = (int)reader.ReadUnsignedVarInt();
+				_directories = new Dictionary<string, DirectoryState>(numDirectories, DirectoryReference.Comparer);
+
+				for (int idx = 0; idx < numDirectories; idx++)
+				{
+					DirectoryState subDirState = new DirectoryState(this, reader);
+					_directories.Add(subDirState.Name, subDirState);
+				}
+
+				int numFiles = (int)reader.ReadUnsignedVarInt();
+				_files = new Dictionary<string, FileState>(numFiles, FileReference.Comparer);
+
+				for (int idx = 0; idx < numFiles; idx++)
 				{
-					throw new InvalidOperationException($"File {fileState.Name} already exists");
+					FileState fileState = new FileStateImpl(this, reader);
+					_files.Add(fileState.Name, fileState);
 				}
 
-				fileState.Parent = this;
-				Files.Insert(~index, fileState);
+				Layers = (LayerFlags)reader.ReadUnsignedVarInt();
+
+				if ((Layers & LayerFlags.Pending) != 0)
+				{
+					_pendingFlags = (PendingFlags)reader.ReadUnsignedVarInt();
+				}
 			}
 
-			public void RemoveFile(FileState fileState)
+			public void Write(IMemoryWriter writer)
 			{
-				int index = Files.BinarySearch(x => x.Name, fileState.Name);
-				if (index >= 0 && Files[index] == fileState)
+				writer.WriteString(Name);
+
+				writer.WriteUnsignedVarInt(_directories.Count);
+				foreach (DirectoryState directory in _directories.Values)
+				{
+					directory.Write(writer);
+				}
+
+				writer.WriteUnsignedVarInt(_files.Count);
+				foreach (FileState file in _files.Values)
+				{
+					file.Write(writer);
+				}
+
+				writer.WriteUnsignedVarInt((ulong)Layers);
+
+				if ((Layers & LayerFlags.Pending) != 0)
 				{
-					Files.RemoveAt(index);
+					writer.WriteUnsignedVarInt((ulong)Pending);
 				}
 			}
 
-			public bool ContainsFile(string name) => TryGetFile(name, out _);
+			public void SetPendingAction(PendingFlags pending)
+			{
+				_pendingFlags |= pending;
+				MarkAsDirty();
+
+				foreach (DirectoryState subDirState in Directories)
+				{
+					subDirState.SetPendingAction(pending);
+				}
+
+				foreach (FileState fileState in Files)
+				{
+					fileState.SetPendingAction(pending);
+				}
+			}
 
-			public bool TryGetFile(string name, [NotNullWhen(true)] out FileState? fileState)
+			public void ClearPendingActions()
 			{
-				int index = Files.BinarySearch(x => x.Name, name, StringComparer.Ordinal);
-				if (index >= 0)
+				foreach (DirectoryState subDirState in Directories)
 				{
-					fileState = Files[index];
-					return true;
+					if ((subDirState.Layers & LayerFlags.Pending) != 0)
+					{
+						subDirState.ClearPendingActions();
+					}
 				}
-				else
+
+				foreach (FileState fileState in Files)
 				{
-					fileState = null;
-					return false;
+					fileState.Pending = PendingFlags.None;
+				}
+
+				_pendingFlags = PendingFlags.None;
+				Layers &= ~LayerFlags.Pending;
+			}
+
+			public void MarkAsDirty()
+			{
+				for (DirectoryState? dirState = this; dirState != null; dirState = dirState.Parent)
+				{
+					if ((dirState.Layers & LayerFlags.Pending) != 0)
+					{
+						dirState.Layers &= ~LayerFlags.Pending;
+					}
+					else
+					{
+						break;
+					}
 				}
 			}
 
-			public FileState FindOrAddFile(string name)
+			public FileState AddFile(string name, IoHash hash, LayerFlags layers)
 			{
-				int index = Files.BinarySearch(x => x.Name, name, StringComparer.Ordinal);
-				if (index >= 0)
+				FileState fileState = new FileStateImpl(this, name, hash, layers);
+				AddFileInternal(fileState);
+				return fileState;
+			}
+
+			void AddFileInternal(FileState fileState)
+			{
+				if (_directories.TryGetValue(fileState.Name, out DirectoryState? existingDir))
 				{
-					return Files[index];
+					throw new WorkspaceException($"Directory {existingDir.GetPath()} already exists; cannot replace with file");
 				}
-				else
+				if (!_files.TryAdd(fileState.Name, fileState))
 				{
-					FileState fileState = new FileState(this, name);
-					Files.Insert(~index, fileState);
-					return fileState;
+					throw new WorkspaceException($"File {fileState.Name} already exists");
 				}
 			}
 
-			public bool TryGetDirectory(string name, [NotNullWhen(true)] out DirectoryState? directoryState)
+			public bool ContainsFile(string name)
+				=> _files.ContainsKey(name);
+
+			public bool TryGetFileCaseSensitive(string name, [NotNullWhen(true)] out FileState? fileState)
+				=> _files.TryGetValue(name, out fileState) && name.Equals(fileState.Name, StringComparison.Ordinal);
+
+			public bool TryGetFileCaseInsensitive(string name, [NotNullWhen(true)] out FileState? fileState)
+				=> _files.TryGetValue(name, out fileState);
+
+			public DirectoryState? FindDirectory(string name)
+			{
+				_directories.TryGetValue(name, out DirectoryState? state);
+				return state;
+			}
+
+			public bool TryGetDirectoryCaseInsensitive(string name, [NotNullWhen(true)] out DirectoryState? directoryState)
+				=> _directories.TryGetValue(name, out directoryState);
+
+			public bool TryGetDirectoryCaseSensitive(string name, [NotNullWhen(true)] out DirectoryState? directoryState)
 			{
-				int index = Directories.BinarySearch(x => x.Name, name, StringComparer.Ordinal);
-				if (index >= 0)
+				if (_directories.TryGetValue(name, out DirectoryState? tempDirectoryState) && tempDirectoryState.Name.Equals(name, StringComparison.Ordinal))
 				{
-					directoryState = Directories[index];
+					directoryState = tempDirectoryState;
 					return true;
 				}
 				else
@@ -105,29 +322,34 @@ namespace EpicGames.Horde.Storage
 				}
 			}
 
-			public DirectoryState FindOrAddDirectory(string name)
+			public DirectoryState AddDirectory(string name, LayerFlags layer)
+			{
+				DirectoryState dirState = new DirectoryState(this, name) { Layers = layer };
+				_directories.Add(name, dirState);
+				return dirState;
+			}
+
+			public DirectoryState FindOrAddDirectory(string name, LayerFlags layer)
 			{
-				int index = Directories.BinarySearch(x => x.Name, name, StringComparer.Ordinal);
-				if (index >= 0)
+				DirectoryState? dirState;
+				if (TryGetDirectoryCaseSensitive(name, out dirState))
 				{
-					return Directories[index];
+					dirState.Layers |= layer;
+					return dirState;
 				}
 				else
 				{
-					DirectoryState subDirState = new DirectoryState(this, name);
-					Directories.Insert(~index, subDirState);
-					return subDirState;
+					return AddDirectory(name, layer);
 				}
 			}
 
-			public DirectoryState FindOrAddDirectoryByPath(string path)
+			public void Unlink()
 			{
-				DirectoryState dirState = this;
-				foreach (string fragment in path.Split('/'))
+				if (Parent != null)
 				{
-					dirState = dirState.FindOrAddDirectory(fragment);
+					Parent._directories.Remove(Name);
+					Parent = null;
 				}
-				return dirState;
 			}
 
 			public string GetPath()
@@ -137,17 +359,29 @@ namespace EpicGames.Horde.Storage
 				return builder.ToString();
 			}
 
+			public string GetChildPath(string name)
+			{
+				StringBuilder builder = new StringBuilder();
+				AppendPath(builder);
+				AppendChildPath(builder, name);
+				return builder.ToString();
+			}
+
 			public void AppendPath(StringBuilder builder)
 			{
 				Parent?.AppendPath(builder);
+				AppendChildPath(builder, Name);
+			}
 
-				if (Name.Length > 0)
-				{
-					builder.Append(Name);
-				}
-				if (builder.Length == 0 || builder[^1] != Path.DirectorySeparatorChar)
+			static void AppendChildPath(StringBuilder builder, string name)
+			{
+				if (name.Length > 0)
 				{
-					builder.Append(Path.DirectorySeparatorChar);
+					if (builder.Length > 0 && builder[^1] != Path.DirectorySeparatorChar)
+					{
+						builder.Append(Path.DirectorySeparatorChar);
+					}
+					builder.Append(name);
 				}
 			}
 
@@ -157,101 +391,79 @@ namespace EpicGames.Horde.Storage
 				AppendPath(builder);
 				return new DirectoryReference(builder.ToString(), DirectoryReference.Sanitize.None);
 			}
-
-			public void Read(IMemoryReader reader)
-			{
-				int numDirectories = reader.ReadInt32();
-				Directories.Capacity = numDirectories;
-				Directories.Clear();
-
-				for (int idx = 0; idx < numDirectories; idx++)
-				{
-					string subDirName = reader.ReadString();
-
-					DirectoryState subDirState = new DirectoryState(this, subDirName);
-					subDirState.Read(reader);
-
-					Directories.Add(subDirState);
-				}
-
-				int numFiles = reader.ReadInt32();
-				Files.Capacity = numFiles;
-				Files.Clear();
-
-				for (int idx = 0; idx < numFiles; idx++)
-				{
-					string fileName = reader.ReadString();
-
-					FileState fileState = new FileState(this, fileName);
-					fileState.Read(reader);
-
-					Files.Add(fileState);
-				}
-
-				LayerFlags = reader.ReadUnsignedVarInt();
-			}
-
-			public void Write(IMemoryWriter writer)
-			{
-				writer.WriteInt32(Directories.Count);
-				foreach (DirectoryState directory in Directories)
-				{
-					writer.WriteString(directory.Name);
-					directory.Write(writer);
-				}
-
-				writer.WriteInt32(Files.Count);
-				foreach (FileState file in Files)
-				{
-					writer.WriteString(file.Name);
-					file.Write(writer);
-				}
-
-				writer.WriteUnsignedVarInt(LayerFlags);
-			}
-
-			public override string ToString() => GetPath();
 		}
 
 		// Tracked state of a file
-		[DebuggerDisplay("{Name}")]
-		class FileState
+		[DebuggerDisplay("{GetPath()}")]
+		abstract class FileState
 		{
-			public DirectoryState Parent { get; set; }
-			public string Name { get; private set; }
+			public DirectoryState Parent { get; protected set; }
+			public string Name { get; protected set; }
 			public long Length { get; private set; }
 			public long LastModifiedTimeUtc { get; private set; }
-			public IoHash Hash { get; set; }
-			public ulong LayerFlags { get; set; }
+			public IoHash Hash { get; }
+			public LayerFlags Layers { get; set; }
+			public PendingFlags Pending { get; set; }
 
-			public FileState(DirectoryState parent, string name)
+			protected FileState(DirectoryState parent, string name, IoHash hash, LayerFlags layers)
 			{
 				Parent = parent;
 				Name = name;
+				Hash = hash;
+				Layers = layers;
 			}
 
-			public void Read(IMemoryReader reader)
+			public FileState(DirectoryState parent, IMemoryReader reader)
 			{
-				Length = reader.ReadInt64();
+				Parent = parent;
+				Name = reader.ReadString();
+				Length = (long)reader.ReadUnsignedVarInt();
 				LastModifiedTimeUtc = reader.ReadInt64();
 				Hash = reader.ReadIoHash();
+				Layers = (LayerFlags)reader.ReadUnsignedVarInt();
+
+				if ((Layers & LayerFlags.Pending) != 0)
+				{
+					Pending = (PendingFlags)reader.ReadUnsignedVarInt();
+				}
+			}
+
+			public void Write(IMemoryWriter writer)
+			{
+				Debug.Assert(Hash != IoHash.Zero);
+				writer.WriteString(Name);
+				writer.WriteUnsignedVarInt((ulong)Length);
+				writer.WriteInt64(LastModifiedTimeUtc);
+				writer.WriteIoHash(Hash);
+				writer.WriteUnsignedVarInt((ulong)Layers);
 
-				LayerFlags = reader.ReadUnsignedVarInt();
+				if ((Layers & LayerFlags.Pending) != 0)
+				{
+					writer.WriteUnsignedVarInt((ulong)Pending);
+				}
 			}
 
-			public void MoveTo(DirectoryState newParent, string newName)
+			public void SetPendingAction(PendingFlags pending)
 			{
-				Parent.RemoveFile(this);
-				Name = newName;
-				newParent.AddFile(this);
+				Pending |= pending;
+				MarkAsDirty();
 			}
 
-			public void Delete()
+			public void MarkAsDirty()
 			{
-				Parent.RemoveFile(this);
+				Layers |= LayerFlags.Pending;
+				Parent.MarkAsDirty();
 			}
 
-			public bool IsModified(FileInfo fileInfo) => Length != fileInfo.Length || LastModifiedTimeUtc != fileInfo.LastWriteTimeUtc.Ticks;
+			public bool InLayer(LayerFlags layers)
+				=> (Layers & layers) != 0;
+
+			public abstract void MoveTo(DirectoryState newParent, string newName, LayerFlags newLayers);
+
+			public abstract void Unlink();
+
+			public bool IsModified(FileInfo fileInfo)
+				=> Length != fileInfo.Length || LastModifiedTimeUtc != fileInfo.LastWriteTimeUtc.Ticks;
 
 			public void Update(FileInfo fileInfo)
 			{
@@ -259,15 +471,6 @@ namespace EpicGames.Horde.Storage
 				LastModifiedTimeUtc = fileInfo.LastWriteTimeUtc.Ticks;
 			}
 
-			public void Write(IMemoryWriter writer)
-			{
-				Debug.Assert(Hash != IoHash.Zero);
-				writer.WriteInt64(Length);
-				writer.WriteInt64(LastModifiedTimeUtc);
-				writer.WriteIoHash(Hash);
-				writer.WriteUnsignedVarInt(LayerFlags);
-			}
-
 			public string GetPath()
 			{
 				StringBuilder builder = new StringBuilder();
@@ -278,17 +481,14 @@ namespace EpicGames.Horde.Storage
 			public void AppendPath(StringBuilder builder)
 			{
 				Parent.AppendPath(builder);
-				builder.Append(Name);
-			}
 
-			public FileReference GetFileReference(DirectoryReference baseDir)
-			{
-				StringBuilder builder = new StringBuilder(baseDir.FullName);
-				AppendPath(builder);
-				return new FileReference(builder.ToString(), FileReference.Sanitize.None);
-			}
+				if (builder.Length > 0 && builder[^1] != Path.DirectorySeparatorChar)
+				{
+					builder.Append(Path.DirectorySeparatorChar);
+				}
 
-			public override string ToString() => GetPath();
+				builder.Append(Name);
+			}
 		}
 
 		// Collates lists of files and chunks with a particular hash
@@ -338,56 +538,22 @@ namespace EpicGames.Horde.Storage
 			}
 		}
 
-		// Maps a layer id to a flag
-		class LayerState
-		{
-			public WorkspaceLayerId Id { get; }
-			public ulong Flag { get; }
-
-			public LayerState(WorkspaceLayerId id, ulong flag)
-			{
-				Id = id;
-				Flag = flag;
-			}
-
-			public LayerState(IMemoryReader reader)
-			{
-				Id = new WorkspaceLayerId(new StringId(reader.ReadUtf8String()));
-				Flag = reader.ReadUnsignedVarInt();
-			}
-
-			public void Write(IMemoryWriter writer)
-			{
-				writer.WriteUtf8String(Id.Id.Text);
-				writer.WriteUnsignedVarInt(Flag);
-			}
-		}
-
 		readonly DirectoryReference _rootDir;
 		readonly FileReference _stateFile;
 		readonly DirectoryState _rootDirState;
-		readonly List<LayerState> _layers;
+		readonly DirectoryState _hordeDirState;
+		readonly DirectoryState _cacheDirState;
+		readonly List<LayerState> _customLayers = new List<LayerState>();
 		readonly Dictionary<IoHash, HashInfo> _hashes = new Dictionary<IoHash, HashInfo>();
 		readonly ILogger _logger;
 
-		const string HordeDirName = ".horde";
-		const string CacheDirName = "cache";
-		const string StateFileName = "contents.dat";
-
 		/// <summary>
-		/// Flag for the default layer
+		/// Metadata directory name
 		/// </summary>
-		const ulong DefaultLayerFlag = 1;
+		public const string HordeDirName = ".horde";
 
-		/// <summary>
-		/// Flag for the cache layer
-		/// </summary>
-		const ulong CacheLayerFlag = 2;
-
-		/// <summary>
-		/// Flags for user layers
-		/// </summary>
-		const ulong ReservedLayerFlags = DefaultLayerFlag | CacheLayerFlag;
+		const string CacheDirName = "cache";
+		const string StateFileName = "contents.dat";
 
 		/// <summary>
 		/// Root directory for the workspace
@@ -397,20 +563,32 @@ namespace EpicGames.Horde.Storage
 		/// <summary>
 		/// Layers current in this workspace
 		/// </summary>
-		public IReadOnlyList<WorkspaceLayerId> Layers => _layers.Select(x => x.Id).ToList();
+		public IReadOnlyList<WorkspaceLayerId> CustomLayers => _customLayers.Select(x => x.Id).ToList();
 
 		/// <summary>
 		/// Constructor
 		/// </summary>
 		/// <param name="rootDir">Directory for the workspace</param>
 		/// <param name="stateFile">Path to the state file for this directory</param>
+		/// <param name="reader">Reader for </param>
 		/// <param name="logger">Logger for diagnostic output</param>
-		private Workspace(DirectoryReference rootDir, FileReference stateFile, ILogger logger)
+		private Workspace(DirectoryReference rootDir, FileReference stateFile, IMemoryReader? reader, ILogger logger)
 		{
 			_rootDir = rootDir;
 			_stateFile = stateFile;
-			_rootDirState = new DirectoryState(null, String.Empty);
-			_layers = new List<LayerState> { new LayerState(WorkspaceLayerId.Default, 1) };
+
+			if (reader == null)
+			{
+				_rootDirState = new DirectoryState(null, String.Empty);
+			}
+			else
+			{
+				_rootDirState = new DirectoryState(null, reader);
+				Read(reader);
+			}
+
+			_hordeDirState = _rootDirState.FindOrAddDirectory(HordeDirName, LayerFlags.Cache);
+			_cacheDirState = _hordeDirState.FindOrAddDirectory(CacheDirName, LayerFlags.Cache);
 			_logger = logger;
 		}
 
@@ -428,10 +606,11 @@ namespace EpicGames.Horde.Storage
 			using FileStream? stream = FileTransaction.OpenRead(stateFile);
 			if (stream != null)
 			{
-				throw new InvalidOperationException($"Workspace already exists in {rootDir}; use Open instead.");
+				throw new WorkspaceException($"Workspace already exists in {rootDir}; use Open instead.");
 			}
 
-			Workspace workspace = new Workspace(rootDir, stateFile, logger);
+			Workspace workspace = new Workspace(rootDir, stateFile, null, logger);
+			DirectoryReference.CreateDirectory(workspace.GetDirectoryReference(workspace._cacheDirState));
 			await workspace.SaveAsync(cancellationToken);
 
 			return workspace;
@@ -467,8 +646,8 @@ namespace EpicGames.Horde.Storage
 				{
 					byte[] data = await stream.ReadAllBytesAsync(cancellationToken);
 
-					Workspace workspace = new Workspace(currentDir, stateFile, logger);
-					workspace.Read(new MemoryReader(data));
+					MemoryReader reader = new MemoryReader(data);
+					Workspace workspace = new Workspace(currentDir, stateFile, reader, logger);
 
 					return workspace;
 				}
@@ -517,10 +696,7 @@ namespace EpicGames.Horde.Storage
 
 		void Read(IMemoryReader reader)
 		{
-			_rootDirState.Read(reader);
-
-			_layers.Clear();
-			reader.ReadList(_layers, () => new LayerState(reader));
+			reader.ReadList(_customLayers, () => new LayerState(reader));
 
 			// Read the hash lookup
 			int numHashes = reader.ReadInt32();
@@ -551,7 +727,7 @@ namespace EpicGames.Horde.Storage
 			RemoveUnusedHashInfoObjects();
 
 			_rootDirState.Write(writer);
-			writer.WriteList(_layers, x => x.Write(writer));
+			writer.WriteList(_customLayers, x => x.Write(writer));
 
 			// Write the hash lookup
 			writer.WriteInt32(_hashes.Count);
@@ -609,23 +785,21 @@ namespace EpicGames.Horde.Storage
 		/// <param name="id">Identifier for the layer</param>
 		public void AddLayer(WorkspaceLayerId id)
 		{
-			if (_layers.Any(x => x.Id == id))
+			if (id != WorkspaceLayerId.Default && !_customLayers.Any(x => x.Id == id))
 			{
-				throw new InvalidOperationException($"Layer {id} already exists");
-			}
+				ulong flags = ((ulong)LayerFlags.Default - 1) | (ulong)LayerFlags.Default;
+				for (int idx = 0; idx < _customLayers.Count; idx++)
+				{
+					flags |= (ulong)_customLayers[idx].Flag;
+				}
+				if (flags == ~0UL)
+				{
+					throw new WorkspaceException("Maximum number of layers reached");
+				}
 
-			ulong flags = ReservedLayerFlags;
-			for (int idx = 0; idx < _layers.Count; idx++)
-			{
-				flags |= _layers[idx].Flag;
-			}
-			if (flags == ~0UL)
-			{
-				throw new InvalidOperationException("Maximum number of layers reached");
+				LayerFlags nextFlag = (LayerFlags)((flags + 1) & ~flags);
+				_customLayers.Add(new LayerState(id, nextFlag));
 			}
-
-			ulong nextFlag = (flags + 1) ^ flags;
-			_layers.Add(new LayerState(id, nextFlag));
 		}
 
 		/// <summary>
@@ -634,258 +808,868 @@ namespace EpicGames.Horde.Storage
 		/// <param name="layerId">Layer to update</param>
 		public void RemoveLayer(WorkspaceLayerId layerId)
 		{
-			int layerIdx = _layers.FindIndex(x => x.Id == layerId);
-			if (layerIdx > 0) // Note: Excluding default layer at index 0
+			int layerIdx = _customLayers.FindIndex(x => x.Id == layerId);
+			if (layerIdx >= 0)
 			{
-				LayerState layer = _layers[layerIdx];
-				if ((_rootDirState.LayerFlags & layer.Flag) != 0)
+				LayerState layer = _customLayers[layerIdx];
+				if ((_rootDirState.Layers & layer.Flag) != 0)
 				{
-					throw new InvalidOperationException($"Workspace still contains files for layer {layerId}");
+					throw new WorkspaceException($"Workspace still contains files for layer {layerId}");
 				}
-				_layers.RemoveAt(layerIdx);
+				_customLayers.RemoveAt(layerIdx);
 			}
 		}
 
-		LayerState? GetLayerState(WorkspaceLayerId layerId) => _layers.FirstOrDefault(x => x.Id == layerId);
+		// Gets the names of layers
+		List<string> GetLayerNames(LayerFlags layers)
+		{
+			List<string> names = new List<string>();
+			while (layers != 0)
+			{
+				LayerFlags layer = (layers & ~(layers - 1));
+				LayerState? layerState = _customLayers.FirstOrDefault(x => x.Flag == layer);
+#pragma warning disable CA1308 // Normalize strings to uppercase
+				names.Add(layerState?.Id.ToString() ?? $"{layer}".ToLowerInvariant());
+#pragma warning restore CA1308 // Normalize strings to uppercase
+			}
+			return names;
+		}
 
-		#endregion
+		// Gets the flag for a particular layer
+		LayerFlags GetLayerFlag(WorkspaceLayerId id)
+		{
+			if (id == WorkspaceLayerId.Default)
+			{
+				return LayerFlags.Default;
+			}
+
+			LayerState? customLayer = _customLayers.FirstOrDefault(x => x.Id == id);
+			if (customLayer == null)
+			{
+				throw new WorkspaceException($"Layer '{id}' does not exist");
+			}
+
+			return customLayer.Flag;
+		}
+
+		#endregion
+
+		#region Add
+
+		/// <summary>
+		/// Adds files to the workspace
+		/// </summary>
+		/// <param name="layerId">Layer to add the files to</param>
+		/// <param name="files">Files to add</param>
+		/// <param name="cancellationToken">Cancellation token for the operation</param>
+		public async Task AddAsync(WorkspaceLayerId layerId, List<FileInfo> files, CancellationToken cancellationToken = default)
+		{
+			LayerFlags layer = GetLayerFlag(layerId);
+
+			await using (AsyncPipeline pipeline = new AsyncPipeline(cancellationToken))
+			{
+				using SemaphoreSlim semaphore = new SemaphoreSlim(1);
+				ParallelBuzHashChunker chunker = new ParallelBuzHashChunker(pipeline, LeafChunkedDataNodeOptions.Default);
+
+				_ = pipeline.AddTask(token => EnumerateFilesToAddAsync(layer, files, semaphore, chunker.SourceWriter, token));
+				_ = pipeline.AddTask(chunker.OutputReader, (output, token) => RegisterChunksAsync(output, layer, semaphore, token));
+
+				await pipeline.WaitForCompletionAsync();
+			}
+
+			RemoveUnusedHashInfoObjects();
+			await SaveAsync(cancellationToken);
+		}
+
+		async Task EnumerateFilesToAddAsync(LayerFlags layer, List<FileInfo> files, SemaphoreSlim semaphore, ChannelWriter<ChunkerSource> chunkerSourceWriter, CancellationToken cancellationToken)
+		{
+			foreach (FileInfo file in files)
+			{
+				cancellationToken.ThrowIfCancellationRequested();
+
+				string[] fragments = new FileReference(file).MakeRelativeTo(_rootDir).Split(Path.DirectorySeparatorChar);
+
+				DirectoryState dirState = _rootDirState;
+				using (IDisposable semaLock = await semaphore.WaitDisposableAsync(cancellationToken))
+				{
+					for (int idx = 0; idx < fragments.Length - 1; idx++)
+					{
+						dirState = dirState.FindOrAddDirectory(fragments[idx], layer);
+					}
+
+					string fileName = fragments[fragments.Length - 1];
+					if (dirState.TryGetFileCaseSensitive(fileName, out FileState? fileState))
+					{
+						fileState.Layers |= layer;
+						continue;
+					}
+				}
+
+				FileToAdd fileToAdd = new FileToAdd(dirState, file);
+				await chunkerSourceWriter.WriteAsync(new FileChunkerSource(file, fileToAdd), cancellationToken);
+			}
+			chunkerSourceWriter.Complete();
+		}
+
+		#endregion
+
+		#region Filter
+
+		/// <summary>
+		/// Filters the contents of a layer
+		/// </summary>
+		/// <param name="layerId">Layer to filter</param>
+		/// <param name="filter">Filter to be applied to the tree</param>
+		/// <param name="cancellationToken">Cancellation token for the operation</param>
+		public async Task ForgetAsync(WorkspaceLayerId layerId, FileFilter filter, CancellationToken cancellationToken)
+		{
+			LayerFlags layer = GetLayerFlag(layerId);
+			ForgetMatchingFiles(layer, _rootDirState, "/", filter);
+			await SaveAsync(cancellationToken);
+		}
+
+		void ForgetMatchingFiles(LayerFlags layer, DirectoryState dirState, string namePrefix, FileFilter filter)
+		{
+			List<FileState> removeFiles = new List<FileState>();
+			foreach (FileState nextFile in dirState.Files)
+			{
+				if ((nextFile.Layers & layer) != 0)
+				{
+					string fileName = namePrefix + nextFile.Name;
+					if (filter.Matches(fileName))
+					{
+						nextFile.Layers ^= layer;
+						if (nextFile.Layers == 0)
+						{
+							removeFiles.Add(nextFile);
+						}
+					}
+				}
+			}
+
+			List<DirectoryState> removeDirectories = new List<DirectoryState>();
+			foreach (DirectoryState nextDirectory in dirState.Directories)
+			{
+				if ((nextDirectory.Layers & layer) != 0)
+				{
+					string nextNamePrefix = namePrefix + nextDirectory.Name;
+					if (filter.Matches(nextNamePrefix))
+					{
+						nextDirectory.Layers ^= layer;
+						if (nextDirectory.Layers == 0)
+						{
+							removeDirectories.Add(nextDirectory);
+						}
+					}
+					else if(filter.PossiblyMatches(nextNamePrefix))
+					{
+						ForgetMatchingFiles(layer, nextDirectory, nextNamePrefix + "/", filter);
+					}
+				}
+			}
+
+			foreach (FileState removeFile in removeFiles)
+			{
+				_logger.LogInformation("Forgetting {Path}", removeFile.GetPath());
+				RemoveFileFromHashLookup(removeFile);
+				removeFile.Unlink();
+			}
+
+			foreach (DirectoryState removeDirectory in removeDirectories)
+			{
+				_logger.LogInformation("Forgetting {Path}", removeDirectory.GetPath());
+				RemoveDirectoryFromHashLookup(removeDirectory);
+				removeDirectory.Unlink();
+			}
+		}
+
+		#endregion
+
+		#region Clean
+
+		class CleanVisitor(ILogger logger) : Visitor
+		{
+			public ConcurrentBag<FileState> RemoveFiles { get; } = new ConcurrentBag<FileState>();
+			public ConcurrentBag<DirectoryState> RemoveDirs { get; } = new ConcurrentBag<DirectoryState>();
+
+			public override void FileAdded(string path, DirectoryState parentDir, FileInfo fileInfo)
+			{
+				logger.LogInformation("Deleting {Path} (untracked)", path);
+				FileUtils.ForceDeleteFile(fileInfo);
+			}
+
+			public override void FileModified(string path, FileState fileState, FileInfo fileInfo)
+			{
+				logger.LogInformation("Deleting {Path} (modified)", path);
+				FileUtils.ForceDeleteFile(fileInfo);
+				RemoveFiles.Add(fileState);
+			}
+
+			public override void FileRemoved(string path, FileState fileState)
+			{
+				logger.LogInformation("Removing {Path} (missing)", path);
+				RemoveFiles.Add(fileState);
+			}
+
+			public override void DirectoryAdded(string path, DirectoryState parentDir, DirectoryInfo subDirInfo)
+			{
+				logger.LogInformation("Deleting {Path} (untracked)", path);
+				FileUtils.ForceDeleteDirectory(subDirInfo);
+			}
+
+			public override void DirectoryRemoved(string path, DirectoryState dirState)
+			{
+				logger.LogInformation("Removing {Path} (missing)", path);
+				RemoveDirs.Add(dirState);
+			}
+		}
+
+		/// <summary>
+		/// Prints the status of files in this workspace based on current filesystem metadata
+		/// </summary>
+		/// <param name="filter">Filter for files to include</param>
+		/// <param name="flags">Flags for changes to show</param>
+		/// <param name="cancellationToken">Cancellation token for the operation</param>
+		public async Task CleanAsync(FileFilter? filter = null, WorkspaceStatusFlags flags = WorkspaceStatusFlags.Default, CancellationToken cancellationToken = default)
+		{
+			filter ??= new FileFilter(FileFilterType.Include);
+
+			await using (AsyncPipeline pipeline = new AsyncPipeline(cancellationToken))
+			{
+				ChannelWithRefCount<DirectoryState> queue = Channel.CreateUnbounded<DirectoryState>().WithRefCount();
+				queue.Writer.TryWrite(_rootDirState);
+
+				CleanVisitor visitor = new CleanVisitor(_logger);
+				pipeline.AddTasks(16, ctx => ScanWorkspaceAsync(visitor, filter, flags, queue, ctx));
+
+				await pipeline.WaitForCompletionAsync();
+
+				foreach (DirectoryState removeDir in visitor.RemoveDirs)
+				{
+					removeDir.Unlink();
+				}
+
+				foreach (FileState removeFile in visitor.RemoveFiles)
+				{
+					removeFile.Unlink();
+				}
+			}
+
+			RemoveUnusedHashInfoObjects();
+			await SaveAsync(cancellationToken);
+		}
+
+		#endregion
 
 		#region Syncing
 
+		class DirectoryToSync(string name)
+		{
+			public string Name { get; } = name;
+			public DirectoryState? State { get; set; }
+			public Dictionary<string, DirectoryToSync> Directories { get; } = new Dictionary<string, DirectoryToSync>(DirectoryReference.Comparer);
+			public Dictionary<string, FileToSync> Files { get; } = new Dictionary<string, FileToSync>(FileReference.Comparer);
+
+			public DirectoryToSync? FindDirectoryExactCase(string name)
+				=> (Directories.TryGetValue(name, out DirectoryToSync? dirToSync) && Name.Equals(name, StringComparison.Ordinal)) ? dirToSync : null;
+
+			public FileToSync? FindFileExactCase(string name)
+				=> (Files.TryGetValue(name, out FileToSync? fileToSync) && Name.Equals(name, StringComparison.Ordinal)) ? fileToSync : null;
+		}
+
+		class FileToSync(FileEntry entry)
+		{
+			public FileEntry Entry => entry;
+			public string Name => Entry.Name;
+			public FileState? State { get; set; }
+
+			public HashInfo? CopyFrom { get; set; }
+			public FileState? MoveFrom { get; set; }
+			public List<LeafChunkedDataNodeRef>? Chunks { get; set; }
+			public OutputFile? OutputFile { get; set; }
+
+			public bool Complete { get; set; }
+		}
+
 		/// <summary>
-		/// Syncs a layer to the given contents
+		/// Syncs a layer to the given contents at a given base path
 		/// </summary>
 		/// <param name="layerId">Identifier for the layer</param>
 		/// <param name="basePath">Base path within the workspace to sync to.</param>
 		/// <param name="contents">New contents for the layer</param>
+		/// <param name="options">Options for extraction</param>
+		/// <param name="stats">Stats for the sync operation</param>
 		/// <param name="cancellationToken">Cancellation token for the operation</param>
-		public async Task SyncAsync(WorkspaceLayerId layerId, string? basePath, DirectoryNode? contents, CancellationToken cancellationToken = default)
+		public async Task SyncAsync(WorkspaceLayerId layerId, string? basePath, IBlobRef<DirectoryNode>? contents, ExtractOptions? options = null, WorkspaceSyncStats? stats = null, CancellationToken cancellationToken = default)
 		{
 			basePath = basePath?.Replace('\\', '/').Trim('/') ?? String.Empty;
 
 			await using IBlobWriter writer = new MemoryBlobWriter(BlobSerializerOptions.Default);
-			if (basePath.Length > 0)
+			if (basePath.Length > 0 && contents != null)
 			{
+				using BlobData blobData = await contents.ReadBlobDataAsync(cancellationToken);
+
+				IoHash hash = IoHash.Compute(blobData.Data.Span);
+				IHashedBlobRef<DirectoryNode> hashedContents = HashedBlobRef.Create<DirectoryNode>(hash, contents);
+
+				DirectoryNode current = await contents.ReadBlobAsync(cancellationToken);
 				foreach (string fragment in basePath.Split('/').Reverse())
 				{
 					DirectoryNode next = new DirectoryNode();
-					if (contents != null)
-					{
-						IHashedBlobRef<DirectoryNode> handle = await writer.WriteBlobAsync(contents, cancellationToken);
-						DirectoryEntry entry = new DirectoryEntry(fragment, contents.Length, handle);
-						next.AddDirectory(entry);
-					}
-					contents = next;
+
+					DirectoryEntry entry = new DirectoryEntry(fragment, current.Length, hashedContents);
+					next.AddDirectory(entry);
+					hashedContents = await writer.WriteBlobAsync(next, cancellationToken);
+
+					current = next;
 				}
-			}
 
-			LayerState layerState = GetLayerState(layerId) ?? throw new InvalidOperationException($"Layer '{layerId}' does not exist");
+				contents = hashedContents;
+			}
 
-			DirectoryState hordeDirState = _rootDirState.FindOrAddDirectory(HordeDirName);
-			hordeDirState.LayerFlags |= CacheLayerFlag;
+			await SyncAsync(layerId, contents, options, stats, cancellationToken);
+		}
 
-			DirectoryState cacheDirState = hordeDirState.FindOrAddDirectory(CacheDirName);
-			cacheDirState.LayerFlags |= CacheLayerFlag;
+		/// <summary>
+		/// Syncs a layer to the given contents
+		/// </summary>
+		/// <param name="layerId">Identifier for the layer</param>
+		/// <param name="contents">New contents for the layer</param>
+		/// <param name="options">Options for the extraction</param>
+		/// <param name="stats">Stats for the sync operation</param>
+		/// <param name="cancellationToken">Cancellation token for the operation</param>
+		public async Task SyncAsync(WorkspaceLayerId layerId, IBlobRef<DirectoryNode>? contents, ExtractOptions? options = null, WorkspaceSyncStats? stats = null, CancellationToken cancellationToken = default)
+		{
+			LayerFlags layer = GetLayerFlag(layerId);
 
-			ExtractOptions options = new ExtractOptions();
-			options.Progress = new ExtractStatsLogger(_logger);
+			options ??= new ExtractOptions { Progress = new ExtractStatsLogger(_logger) };
+			stats ??= new WorkspaceSyncStats();
 
+			// Find all the files that need to be synced, and create a plan for the existing tree to get there
 			_logger.LogInformation("Creating manifest...");
+			DirectoryToSync dirToSync = await CreateTreeToSyncAsync("", contents, cancellationToken);
+			CreateManifest(_rootDirState, dirToSync, layer);
+			ValidateManifest(_rootDirState, layer);
+			await SaveAsync(cancellationToken);
 
-			// Find all the files that need updating, and move anything that needs to be overwritten into the cache
-			List<FileUpdate> fileUpdates = new List<FileUpdate>();
-			await SyncDirectoryAsync(_rootDir, "", _rootDirState, contents, layerState.Flag, cacheDirState, fileUpdates, options, cancellationToken);
+			// Remove all the old files and update the workspace state to match the new tree
+			RemoveOldState(_rootDirState, layer);
+			CreateNewState(dirToSync, layer);
 			await SaveAsync(cancellationToken);
 
-			// Split the work into copies within the workspace, and reads from the upstream store
-			Channel<ChunkCopyRequest> chunkCopyRequests = Channel.CreateUnbounded<ChunkCopyRequest>();
-			Channel<ChunkReadRequest> chunkReadRequests = Channel.CreateUnbounded<ChunkReadRequest>();
+			// Materialize the state and remove all the temporary files
+			_logger.LogInformation("Updating workspace...");
+			CreateDirectories(_rootDir, dirToSync);
+			await MaterializeFilesAsync(dirToSync, layer, options, stats, cancellationToken);
+			RemoveIntermediates(_rootDirState);
+			UpdateMetadata(dirToSync, _rootDirState, layer);
+			await SaveAsync(cancellationToken);
+		}
+
+		// Figure out all the files and directories that need to be synced, reusing any tracked items to 
+		// short-circuit traversal through the tree.
+		async Task<DirectoryToSync> CreateTreeToSyncAsync(string name, IBlobRef<DirectoryNode>? dirNodeRef, CancellationToken cancellationToken)
+		{
+			DirectoryToSync dirToSync = new DirectoryToSync(name);
 
-			foreach (FileUpdate fileUpdate in fileUpdates)
+			if (dirNodeRef != null)
 			{
-				OutputFile outputFile = fileUpdate.OutputFile;
+				DirectoryNode dirNode = await dirNodeRef.ReadBlobAsync(cancellationToken);
 
-				long offset = 0;
-				foreach (LeafChunkedDataNodeRef leafNodeRef in fileUpdate.Chunks)
+				foreach (DirectoryEntry subDirEntry in dirNode.Directories)
+				{
+					DirectoryToSync subDirToSync = await CreateTreeToSyncAsync(subDirEntry.Name, subDirEntry.Handle, cancellationToken);
+					dirToSync.Directories.Add(subDirEntry.Name, subDirToSync);
+				}
+
+				foreach (FileEntry fileEntry in dirNode.Files)
 				{
-					outputFile.IncrementRemaining();
+					FileToSync fileToSync = new FileToSync(fileEntry);
 
-					SourceChunk? sourceChunk = TryGetSourceChunk(leafNodeRef.Hash);
-					if (sourceChunk != null)
+					HashInfo? copyFrom;
+					if (_hashes.TryGetValue(fileEntry.Hash, out copyFrom))
 					{
-						ChunkCopyRequest copyRequest = new ChunkCopyRequest(outputFile, offset, sourceChunk);
-						await chunkCopyRequests.Writer.WriteAsync(copyRequest, cancellationToken);
+						fileToSync.CopyFrom = copyFrom;
 					}
 					else
 					{
-						ChunkReadRequest readRequest = new ChunkReadRequest(outputFile, offset, leafNodeRef);
-						await chunkReadRequests.Writer.WriteAsync(readRequest, cancellationToken);
+						fileToSync.Chunks = await fileEntry.Target.EnumerateLeafNodesAsync(cancellationToken).ToListAsync(cancellationToken);
 					}
 
-					offset += leafNodeRef.Length;
+					dirToSync.Files.Add(fileEntry.Name, fileToSync);
 				}
 			}
-			chunkCopyRequests.Writer.Complete();
-			chunkReadRequests.Writer.Complete();
 
-			// Run the pipeline to create everything
-			_logger.LogInformation("Updating workspace...");
+			return dirToSync;
+		}
 
-			int numReadTasks = options.NumReadTasks ?? ExtractOptions.DefaultNumReadTasks;
-			int numDecodeTasks = options.NumDecodeTasks ?? ExtractOptions.DefaultNumDecodeTasks;
-			int numWriteTasks = 16;
+		#region Planning
 
-			using BatchOutputWriter outputWriter = new BatchOutputWriter(_logger);
-			using BatchChunkReader batchReader = new BatchChunkReader(outputWriter.RequestWriter) { VerifyHashes = options.VerifyOutput };
+		// Merges the manifest of files to sync with the existing directory state
+		static void CreateManifest(DirectoryState dirState, DirectoryToSync dirToSync, LayerFlags layer)
+		{
+			// Clear out all the pending actions
+			dirState.ClearPendingActions();
 
-			await using (AsyncPipeline pipeline = new AsyncPipeline(cancellationToken))
+			// Schedule everything that needs to be removed
+			ScheduleRemoveActions(dirState, dirToSync, layer);
+
+			// Reprieve any files that can be moved to new locations
+			ScheduleMoveActions(dirToSync, layer);
+
+			// Reprieve any files that can be harvested for chunks in new locations
+			ScheduleCopyActions(dirToSync);
+		}
+
+		// Schedule all the existing files in the current directory for removal
+		static void ScheduleRemoveActions(DirectoryState dirState, DirectoryToSync dirToSync, LayerFlags layer)
+		{
+			dirToSync.State = dirState;
+			dirToSync.State.Layers |= layer;
+
+			foreach (DirectoryState subDirState in dirState.Directories)
 			{
-				Task[] copyTasks = pipeline.AddTasks(8, ct => CopyChunksAsync(chunkCopyRequests.Reader, outputWriter.RequestWriter, ct));
-				Task[] readTasks = batchReader.AddToPipeline(pipeline, numReadTasks, numDecodeTasks, chunkReadRequests.Reader);
-				_ = Task.WhenAll(Enumerable.Concat(copyTasks, readTasks)).ContinueWith(_ => outputWriter.RequestWriter.TryComplete(), TaskScheduler.Default);
+				DirectoryToSync? subDirToSync;
+				if (dirToSync.Directories.TryGetValue(subDirState.Name, out subDirToSync))
+				{
+					if (!String.Equals(subDirState.Name, subDirToSync.Name, StringComparison.Ordinal))
+					{
+						// Remove everything, and set the flag to move through the cache if we need to reuse it.
+						subDirState.SetPendingAction(PendingFlags.RemoveFromLayer | PendingFlags.RemoveFromDisk | PendingFlags.ForceRenameOrRemove);
+					}
+					else
+					{
+						// Run recursively over the contents
+						ScheduleRemoveActions(subDirState, subDirToSync, layer);
+					}
+				}
+				else if (dirToSync.Files.ContainsKey(subDirState.Name))
+				{
+					// Must remove everything and rename if we want to keep it.
+					subDirState.SetPendingAction(PendingFlags.RemoveFromLayer | PendingFlags.RemoveFromDisk | PendingFlags.ForceRenameOrRemove);
+				}
+				else
+				{
+					if (subDirState.Layers == layer)
+					{
+						// Remove if from the layer and disk
+						subDirState.SetPendingAction(PendingFlags.RemoveFromLayer | PendingFlags.RemoveFromDisk);
+					}
+					else if ((subDirState.Layers & layer) != 0)
+					{
+						// Remove it from the layer
+						subDirState.SetPendingAction(PendingFlags.RemoveFromLayer);
+					}
+				}
+			}
 
-				Task[] writeTasks = outputWriter.AddToPipeline(pipeline, numWriteTasks);
+			foreach (FileState fileState in dirState.Files)
+			{
+				FileToSync? fileToSync;
+				if (dirToSync.Files.TryGetValue(fileState.Name, out fileToSync))
+				{
+					if (!String.Equals(fileState.Name, fileToSync.Name, StringComparison.Ordinal))
+					{
+						// Must remove and set cache flag
+						fileState.SetPendingAction(PendingFlags.RemoveFromLayer | PendingFlags.RemoveFromDisk | PendingFlags.ForceRenameOrRemove);
+					}
+					else if (fileState.Hash == fileToSync.Entry.Hash)
+					{
+						// Can retain it
+						fileState.Layers &= ~LayerFlags.Pending;
+						fileState.Layers |= layer;
+						fileState.Pending = PendingFlags.None;
 
-				if (options.Progress != null)
+						fileToSync.State = fileToSync.MoveFrom = fileState;
+					}
+					else
+					{
+						// Regular remove
+						fileState.SetPendingAction(PendingFlags.RemoveFromLayer | PendingFlags.RemoveFromDisk);
+					}
+				}
+				else if ((dirToSync.Directories.ContainsKey(fileState.Name)))
 				{
-					_ = pipeline.AddTask(ctx => DirectoryNodeExtract.UpdateStatsAsync(batchReader, outputWriter, Task.WhenAll(writeTasks), options.Progress, options.ProgressUpdateFrequency, ctx));
+					// Must remove it before we can write the new thing
+					fileState.SetPendingAction(PendingFlags.RemoveFromLayer | PendingFlags.RemoveFromDisk | PendingFlags.ForceRenameOrRemove);
 				}
+				else
+				{
+					if (fileState.Layers == layer)
+					{
+						// Remove if from the layer and disk
+						fileState.SetPendingAction(PendingFlags.RemoveFromLayer | PendingFlags.RemoveFromDisk);
+					}
+					else if ((fileState.Layers & layer) != 0)
+					{
+						// Remove it from the layer
+						fileState.SetPendingAction(PendingFlags.RemoveFromLayer);
+					}
+				}
+			}
+		}
 
-				await pipeline.WaitForCompletionAsync();
+		// Identify files which can be removed from other layers
+		static void ScheduleMoveActions(DirectoryToSync dirToSync, LayerFlags layer)
+		{
+			foreach (DirectoryToSync subDirToSync in dirToSync.Directories.Values)
+			{
+				ScheduleMoveActions(subDirToSync, layer);
 			}
 
-			// Update metadata for the workspace
-			foreach (FileUpdate fileUpdate in fileUpdates)
+			foreach (FileToSync fileToSync in dirToSync.Files.Values)
 			{
-				OutputFile outputFile = fileUpdate.OutputFile;
+				if (fileToSync.CopyFrom != null && fileToSync.MoveFrom == null)
+				{
+					FileState? sourceFile = FindFileToMove(fileToSync.CopyFrom, layer);
+					if (sourceFile != null)
+					{
+						if ((sourceFile.Pending & PendingFlags.RemoveFromDisk) != 0)
+						{
+							sourceFile.Pending ^= PendingFlags.RemoveFromDisk | PendingFlags.Move;
+						}
+						else if (sourceFile.Layers == LayerFlags.Cache)
+						{
+							sourceFile.SetPendingAction(PendingFlags.Move);
+						}
+						fileToSync.MoveFrom = sourceFile;
+					}
+				}
+			}
+		}
 
-				FileState fileState = fileUpdate.ParentDir.FindOrAddFile(outputFile.FileEntry.Name);
-				fileState.Hash = outputFile.FileEntry.Hash;
-				fileState.Update(outputFile.FileInfo);
+		static FileState? FindFileToMove(HashInfo hashInfo, LayerFlags layer)
+		{
+			FileState? cacheFile = null;
+			foreach (FileState fileState in hashInfo.Files)
+			{
+				if (fileState.Layers == (layer | LayerFlags.Pending) && (fileState.Pending & PendingFlags.RemoveFromDisk) != 0)
+				{
+					return fileState;
+				}
+				else if ((fileState.Layers & LayerFlags.Cache) != 0)
+				{
+					cacheFile = fileState;
+				}
+			}
+			return cacheFile;
+		}
 
-				HashInfo hashInfo = AddFileToHashLookup(fileState);
+		// Identify files which can be used
+		static void ScheduleCopyActions(DirectoryToSync dirToSync)
+		{
+			foreach (DirectoryToSync subDirToSync in dirToSync.Directories.Values)
+			{
+				ScheduleCopyActions(subDirToSync);
+			}
 
-				long offset = 0;
-				foreach (LeafChunkedDataNodeRef chunkRef in fileUpdate.Chunks)
+			foreach (FileToSync fileToSync in dirToSync.Files.Values)
+			{
+				if (fileToSync.CopyFrom != null)
 				{
-					ChunkInfo chunkInfo = new ChunkInfo(hashInfo, offset, chunkRef.Length);
+					FileState? sourceFile = FindFileForCopy(fileToSync.CopyFrom);
+					if (sourceFile != null && (sourceFile.Pending & PendingFlags.RemoveFromDisk) != 0)
+					{
+						sourceFile.Pending ^= PendingFlags.RemoveFromDisk | PendingFlags.Copy;
+					}
+				}
+			}
+		}
 
-					HashInfo chunkHashInfo = FindOrAddHashInfo(chunkRef.Hash, chunkRef.Length);
-					chunkHashInfo.AddChunk(chunkInfo);
+		// Find a file to copy from, preferring existing files in other layers first
+		static FileState FindFileForCopy(HashInfo hashInfo)
+			=> FindFileForCopy(hashInfo, PendingFlags.RemoveFromDisk, 0)
+			?? FindFileForCopy(hashInfo, PendingFlags.RemoveFromDisk, PendingFlags.RemoveFromDisk)
+			?? throw new InvalidDataException($"Unable to find any file matching hash; invalid state?");
 
-					offset += chunkRef.Length;
+		// Try to find any file matching the given layer flags
+		static FileState? FindFileForCopy(HashInfo hashInfo, PendingFlags pendingMask, PendingFlags pendingFlag)
+		{
+			foreach (FileState fileState in hashInfo.Files)
+			{
+				if ((fileState.Pending & pendingMask) == pendingFlag)
+				{
+					return fileState;
 				}
 			}
+			foreach (ChunkInfo chunkInfo in hashInfo.Chunks)
+			{
+				FileState? fileState = FindFileForCopy(chunkInfo.WithinHashInfo, pendingMask, pendingFlag);
+				if (fileState != null)
+				{
+					return fileState;
+				}
+			}
+			return null;
 		}
 
-		record class FileUpdate(OutputFile OutputFile, DirectoryState ParentDir, List<LeafChunkedDataNodeRef> Chunks);
-
-		async Task SyncDirectoryAsync(DirectoryReference dirRef, string path, DirectoryState dirState, DirectoryNode? dirNode, ulong flag, DirectoryState cacheDirState, List<FileUpdate> fileUpdates, ExtractOptions options, CancellationToken cancellationToken)
+		void ValidateManifest(DirectoryState dirState, LayerFlags layer)
 		{
-			// Remove any directories that no longer exist
-			for (int subDirIdx = 0; subDirIdx < dirState.Directories.Count; subDirIdx++)
+			foreach (DirectoryState subDirState in dirState.Directories)
 			{
-				DirectoryState subDirState = dirState.Directories[subDirIdx];
-				if ((subDirState.LayerFlags & flag) != 0)
+				if ((subDirState.Layers & LayerFlags.Pending) != 0)
 				{
-					if (dirNode == null || !dirNode.TryGetDirectoryEntry(subDirState.Name, out _))
+					LayerFlags otherLayers = subDirState.Layers & ~(layer | LayerFlags.Pending);
+					if ((subDirState.Pending & (PendingFlags.Move | PendingFlags.RemoveFromDisk)) != 0 && otherLayers != 0)
 					{
-						DirectoryReference subDirPath = DirectoryReference.Combine(dirRef, subDirState.Name.ToString());
-						await SyncDirectoryAsync(subDirPath, $"{path}{subDirState.Name}/", subDirState, null, flag, cacheDirState, fileUpdates, options, cancellationToken);
-
-						if (subDirState.LayerFlags == 0)
-						{
-							dirState.Directories.Remove(subDirState);
-							subDirIdx--;
-						}
+						throw new WorkspaceException($"Cannot remove {subDirState.GetPath()}: directory is referenced by other layers ({GetLayerNames(otherLayers)})");
+					}
+					else
+					{
+						ValidateManifest(subDirState, layer);
 					}
 				}
 			}
 
-			// Remove any files that no longer exist
-			for (int fileIdx = 0; fileIdx < dirState.Files.Count; fileIdx++)
+			foreach (FileState fileState in dirState.Files)
 			{
-				FileState fileState = dirState.Files[fileIdx];
-				if ((fileState.LayerFlags & flag) != 0)
+				LayerFlags otherLayers = fileState.Layers & ~(layer | LayerFlags.Pending);
+				if ((fileState.Pending & (PendingFlags.Move | PendingFlags.RemoveFromDisk)) != 0 && otherLayers != 0)
 				{
-					if (dirNode == null || !dirNode.TryGetFileEntry(fileState.Name, out FileEntry? entry) || entry.Hash != fileState.Hash)
-					{
-						fileState.LayerFlags &= ~flag;
-
-						if (fileState.LayerFlags == 0)
-						{
-							MoveFileToCache(fileState, cacheDirState);
-							fileIdx--;
-						}
-					}
+					throw new WorkspaceException($"Cannot remove {fileState.GetPath()}: file is referenced by other layers ({GetLayerNames(otherLayers)})");
 				}
 			}
+		}
 
-			// Clear out the layer flag for this directory. It'll be added back if we add/reuse files below.
-			dirState.LayerFlags &= ~flag;
+		#endregion
 
-			// Add files for this directory
-			if (dirNode != null)
+		static void CreateNewState(DirectoryToSync dirToSync, LayerFlags layer)
+		{
+			DirectoryState dirState = dirToSync.State ?? throw new InvalidDataException();
+			foreach ((string name, DirectoryToSync subDirToSync) in dirToSync.Directories)
 			{
-				DirectoryReference.CreateDirectory(dirRef);
+				subDirToSync.State ??= dirState.AddDirectory(name, layer);
+				CreateNewState(subDirToSync, layer);
+			}
+		}
 
-				// Update directories
-				foreach (DirectoryEntry subDirEntry in dirNode.Directories)
-				{
-					DirectoryReference subDirPath = DirectoryReference.Combine(dirRef, subDirEntry.Name.ToString());
-					DirectoryState subDirState = dirState.FindOrAddDirectory(subDirEntry.Name);
+		static void CreateDirectories(DirectoryReference dirRef, DirectoryToSync dirToSync)
+		{
+			// TODO: only if pending add
+			foreach ((string name, DirectoryToSync subDirToSync) in dirToSync.Directories)
+			{
+				DirectoryReference subDirRef = DirectoryReference.Combine(dirRef, name);
+				DirectoryReference.CreateDirectory(subDirRef);
 
-					DirectoryNode subDirNode = await subDirEntry.Handle.ReadBlobAsync(cancellationToken: cancellationToken);
-					await SyncDirectoryAsync(subDirPath, $"{path}{subDirEntry.Name}/", subDirState, subDirNode, flag, cacheDirState, fileUpdates, options, cancellationToken);
+				CreateDirectories(subDirRef, subDirToSync);
+			}
+		}
 
-					dirState.LayerFlags |= flag;
-				}
+		void RemoveIntermediates(DirectoryState dirState)
+		{
+			foreach (DirectoryState subDirState in dirState.Directories)
+			{
+				RemoveIntermediates(subDirState);
+			}
 
-				// Update files
-				foreach (FileEntry fileEntry in dirNode.Files)
+			foreach (FileState fileState in dirState.Files)
+			{
+				if ((fileState.Layers & LayerFlags.Cache) != 0 || (fileState.Pending & PendingFlags.Copy) != 0)
 				{
-					await SyncFileAsync(dirRef, $"{path}{fileEntry.Name}", dirState, fileEntry, flag, fileUpdates, cancellationToken);
-					dirState.LayerFlags |= flag;
+					FileInfo fileInfo = GetFileInfo(fileState);
+					FileUtils.ForceDeleteFile(fileInfo);
 				}
 			}
+		}
 
-			// Delete the directory if it's no longer needed
-			if (dirState.LayerFlags == 0 && dirState.Parent != null)
+		void UpdateMetadata(DirectoryToSync dirToSync, DirectoryState dirState, LayerFlags layer)
+		{
+			foreach ((string subDirName, DirectoryToSync subDirToSync) in dirToSync.Directories)
 			{
-				FileUtils.ForceDeleteDirectory(dirRef);
+				DirectoryState subDirState = dirState.FindOrAddDirectory(subDirName, layer);
+				UpdateMetadata(subDirToSync, subDirState, layer);
+			}
+
+			foreach (FileToSync fileToSync in dirToSync.Files.Values)
+			{
+				OutputFile outputFile = fileToSync.OutputFile!;
+
+				FileState? fileState = fileToSync.State;
+				if (fileState == null)
+				{
+					fileState = dirState.AddFile(outputFile.FileEntry.Name, fileToSync.Entry.Hash, layer);
+					fileState.Update(outputFile.FileInfo);
+
+					HashInfo hashInfo = AddFileToHashLookup(fileState);
+
+					if (fileToSync.Chunks != null)
+					{
+						long offset = 0;
+						foreach (LeafChunkedDataNodeRef chunkRef in fileToSync.Chunks!)
+						{
+							ChunkInfo chunkInfo = new ChunkInfo(hashInfo, offset, chunkRef.Length);
+
+							HashInfo chunkHashInfo = FindOrAddHashInfo(chunkRef.Hash, chunkRef.Length);
+							chunkHashInfo.AddChunk(chunkInfo);
+
+							offset += chunkRef.Length;
+						}
+					}
+				}
+				else
+				{
+					fileState.Update(outputFile.FileInfo);
+				}
 			}
 		}
 
-		async Task SyncFileAsync(DirectoryReference dirRef, string path, DirectoryState dirState, FileEntry fileEntry, ulong flag, List<FileUpdate> fileUpdates, CancellationToken cancellationToken)
+		void RemoveOldState(DirectoryState dirState, LayerFlags layer)
 		{
-			FileState? fileState;
-			if (dirState.TryGetFile(fileEntry.Name, out fileState))
+			List<DirectoryState> removeDirs = new List<DirectoryState>();
+			foreach (DirectoryState subDirState in dirState.Directories)
 			{
-				if (fileState.Hash == fileEntry.Hash)
+				RemoveOldState(subDirState, layer);
+
+				if ((subDirState.Pending & PendingFlags.RemoveFromLayer) != 0)
 				{
-					fileState.LayerFlags |= flag;
+					subDirState.Layers &= ~layer;
 				}
-				else
+				if ((subDirState.Pending & PendingFlags.RemoveFromDisk) != 0)
 				{
-					throw new InvalidOperationException($"Conflicting hash for file {fileState.GetFileReference(_rootDir)}");
+					removeDirs.Add(subDirState);
 				}
 			}
-			else
+
+			List<FileState> cacheFiles = new List<FileState>();
+			List<FileState> removeFiles = new List<FileState>();
+			foreach (FileState fileState in dirState.Files)
 			{
-				FileInfo fileInfo = FileReference.Combine(dirRef, fileEntry.Name.ToString()).ToFileInfo();
+				if ((fileState.Pending & PendingFlags.RemoveFromLayer) != 0)
+				{
+					fileState.Layers &= ~layer;
+				}
 
-				fileState = TryMoveCachedData(fileEntry.Hash, dirState, fileEntry.Name);
-				if (fileState == null)
+				if ((fileState.Pending & PendingFlags.ForceRenameOrRemove) != 0)
 				{
-					OutputFile outputFile = new OutputFile(path, fileInfo, fileEntry);
-					List<LeafChunkedDataNodeRef> chunks = await fileEntry.Target.EnumerateLeafNodesAsync(cancellationToken).ToListAsync(cancellationToken);
-					fileUpdates.Add(new FileUpdate(outputFile, dirState, chunks));
+					cacheFiles.Add(fileState);
 				}
-				else
+				else if ((fileState.Pending & PendingFlags.RemoveFromDisk) != 0)
+				{
+					removeFiles.Add(fileState);
+				}
+			}
+
+			foreach (FileState cacheFile in cacheFiles)
+			{
+				MoveFileToCache(cacheFile);
+			}
+
+			foreach (FileState removeFile in removeFiles)
+			{
+				DeleteFile(removeFile);
+			}
+
+			foreach (DirectoryState removeDir in removeDirs)
+			{
+				DeleteDirectory(removeDir);
+			}
+		}
+
+		async Task MaterializeFilesAsync(DirectoryToSync dirToSync, LayerFlags layer, ExtractOptions options, WorkspaceSyncStats stats, CancellationToken cancellationToken)
+		{
+			// Split the work into copies within the workspace, and reads from the upstream store
+			Channel<ChunkCopyRequest> chunkCopyRequests = Channel.CreateUnbounded<ChunkCopyRequest>();
+			Channel<ChunkReadRequest> chunkReadRequests = Channel.CreateUnbounded<ChunkReadRequest>();
+			await FindChunkActionsAsync("", dirToSync, _rootDirState, layer, chunkCopyRequests.Writer, chunkReadRequests.Writer, stats, cancellationToken);
+			chunkCopyRequests.Writer.Complete();
+			chunkReadRequests.Writer.Complete();
+
+			// Run the pipeline to create everything
+			int numReadTasks = options.NumReadTasks ?? ExtractOptions.DefaultNumReadTasks;
+			int numDecodeTasks = options.NumDecodeTasks ?? ExtractOptions.DefaultNumDecodeTasks;
+			int numWriteTasks = 16;
+
+			using BatchOutputWriter outputWriter = new BatchOutputWriter(_logger);
+			using BatchChunkReader batchReader = new BatchChunkReader(outputWriter.RequestWriter) { VerifyHashes = options.VerifyOutput };
+
+			await using (AsyncPipeline pipeline = new AsyncPipeline(cancellationToken))
+			{
+				Task[] copyTasks = pipeline.AddTasks(8, ct => CopyChunksAsync(chunkCopyRequests.Reader, outputWriter.RequestWriter, ct));
+				Task[] readTasks = batchReader.AddToPipeline(pipeline, numReadTasks, numDecodeTasks, chunkReadRequests.Reader);
+				_ = Task.WhenAll(Enumerable.Concat(copyTasks, readTasks)).ContinueWith(_ => outputWriter.RequestWriter.TryComplete(), TaskScheduler.Default);
+
+				Task[] writeTasks = outputWriter.AddToPipeline(pipeline, numWriteTasks);
+
+				if (options.Progress != null)
+				{
+					_ = pipeline.AddTask(ctx => DirectoryNodeExtract.UpdateStatsAsync(batchReader, outputWriter, Task.WhenAll(writeTasks), options.Progress, options.ProgressUpdateFrequency, ctx));
+				}
+
+				await pipeline.WaitForCompletionAsync();
+			}
+		}
+
+		async ValueTask FindChunkActionsAsync(string path, DirectoryToSync dirToSync, DirectoryState dirState, LayerFlags layer, ChannelWriter<ChunkCopyRequest> copyRequests, ChannelWriter<ChunkReadRequest> readRequests, WorkspaceSyncStats stats, CancellationToken cancellationToken)
+		{
+			foreach ((string subDirName, DirectoryToSync subDirToSync) in dirToSync.Directories)
+			{
+				DirectoryState subDirState = dirState.FindOrAddDirectory(subDirName, layer);
+				await FindChunkActionsAsync($"{path}{subDirName}/", subDirToSync, subDirState, layer, copyRequests, readRequests, stats, cancellationToken);
+			}
+
+			foreach ((string fileName, FileToSync fileToSync) in dirToSync.Files)
+			{
+				if (!fileToSync.Complete)
 				{
-					fileInfo.Refresh();
-					fileState.Update(fileInfo);
-					fileState.LayerFlags |= flag;
+					FileInfo fileInfo = FileReference.Combine(_rootDir, $"{path}{fileName}").ToFileInfo();
+					OutputFile outputFile = new OutputFile(path, fileInfo, fileToSync.Entry);
+					fileToSync.OutputFile = outputFile;
+
+					if (fileToSync.Chunks == null)
+					{
+						if (fileToSync.MoveFrom != null)
+						{
+							if (fileToSync.MoveFrom == fileToSync.State)
+							{
+								stats.NumFilesKept++;
+							}
+							else
+							{
+								MoveFile(fileToSync.MoveFrom, dirState, fileToSync.Entry.Name, layer);
+								fileToSync.State = fileToSync.MoveFrom;
+
+								stats.NumFilesMoved++;
+							}
+						}
+						else if (fileToSync.CopyFrom != null)
+						{
+							FileState sourceFile = fileToSync.CopyFrom!.Files[0];
+							SourceChunk sourceChunk = new SourceChunk(sourceFile, 0, sourceFile.Length, sourceFile.Hash);
+
+							outputFile.IncrementRemaining();
+
+							ChunkCopyRequest copyRequest = new ChunkCopyRequest(outputFile, 0, sourceChunk);
+							await copyRequests.WriteAsync(copyRequest, cancellationToken);
+
+							stats.NumBytesReused += sourceFile.Length;
+						}
+					}
+					else
+					{
+						long offset = 0;
+						foreach (LeafChunkedDataNodeRef leafNodeRef in fileToSync.Chunks)
+						{
+							outputFile.IncrementRemaining();
+
+							SourceChunk? sourceChunk = TryGetSourceChunk(leafNodeRef.Hash);
+							if (sourceChunk != null)
+							{
+								ChunkCopyRequest copyRequest = new ChunkCopyRequest(outputFile, offset, sourceChunk);
+								await copyRequests.WriteAsync(copyRequest, cancellationToken);
+
+								stats.NumBytesReused += leafNodeRef.Length;
+							}
+							else
+							{
+								ChunkReadRequest readRequest = new ChunkReadRequest(outputFile, offset, leafNodeRef);
+								await readRequests.WriteAsync(readRequest, cancellationToken);
+
+								stats.NumBytesDownloaded += leafNodeRef.Length;
+							}
+
+							offset += leafNodeRef.Length;
+						}
+					}
 				}
 			}
 		}
@@ -906,25 +1690,44 @@ namespace EpicGames.Horde.Storage
 				};
 
 				SourceChunk sourceChunk = copyRequest.SourceChunk;
-				FileReference sourceFile = sourceChunk.File.GetFileReference(RootDir);
+				FileInfo sourceFile = GetFileInfo(sourceChunk.File);
 
-				byte[] data = new byte[sourceChunk.Length];
-				using (FileStream sourceStream = new FileStream(sourceFile.FullName, options))
+				using Blake3.Hasher hasher = Blake3.Hasher.New();
+				using (FileStream sourceStream = sourceFile.Open(options))
 				{
 					sourceStream.Seek(sourceChunk.Offset, SeekOrigin.Begin);
-					await sourceStream.ReadExactlyAsync(data, cancellationToken);
-				}
 
-				IoHash hash = IoHash.Compute(data);
-				if (hash != sourceChunk.Hash)
-				{
-					throw new InvalidOperationException($"Source data does not have correct hash; expected {sourceChunk.Hash}, actually {hash} ({sourceFile}, offset: {sourceChunk.Offset}, length: {sourceChunk.Length})");
-				}
+					long remaining = sourceChunk.Length;
 
 #pragma warning disable CA2000
-				OutputChunk outputChunk = new OutputChunk(copyRequest.OutputFile, copyRequest.Offset, data, sourceChunk, null);
-				await outputRequests.WriteAsync([outputChunk], cancellationToken);
+					OutputChunk? outputChunk = null;
+					while (outputChunk == null || remaining > 0)
+					{
+						int length = (int)Math.Min(remaining, 256 * 1024);
+
+						byte[] data = new byte[length];
+						await sourceStream.ReadExactlyAsync(data, cancellationToken);
+
+						hasher.Update(data);
+
+						if (outputChunk != null)
+						{
+							copyRequest.OutputFile.IncrementRemaining();
+							await outputRequests.WriteAsync([outputChunk], cancellationToken);
+						}
+
+						outputChunk = new OutputChunk(copyRequest.OutputFile, copyRequest.Offset, data, sourceChunk, null);
+						remaining -= length;
+					}
+					await outputRequests.WriteAsync([outputChunk], cancellationToken);
 #pragma warning restore CA2000
+				}
+
+				IoHash hash = IoHash.FromBlake3(hasher);
+				if (hash != sourceChunk.Hash)
+				{
+					throw new WorkspaceException($"Source data does not have correct hash; expected {sourceChunk.Hash}, actually {hash} ({sourceFile}, offset: {sourceChunk.Offset}, length: {sourceChunk.Length})");
+				}
 			}
 		}
 
@@ -967,24 +1770,6 @@ namespace EpicGames.Horde.Storage
 			return null;
 		}
 
-		FileState? TryMoveCachedData(IoHash hash, DirectoryState targetDirState, string targetName)
-		{
-			HashInfo? hashInfo;
-			if (_hashes.TryGetValue(hash, out hashInfo))
-			{
-				foreach (FileState file in hashInfo.Files)
-				{
-					if (file.LayerFlags == CacheLayerFlag)
-					{
-						MoveFile(file, targetDirState, targetName);
-						file.LayerFlags = 0;
-						return file;
-					}
-				}
-			}
-			return null;
-		}
-
 		async Task<bool> TryCopyCachedDataAsync(HashInfo hashInfo, long offset, long length, Stream outputStream, CancellationToken cancellationToken)
 		{
 			foreach (FileState file in hashInfo.Files)
@@ -1034,47 +1819,432 @@ namespace EpicGames.Horde.Storage
 
 		#endregion
 
-		#region Refresh
+		class Visitor
+		{
+			public virtual void DirectoryAdded(string path, DirectoryState parentDir, DirectoryInfo subDirInfo)
+			{ }
+
+			public virtual void DirectoryRemoved(string path, DirectoryState dirState)
+			{ }
+
+			public virtual void FileAdded(string path, DirectoryState parentDir, FileInfo fileInfo)
+			{ }
+
+			public virtual void FileModified(string path, FileState fileState, FileInfo fileInfo)
+			{ }
+
+			public virtual void FileUnmodified(string path, FileState fileState, FileInfo fileInfo)
+			{ }
+
+			public virtual void FileRemoved(string path, FileState fileState)
+			{ }
+		}
+
+		async Task ScanWorkspaceAsync(Visitor visitor, FileFilter filter, WorkspaceStatusFlags flags, ChannelWithRefCount<DirectoryState> channel, CancellationToken cancellationToken)
+		{
+			await foreach (DirectoryState dirState in channel.Reader.ReadAllAsync(cancellationToken))
+			{
+				DirectoryInfo dirInfo = GetDirectoryReference(dirState).ToDirectoryInfo();
+
+				// Enumerate everything in this directory
+				Dictionary<string, FileInfo> nameToFileInfo = new Dictionary<string, FileInfo>(StringComparer.Ordinal);
+				Dictionary<string, DirectoryInfo> nameToDirectoryInfo = new Dictionary<string, DirectoryInfo>(StringComparer.Ordinal);
+				foreach (FileSystemInfo fileSystemInfo in dirInfo.EnumerateFileSystemInfos())
+				{
+					if (fileSystemInfo is FileInfo fileInfo)
+					{
+						nameToFileInfo.Add(fileInfo.Name, fileInfo);
+					}
+					else if (fileSystemInfo is DirectoryInfo directoryInfo && !IsRelativeDir(fileSystemInfo.Name))
+					{
+						nameToDirectoryInfo.Add(directoryInfo.Name, directoryInfo);
+					}
+				}
+
+				// Add new directories
+				foreach (DirectoryInfo subDirInfo in nameToDirectoryInfo.Values)
+				{
+					if (!dirState.TryGetDirectoryCaseSensitive(subDirInfo.Name, out _))
+					{
+						if ((flags & WorkspaceStatusFlags.DirectoryAdded) != 0)
+						{
+							string path = dirState.GetChildPath(subDirInfo.Name);
+							if (filter.Matches(path))
+							{
+								visitor.DirectoryAdded(path, dirState, subDirInfo);
+							}
+						}
+					}
+				}
+
+				// Ignore metadata files in the .horde directory
+				if (dirState != _hordeDirState)
+				{
+					// Add new files
+					foreach (FileInfo fileInfo in nameToFileInfo.Values)
+					{
+						if (!dirState.TryGetFileCaseSensitive(fileInfo.Name, out _))
+						{
+							if ((flags & WorkspaceStatusFlags.FileAdded) != 0)
+							{
+								string path = dirState.GetChildPath(fileInfo.Name);
+								if (filter.Matches(path))
+								{
+									visitor.FileAdded(path, dirState, fileInfo);
+								}
+							}
+						}
+					}
+
+					// Find all the files to remove
+					foreach (FileState fileState in dirState.Files)
+					{
+						if (!nameToFileInfo.TryGetValue(fileState.Name, out FileInfo? fileInfo))
+						{
+							if ((flags & WorkspaceStatusFlags.FileRemoved) != 0)
+							{
+								string path = fileState.GetPath();
+								if (filter.Matches(path))
+								{
+									visitor.FileRemoved(path, fileState);
+								}
+							}
+						}
+						else if (fileState.IsModified(fileInfo))
+						{
+							if ((flags & WorkspaceStatusFlags.FileModified) != 0)
+							{
+								string path = fileState.GetPath();
+								if (filter.Matches(path))
+								{
+									visitor.FileModified(path, fileState, fileInfo);
+								}
+							}
+						}
+						else
+						{
+							if ((flags & WorkspaceStatusFlags.FileUnmodified) != 0)
+							{
+								string path = fileState.GetPath();
+								if (filter.Matches(path))
+								{
+									visitor.FileUnmodified(path, fileState, fileInfo);
+								}
+							}
+						}
+					}
+				}
+
+				// Find directories that do not exist any more
+				foreach (DirectoryState subDirState in dirState.Directories)
+				{
+					if (nameToDirectoryInfo.ContainsKey(subDirState.Name))
+					{
+						string path = subDirState.GetPath();
+						if (filter.PossiblyMatches(path))
+						{
+							channel.AddRef();
+							await channel.Writer.WriteAsync(subDirState, cancellationToken);
+						}
+					}
+					else
+					{
+						if ((flags & WorkspaceStatusFlags.DirectoryRemoved) != 0)
+						{
+							string path = subDirState.GetPath();
+							if (filter.Matches(path))
+							{
+								visitor.DirectoryRemoved(path, subDirState);
+							}
+						}
+					}
+				}
+
+				// Release the completion counter
+				channel.Release();
+			}
+		}
+
+		#region Status
 
 		/// <summary>
-		/// Updates the status of files in this workspace based on current filesystem metadata
+		/// Prints the status of files in this workspace based on current filesystem metadata
 		/// </summary>
+		/// <param name="filter">Filter for files to include</param>
+		/// <param name="flags">Flags for changes to show</param>
 		/// <param name="cancellationToken">Cancellation token for the operation</param>
-		public async Task RefreshAsync(CancellationToken cancellationToken = default)
+		public async Task StatusAsync(FileFilter? filter = null, WorkspaceStatusFlags flags = WorkspaceStatusFlags.Default, CancellationToken cancellationToken = default)
 		{
-			await RefreshAsync(_rootDirState, cancellationToken);
-			RemoveUnusedHashInfoObjects();
+			await using (AsyncPipeline pipeline = new AsyncPipeline(cancellationToken))
+			{
+				ChannelWithRefCount<DirectoryState> directoryQueue = Channel.CreateUnbounded<DirectoryState>().WithRefCount();
+				directoryQueue.Writer.TryWrite(_rootDirState);
+
+				filter ??= new FileFilter(FileFilterType.Include);
+
+				StatusVisitor visitor = new StatusVisitor();
+				pipeline.AddTasks(16, ctx => ScanWorkspaceAsync(visitor, filter, flags, directoryQueue, ctx));
+
+				await pipeline.WaitForCompletionAsync();
+			}
 		}
 
-		async Task RefreshAsync(DirectoryState dirState, CancellationToken cancellationToken)
+		class StatusVisitor : Visitor
 		{
-			foreach (DirectoryState subDirState in dirState.Directories)
+			public bool AnyChanges { get; private set; }
+
+			readonly object _lockObject = new object();
+
+			public override void DirectoryAdded(string path, DirectoryState parentDir, DirectoryInfo subDirInfo)
+				=> WriteLine($"{FormatPath(path)}/ (added)", ConsoleColor.Green);
+
+			public override void DirectoryRemoved(string path, DirectoryState dirState)
+				=> WriteLine($"{FormatPath(path)}/ (removed)", ConsoleColor.Red);
+
+			public override void FileAdded(string path, DirectoryState parentDir, FileInfo fileInfo)
+				=> WriteLine($"{FormatPath(path)} (added)", ConsoleColor.Green);
+
+			public override void FileModified(string path, FileState fileState, FileInfo fileInfo)
+				=> WriteLine($"{FormatPath(path)} (modified)", ConsoleColor.Yellow);
+
+			public override void FileUnmodified(string path, FileState fileState, FileInfo fileInfo)
+				=> WriteLine(FormatPath(path), null);
+
+			public override void FileRemoved(string path, FileState fileState)
+				=> WriteLine($"{FormatPath(path)} (removed)", ConsoleColor.Red);
+
+			static string FormatPath(string path)
+				=> path.Replace(Path.DirectorySeparatorChar, '/');
+
+			void WriteLine(string text, ConsoleColor? color)
 			{
-				await RefreshAsync(subDirState, cancellationToken);
+				lock (_lockObject)
+				{
+					AnyChanges = true;
+
+					if (color == null)
+					{
+						Console.WriteLine(text);
+					}
+					else
+					{
+						Console.ForegroundColor = color.Value;
+						Console.WriteLine(text);
+						Console.ResetColor();
+					}
+				}
 			}
+		}
+
+		#endregion
+
+		#region Reconcile
+
+		/// <summary>
+		/// Updates the status of files in this workspace based on current filesystem metadata
+		/// </summary>
+		/// <param name="cancellationToken">Cancellation token for the operation</param>
+		public Task ReconcileAsync(CancellationToken cancellationToken = default)
+			=> TrackAsync(WorkspaceLayerId.Default, null, cancellationToken);
+
+		/// <summary>
+		/// Updates the status of files in this workspace based on current filesystem metadata
+		/// </summary>
+		/// <param name="layerId">Layer to add new files to</param>
+		/// <param name="filter">Filter for the files to add</param>
+		/// <param name="cancellationToken">Cancellation token for the operation</param>
+		public async Task TrackAsync(WorkspaceLayerId layerId, FileFilter? filter = null, CancellationToken cancellationToken = default)
+		{
+			LayerFlags layer = GetLayerFlag(layerId);
+
+			filter ??= new FileFilter(FileFilterType.Include);
 
-			int outputIdx = 0;
-			for (int idx = 0; idx < dirState.Files.Count; idx++)
+			byte[] buffer = new byte[2 * 1024 * 1024];
+			await using (AsyncPipeline pipeline = new AsyncPipeline(cancellationToken))
 			{
-				FileState fileState = dirState.Files[idx];
+				using SemaphoreSlim semaphore = new SemaphoreSlim(1);
+				ParallelBuzHashChunker chunker = new ParallelBuzHashChunker(pipeline, buffer, LeafChunkedDataNodeOptions.Default);
 
-				FileInfo fileInfo = GetFileInfo(fileState);
-				if (!fileInfo.Exists)
+				ChannelWithRefCount<DirectoryState> dirsToScan = Channel.CreateUnbounded<DirectoryState>().WithRefCount();
+				await dirsToScan.Writer.WriteAsync(_rootDirState, cancellationToken);
+				dirsToScan.OnComplete += () => chunker.SourceWriter.Complete();
+
+				_ = pipeline.AddTasks(16, token => FindFilesToReconcileAsync(filter, dirsToScan, layer, semaphore, chunker, token));
+				_ = pipeline.AddTask(chunker.OutputReader, (output, token) => RegisterChunksAsync(output, layer, semaphore, token));
+
+				await pipeline.WaitForCompletionAsync();
+			}
+
+			RemoveUnusedHashInfoObjects();
+			await SaveAsync(cancellationToken);
+		}
+
+		static bool IsRelativeDir(string name)
+			=> (name.Length == 1 && name[0] == '.') || (name.Length == 2 && name[0] == '.' && name[1] == '.');
+
+		record class FileToAdd(DirectoryState ParentDir, FileInfo FileInfo);
+
+		async Task FindFilesToReconcileAsync(FileFilter filter, ChannelWithRefCount<DirectoryState> channel, LayerFlags layer, SemaphoreSlim semaphore, ContentChunker chunker, CancellationToken cancellationToken)
+		{
+			await foreach (DirectoryState dirState in channel.Reader.ReadAllAsync(cancellationToken))
+			{
+				if (dirState == _hordeDirState)
+				{
+					channel.Release();
+					continue;
+				}
+
+				DirectoryInfo dirInfo = GetDirectoryReference(dirState).ToDirectoryInfo();
+				FileSystemInfo[] fileSystemInfos = dirInfo.GetFileSystemInfos();
+
+				List<DirectoryState> addedDirectories = new List<DirectoryState>();
+				List<FileInfo> addedFiles = new List<FileInfo>();
+				using (IDisposable stateLock = await semaphore.WaitDisposableAsync(cancellationToken))
+				{
+					// Duplicate the current files and subdirectories before we add anything new
+					FileState[] files = dirState.Files.ToArray();
+					DirectoryState[] subDirStates = dirState.Directories.ToArray();
+
+					// Enumerate everything in this directory
+					Dictionary<string, FileInfo> nameToFileInfo = new Dictionary<string, FileInfo>(StringComparer.Ordinal);
+					Dictionary<string, DirectoryInfo> nameToDirectoryInfo = new Dictionary<string, DirectoryInfo>(StringComparer.Ordinal);
+					foreach (FileSystemInfo fileSystemInfo in fileSystemInfos)
+					{
+						if (fileSystemInfo is FileInfo fileInfo)
+						{
+							nameToFileInfo.Add(fileInfo.Name, fileInfo);
+						}
+						else if (fileSystemInfo is DirectoryInfo directoryInfo && !IsRelativeDir(fileSystemInfo.Name))
+						{
+							nameToDirectoryInfo.Add(directoryInfo.Name, directoryInfo);
+						}
+					}
+
+					// Add new directories
+					foreach (DirectoryInfo subDirInfo in nameToDirectoryInfo.Values)
+					{
+						if (!dirState.TryGetDirectoryCaseSensitive(subDirInfo.Name, out _))
+						{
+							string path = dirState.GetChildPath(subDirInfo.Name);
+							if (filter.Matches(path))
+							{
+								DirectoryState subDirState = dirState.AddDirectory(subDirInfo.Name, layer);
+								_logger.LogInformation("Adding {Dir}", GetDirectoryReference(subDirState));
+								addedDirectories.Add(subDirState);
+							}
+						}
+					}
+
+					// Add new files
+					foreach (FileInfo fileInfo in nameToFileInfo.Values)
+					{
+						if (!dirState.TryGetFileCaseSensitive(fileInfo.Name, out _))
+						{
+							string path = dirState.GetChildPath(fileInfo.Name);
+							if (filter.Matches(path))
+							{
+								_logger.LogInformation("Adding {File}", new FileReference(fileInfo));
+								addedFiles.Add(fileInfo);
+							}
+						}
+					}
+
+					// Find all the files to remove
+					foreach (FileState fileState in files)
+					{
+						string path = fileState.GetPath();
+						if (filter.Matches(path))
+						{
+							if (!nameToFileInfo.TryGetValue(fileState.Name, out FileInfo? fileInfo))
+							{
+								_logger.LogInformation("Removing {File}", GetFileReference(fileState));
+								RemoveFileFromHashLookup(fileState);
+								fileState.Unlink();
+							}
+							else if (fileState.IsModified(fileInfo))
+							{
+								_logger.LogInformation("Updating {File}", new FileReference(fileInfo));
+
+								RemoveFileFromHashLookup(fileState);
+								fileState.Unlink();
+
+								addedFiles.Add(fileInfo);
+							}
+						}
+					}
+
+					// Find directories that do not exist any more
+					foreach (DirectoryState subDirState in subDirStates)
+					{
+						string path = subDirState.GetPath();
+						if (!nameToDirectoryInfo.ContainsKey(subDirState.Name))
+						{
+							if (filter.Matches(path))
+							{
+								_logger.LogInformation("Removing {File}", GetDirectoryReference(subDirState));
+
+								RemoveDirectoryFromHashLookup(subDirState);
+								subDirState.Unlink();
+							}
+						}
+						else
+						{
+							if (filter.PossiblyMatches(path))
+							{
+								addedDirectories.Add(subDirState);
+							}
+						}
+					}
+				}
+
+				foreach (DirectoryState addedDirectory in addedDirectories)
 				{
-					_logger.LogWarning("File {Path} has been removed", fileInfo.FullName);
-					RemoveFileFromHashLookup(fileState);
+					channel.AddRef();
+					await channel.Writer.WriteAsync(addedDirectory, cancellationToken);
 				}
-				else if (fileState.IsModified(fileInfo))
+
+				foreach (FileInfo addedFile in addedFiles)
 				{
-					_logger.LogWarning("File {Path} has been modified", fileInfo.FullName);
-					RemoveFileFromHashLookup(fileState);
+					FileChunkerSource fileSource = new FileChunkerSource(addedFile, new FileToAdd(dirState, addedFile));
+					await chunker.SourceWriter.WriteAsync(fileSource, cancellationToken);
 				}
-				else
+
+				channel.Release();
+			}
+		}
+
+		async ValueTask RegisterChunksAsync(ChunkerOutput output, LayerFlags layer, SemaphoreSlim stateSema, CancellationToken cancellationToken)
+		{
+			FileToAdd addedFile = (FileToAdd)output.UserData!;
+
+			using Blake3.Hasher hasher = Blake3.Hasher.New();
+
+			List<(IoHash, long)> chunks = new List<(IoHash, long)>();
+			while (await output.MoveNextAsync(cancellationToken))
+			{
+				IoHash chunkHash = IoHash.Compute(output.Data.Span);
+				chunks.Add((chunkHash, output.Data.Length));
+				hasher.Update(output.Data.Span);
+			}
+
+			IoHash hash = IoHash.FromBlake3(hasher);
+
+			using (IDisposable stateLock = await stateSema.WaitDisposableAsync(cancellationToken))
+			{
+				FileState fileState = addedFile.ParentDir.AddFile(addedFile.FileInfo.Name, hash, layer);
+				fileState.Update(addedFile.FileInfo);
+
+				HashInfo fileHashInfo = FindOrAddHashInfo(hash, addedFile.FileInfo.Length);
+				fileHashInfo.Files.Add(fileState);
+
+				long offset = 0;
+				foreach ((IoHash chunkHash, long length) in chunks)
 				{
-					dirState.Files[outputIdx++] = fileState;
+					HashInfo chunkHashInfo = FindOrAddHashInfo(chunkHash, length);
+					chunkHashInfo.AddChunk(new ChunkInfo(fileHashInfo, offset, length));
+					offset += length;
 				}
 			}
-			dirState.Files.RemoveRange(outputIdx, dirState.Files.Count - outputIdx);
 		}
 
 		#endregion
@@ -1101,7 +2271,7 @@ namespace EpicGames.Horde.Storage
 				FileInfo fileInfo = GetFileInfo(fileState);
 				if (fileState.IsModified(fileInfo))
 				{
-					throw new InvalidDataException($"File {fileInfo.FullName} has been modified");
+					throw new WorkspaceException($"File {fileInfo.FullName} has been modified");
 				}
 
 				IoHash hash;
@@ -1112,7 +2282,7 @@ namespace EpicGames.Horde.Storage
 
 				if (hash != fileState.Hash)
 				{
-					throw new InvalidDataException($"Hash for {fileInfo.FullName} was {hash}, expected {fileState.Hash}");
+					throw new WorkspaceException($"Hash for {fileInfo.FullName} was {hash}, expected {fileState.Hash}");
 				}
 			}
 		}
@@ -1158,35 +2328,62 @@ namespace EpicGames.Horde.Storage
 			}
 		}
 
-		void MoveFileToCache(FileState fileState, DirectoryState cacheDirState)
+		void RemoveDirectoryFromHashLookup(DirectoryState dirState)
+		{
+			foreach (DirectoryState subDirState in dirState.Directories)
+			{
+				RemoveDirectoryFromHashLookup(subDirState);
+			}
+			foreach (FileState fileState in dirState.Files)
+			{
+				RemoveFileFromHashLookup(fileState);
+			}
+		}
+
+		#region Directory operations
+
+		void DeleteDirectory(DirectoryState dirState)
+		{
+			DirectoryReference dir = GetDirectoryReference(dirState);
+			FileUtils.ForceDeleteDirectory(dir);
+
+			dirState.Unlink();
+		}
+
+		DirectoryReference GetDirectoryReference(DirectoryState dirState)
+		{
+			StringBuilder builder = new StringBuilder(_rootDir.FullName);
+			dirState.AppendPath(builder);
+			return new DirectoryReference(builder.ToString(), DirectoryReference.Sanitize.None);
+		}
+
+		#endregion
+
+		#region File operations
+
+		void MoveFileToCache(FileState fileState)
 		{
-			string name = fileState.Hash.ToString();
+			DirectoryState cacheDirState = _cacheDirState;
+
+			string name = Guid.NewGuid().ToString("N");
 			for (int idx = 0; idx < 2; idx++)
 			{
-				cacheDirState = cacheDirState.FindOrAddDirectory(name.Substring(idx * 2, 2));
+				cacheDirState = cacheDirState.FindOrAddDirectory(name.Substring(idx * 2, 2), LayerFlags.Cache);
 			}
 
-			FileInfo fileInfo = fileState.GetFileReference(_rootDir).ToFileInfo();
+			FileInfo fileInfo = GetFileInfo(fileState);
 			if (fileInfo.Exists)
 			{
-				if (cacheDirState.ContainsFile(name) || fileState.IsModified(fileInfo))
-				{
-					DeleteFile(fileState);
-				}
-				else
-				{
-					MoveFile(fileState, cacheDirState, name);
-					fileState.LayerFlags = CacheLayerFlag;
-				}
+				MoveFile(fileState, cacheDirState, name, LayerFlags.Cache);
 			}
 		}
 
-		void MoveFile(FileState fileState, DirectoryState targetDirState, string targetName)
+		void MoveFile(FileState fileState, DirectoryState targetDirState, string targetName, LayerFlags targetLayers)
 		{
-			FileReference sourceFile = fileState.GetFileReference(_rootDir);
-			fileState.MoveTo(targetDirState, targetName);
+			FileReference sourceFile = GetFileReference(fileState);
+			fileState.MoveTo(targetDirState, targetName, targetLayers);
 
-			FileReference targetFile = fileState.GetFileReference(_rootDir);
+			FileReference targetFile = GetFileReference(fileState);
 			DirectoryReference.CreateDirectory(targetFile.Directory);
 			FileReference.Move(sourceFile, targetFile, true);
 
@@ -1195,16 +2392,99 @@ namespace EpicGames.Horde.Storage
 
 		void DeleteFile(FileState fileState)
 		{
-			FileReference fileRef = fileState.GetFileReference(_rootDir);
-			FileUtils.ForceDeleteFile(fileRef);
-			fileState.Delete();
+			RemoveFileFromHashLookup(fileState);
+
+			FileReference file = GetFileReference(fileState);
+			FileUtils.ForceDeleteFile(file);
+
+			fileState.Unlink();
+		}
 
-			_logger.LogDebug("Deleting file {File}", fileRef);
+		FileInfo GetFileInfo(FileState fileState)
+		{
+			return GetFileReference(fileState).ToFileInfo();
 		}
 
-		FileInfo GetFileInfo(FileState file)
+		FileReference GetFileReference(FileState fileState)
 		{
-			return file.GetFileReference(_rootDir).ToFileInfo();
+			StringBuilder builder = new StringBuilder(_rootDir.FullName);
+			fileState.AppendPath(builder);
+			return new FileReference(builder.ToString(), FileReference.Sanitize.None);
 		}
+
+		#endregion
+	}
+
+	/// <summary>
+	/// Flags for showing modified flags in a workspace 
+	/// </summary>
+	[Flags]
+	public enum WorkspaceStatusFlags
+	{
+		/// <summary>
+		/// Include nothing
+		/// </summary>
+		None = 0,
+
+		/// <summary>
+		/// Directories which have been added
+		/// </summary>
+		DirectoryAdded = 1,
+
+		/// <summary>
+		/// Directories which have been removed
+		/// </summary>
+		DirectoryRemoved = 2,
+
+		/// <summary>
+		/// Files which have been added
+		/// </summary>
+		FileAdded = 4,
+
+		/// <summary>
+		/// Files which have been removed
+		/// </summary>
+		FileRemoved = 8,
+
+		/// <summary>
+		/// Files which have been modified
+		/// </summary>
+		FileModified = 16,
+
+		/// <summary>
+		/// Files which have not been modified
+		/// </summary>
+		FileUnmodified = 32,
+
+		/// <summary>
+		/// Default flags for status calls
+		/// </summary>
+		Default = DirectoryAdded | DirectoryRemoved | FileAdded | FileRemoved | FileModified,
+	}
+
+	/// <summary>
+	/// Stats for a sync
+	/// </summary>
+	public class WorkspaceSyncStats
+	{
+		/// <summary>
+		/// Number of identical files that were kept
+		/// </summary>
+		public int NumFilesKept { get; set; }
+
+		/// <summary>
+		/// Number of files that were moved during a sync operation
+		/// </summary>
+		public int NumFilesMoved { get; set; }
+
+		/// <summary>
+		/// Number of bytes that were copied from elsewhere in the tree
+		/// </summary>
+		public long NumBytesReused { get; set; }
+
+		/// <summary>
+		/// Number of bytes that were downloaded
+		/// </summary>
+		public long NumBytesDownloaded { get; set; }
 	}
 }
diff --git a/Engine/Source/Programs/Shared/EpicGames.Horde/Storage/WorkspaceLayerId.cs b/Engine/Source/Programs/Shared/EpicGames.Horde/Storage/WorkspaceLayerId.cs
--- a/Engine/Source/Programs/Shared/EpicGames.Horde/Storage/WorkspaceLayerId.cs
+++ b/Engine/Source/Programs/Shared/EpicGames.Horde/Storage/WorkspaceLayerId.cs
@@ -1,6 +1,7 @@
 // Copyright Epic Games, Inc. All Rights Reserved.
 
 using System.ComponentModel;
+using System.Diagnostics.CodeAnalysis;
 using EpicGames.Core;
 
 namespace EpicGames.Horde.Storage
@@ -19,6 +20,31 @@ namespace EpicGames.Horde.Storage
 		/// </summary>
 		public static WorkspaceLayerId Default { get; } = new WorkspaceLayerId(new StringId("default"));
 
+		/// <summary>
+		/// Constructor
+		/// </summary>
+		public WorkspaceLayerId(string id)
+			: this(new StringId(id))
+		{ }
+
+		/// <summary>
+		/// Parse an identifier
+		/// </summary>
+		public static bool TryParse(string text, out WorkspaceLayerId id, [NotNullWhen(false)] out string? errorMessage)
+		{
+			StringId stringId;
+			if (StringId.TryParse(text, out stringId, out errorMessage))
+			{
+				id = new WorkspaceLayerId(stringId);
+				return true;
+			}
+			else
+			{
+				id = default;
+				return false;
+			}
+		}
+
 		/// <inheritdoc/>
 		public override string ToString() => Id.ToString();
 	}
diff --git a/Engine/Source/Programs/Shared/EpicGames.Horde/StringId.cs b/Engine/Source/Programs/Shared/EpicGames.Horde/StringId.cs
--- a/Engine/Source/Programs/Shared/EpicGames.Horde/StringId.cs
+++ b/Engine/Source/Programs/Shared/EpicGames.Horde/StringId.cs
@@ -65,7 +65,12 @@ namespace EpicGames.Horde
 		/// <param name="text">Unique id for the string</param>
 		public StringId(Utf8String text)
 		{
-			Text = ValidateArgument(text, nameof(text));
+			StringId id;
+			if (!TryParse(text, out id, out string? errorMessage))
+			{
+				throw new ArgumentException(errorMessage, nameof(text));
+			}
+			Text = id.Text;
 		}
 
 		/// <summary>
@@ -133,18 +138,51 @@ namespace EpicGames.Horde
 		/// Validates the given string as a StringId, normalizing it if necessary.
 		/// </summary>
 		/// <param name="text">Text to validate as a StringId</param>
-		/// <param name="paramName">Name of the parameter to show if invalid characters are returned.</param>
-		/// <returns></returns>
-		public static Utf8String ValidateArgument(Utf8String text, string paramName)
+		/// <param name="id">Parsed identifier</param>
+		/// <returns>True if the identifier was parsed</returns>
+		public static bool TryParse(string text, out StringId id)
+			=> TryParse(new Utf8String(text), out id, out _);
+
+		/// <summary>
+		/// Validates the given string as a StringId, normalizing it if necessary.
+		/// </summary>
+		/// <param name="text">Text to validate as a StringId</param>
+		/// <param name="id">Parsed identifier</param>
+		/// <param name="errorMessage">Error message if the id cannot be parsed</param>
+		/// <returns>True if the identifier was parsed</returns>
+		public static bool TryParse(string text, out StringId id, [NotNullWhen(false)] out string? errorMessage)
+			=> TryParse(new Utf8String(text), out id, out errorMessage);
+
+		/// <summary>
+		/// Validates the given string as a StringId, normalizing it if necessary.
+		/// </summary>
+		/// <param name="text">Text to validate as a StringId</param>
+		/// <param name="id">Parsed identifier</param>
+		/// <returns>True if the identifier was parsed</returns>
+		public static bool TryParse(Utf8String text, out StringId id)
+			=> TryParse(text, out id, out _);
+
+		/// <summary>
+		/// Validates the given string as a StringId, normalizing it if necessary.
+		/// </summary>
+		/// <param name="text">Text to validate as a StringId</param>
+		/// <param name="id">Parsed identifier</param>
+		/// <param name="errorMessage">Error message if the id cannot be parsed</param>
+		/// <returns>True if the identifier was parsed</returns>
+		public static bool TryParse(Utf8String text, out StringId id, [NotNullWhen(false)] out string? errorMessage)
 		{
 			if (text.Length > MaxLength)
 			{
-				throw new ArgumentException($"String id '{text}' may not be longer than {MaxLength} characters", paramName);
+				id = default;
+				errorMessage = $"Id '{text}' is longer than {MaxLength} characters";
+				return false;
 			}
 
 			if (text.Length > 0 && (text[0] == '.' || text[^1] == '.'))
 			{
-				throw new ArgumentException($"'{text}' is not a valid string id (cannot start or end with a period)");
+				id = default;
+				errorMessage = $"'{text}' is not a valid id (cannot start or end with a period)";
+				return false;
 			}
 
 			for (int idx = 0; idx < text.Length; idx++)
@@ -158,12 +196,17 @@ namespace EpicGames.Horde
 					}
 					else
 					{
-						throw new ArgumentException($"'{text}' is not a valid string id (character '{(char)character}' is not allowed)", paramName);
+						id = default;
+						errorMessage = $"'{text}' is not a valid id (character '{(char)character}' is not allowed)";
+						return false;
+
 					}
 				}
 			}
 
-			return text;
+			id = new StringId(text, Validate.None);
+			errorMessage = null;
+			return true;
 		}
 
 		/// <summary>
