Horde: Fix SwapTest() errors due to MemoryStorageClient capturing a reference to the supplied buffer for written blobs rather than copying it.

#rnx

diff --git a/Engine/Source/Programs/Shared/EpicGames.Horde.Tests/Storage/WorkspaceTests.cs b/Engine/Source/Programs/Shared/EpicGames.Horde.Tests/Storage/WorkspaceTests.cs
--- a/Engine/Source/Programs/Shared/EpicGames.Horde.Tests/Storage/WorkspaceTests.cs
+++ b/Engine/Source/Programs/Shared/EpicGames.Horde.Tests/Storage/WorkspaceTests.cs
@@ -224,7 +224,6 @@ namespace EpicGames.Horde.Tests.Storage
 			await AssertContentsEqualAsync(emptyState, output);
 		}
 
-		[Ignore]
 		[TestMethod]
 		public async Task SwapTestAsync()
 		{
diff --git a/Engine/Source/Programs/Shared/EpicGames.Horde/Storage/IBlobWriter.cs b/Engine/Source/Programs/Shared/EpicGames.Horde/Storage/IBlobWriter.cs
--- a/Engine/Source/Programs/Shared/EpicGames.Horde/Storage/IBlobWriter.cs
+++ b/Engine/Source/Programs/Shared/EpicGames.Horde/Storage/IBlobWriter.cs
@@ -356,7 +356,7 @@ namespace EpicGames.Horde.Storage
 		public override ValueTask<IHashedBlobRef> WriteBlobAsync(BlobType type, int size, IReadOnlyList<IBlobRef> imports, IReadOnlyList<AliasInfo> aliases, CancellationToken cancellationToken)
 		{
 			Memory<byte> memory = _memoryWriter.GetMemoryAndAdvance(size);
-			BlobData data = new BlobData(type, memory, imports.ToArray());
+			BlobData data = new BlobData(type, memory.ToArray(), imports.ToArray());
 			return new ValueTask<IHashedBlobRef>(new BlobRef(++_nextIndex, data));
 		}
 
