Horde: Propagate storage exceptions up to the thread that's invoking a write. The current situation logs an error but allows the system to continue even when data failed to be flushed to the storage system.

#rnx

diff --git a/Engine/Source/Programs/Shared/EpicGames.Horde/Storage/Bundles/V2/BundleWriter.cs b/Engine/Source/Programs/Shared/EpicGames.Horde/Storage/Bundles/V2/BundleWriter.cs
--- a/Engine/Source/Programs/Shared/EpicGames.Horde/Storage/Bundles/V2/BundleWriter.cs
+++ b/Engine/Source/Programs/Shared/EpicGames.Horde/Storage/Bundles/V2/BundleWriter.cs
@@ -70,8 +70,8 @@ namespace EpicGames.Horde.Storage.Bundles.V2
 			RefCountedMemoryWriter? _encodedPacketWriter;
 			readonly HashSet<BlobLocator> _bundleImports = new HashSet<BlobLocator>();
 
+			public Task? _lockTask;
 			public Task? _writeTask;
-			public readonly TaskCompletionSource StartWriteTcs = new TaskCompletionSource();
 
 			public FlushedBundleHandle? FlushedHandle => _flushedHandle;
 
@@ -289,8 +289,8 @@ namespace EpicGames.Horde.Storage.Bundles.V2
 
 		PendingBundleHandle _currentBundle;
 
-		readonly AsyncEvent _taskFinished = new AsyncEvent();
-		int _numTasks;
+		readonly object _lockObject = new object();
+		Task _flushTask = Task.CompletedTask;
 
 		/// <summary>
 		/// Constructor
@@ -343,46 +343,52 @@ namespace EpicGames.Horde.Storage.Bundles.V2
 		{
 			lock (bundle.LockObject)
 			{
-				bundle._writeTask ??= Task.Run(() => HandleBundleWriteAsync(bundle), CancellationToken.None);
+				Task lockTask = bundle._lockTask ??= _bundleCache.WriteSemaphore.WaitAsync(_cancellationToken);
+				bundle._writeTask ??= Task.Run(() => HandleBundleWriteAsync(lockTask, bundle), CancellationToken.None);
 			}
-			await bundle.StartWriteTcs.Task.WaitAsync(cancellationToken);
+
+			try
+			{
+				await bundle._lockTask;
+			}
+			catch
+			{
+				await bundle._writeTask; // Allow this exception to be observed as well, since we're discarding it
+				throw;
+			}
+
+			lock (_lockObject)
+			{
+				_flushTask = Task.WhenAll(_flushTask, bundle._writeTask);
+			}
+
 			return bundle._writeTask;
 		}
 
-		async Task HandleBundleWriteAsync(PendingBundleHandle bundle)
+		async Task HandleBundleWriteAsync(Task lockTask, PendingBundleHandle bundle)
 		{
-			Interlocked.Increment(ref _numTasks);
+			// Wait until we've got the semaphore
+			await lockTask;
+
+			// Execute the write
 			try
 			{
-				await _bundleCache.WriteSemaphore.WaitAsync(_cancellationToken);
-				bundle.StartWriteTcs.SetResult();
-
-				try
-				{
-					await bundle.WriteAsync(_cancellationToken);
-				}
-				catch (OperationCanceledException)
-				{
-					_logger.LogDebug("Write of {Locator} was cancelled", bundle.Locator);
-				}
-				catch (Exception ex)
-				{
-					_logger.LogError(ex, "Error writing {Locator}: {Message}", bundle.Locator, ex.Message);
-				}
-				finally
-				{
-					bundle.Dispose();
-					_bundleCache.WriteSemaphore.Release();
-				}
+				await bundle.WriteAsync(_cancellationToken);
+			}
+			catch (OperationCanceledException)
+			{
+				_logger.LogDebug("Write of {Locator} was cancelled", bundle.Locator);
+				throw;
 			}
 			catch (Exception ex)
 			{
-				bundle.StartWriteTcs.SetException(ex);
+				_logger.LogError(ex, "Error writing {Locator}: {Message}", bundle.Locator, ex.Message);
+				throw;
 			}
 			finally
 			{
-				Interlocked.Decrement(ref _numTasks);
-				_taskFinished.Set();
+				bundle.Dispose();
+				_bundleCache.WriteSemaphore.Release();
 			}
 		}
 
@@ -391,15 +397,13 @@ namespace EpicGames.Horde.Storage.Bundles.V2
 		{
 			await WriteCurrentBundleAsync(cancellationToken);
 
-			for (; ; )
+			Task flushTask;
+			lock (_lockObject)
 			{
-				Task taskFinished = _taskFinished.Task;
-				if (Interlocked.CompareExchange(ref _numTasks, 0, 0) == 0)
-				{
-					break;
-				}
-				await taskFinished.WaitAsync(cancellationToken);
+				flushTask = _flushTask;
 			}
+
+			await flushTask.WaitAsync(cancellationToken);
 		}
 
 		/// <inheritdoc/>
