Horde: Add default CancellationToken parameters to directory manipulation methods.

#rnx

diff --git a/Engine/Source/Programs/AutomationTool/BuildGraph/Tasks/CreateArtifactTask.cs b/Engine/Source/Programs/AutomationTool/BuildGraph/Tasks/CreateArtifactTask.cs
--- a/Engine/Source/Programs/AutomationTool/BuildGraph/Tasks/CreateArtifactTask.cs
+++ b/Engine/Source/Programs/AutomationTool/BuildGraph/Tasks/CreateArtifactTask.cs
@@ -235,7 +235,7 @@ namespace AutomationTool.Tasks
 						Logger.LogInformation(ex, "Unable to read previous artifact: {Message}", ex.Message);
 					}
 
-					rootRef = await dedupeWriter.WriteFilesAsync(baseDir, files);
+					rootRef = await dedupeWriter.AddFilesAsync(baseDir, files);
 
 					StorageStats stats = dedupeWriter.GetStats();
 					Logger.LogInformation("Stats:");
@@ -243,7 +243,7 @@ namespace AutomationTool.Tasks
 				}
 				else
 				{
-					rootRef = await blobWriter.WriteFilesAsync(baseDir, files);
+					rootRef = await blobWriter.AddFilesAsync(baseDir, files);
 				}
 			}
 			double uploadTimeSecs = timer.Elapsed.TotalSeconds;
diff --git a/Engine/Source/Programs/Horde/HordeCmd/Commands/Archives/ArchiveCreate.cs b/Engine/Source/Programs/Horde/HordeCmd/Commands/Archives/ArchiveCreate.cs
--- a/Engine/Source/Programs/Horde/HordeCmd/Commands/Archives/ArchiveCreate.cs
+++ b/Engine/Source/Programs/Horde/HordeCmd/Commands/Archives/ArchiveCreate.cs
@@ -97,7 +97,7 @@ namespace Horde.Commands.Archives
 				List<FileInfo> fileInfos = files.ConvertAll(x => x.ToFileInfo());
 				UpdateStatsLogger updateStatsLogger = new UpdateStatsLogger(files.Count, fileInfos.Sum(x => x.Length), logger);
 
-				IHashedBlobRef<DirectoryNode> nodeRef = await writer.WriteFilesAsync(baseDir.ToDirectoryInfo(), fileInfos, options, updateStatsLogger, CancellationToken.None);
+				IHashedBlobRef<DirectoryNode> nodeRef = await writer.AddFilesAsync(baseDir.ToDirectoryInfo(), fileInfos, options, updateStatsLogger, CancellationToken.None);
 
 				await writer.FlushAsync();
 
diff --git a/Engine/Source/Programs/Horde/HordeCmd/Commands/Compute/ComputeRun.cs b/Engine/Source/Programs/Horde/HordeCmd/Commands/Compute/ComputeRun.cs
--- a/Engine/Source/Programs/Horde/HordeCmd/Commands/Compute/ComputeRun.cs
+++ b/Engine/Source/Programs/Horde/HordeCmd/Commands/Compute/ComputeRun.cs
@@ -161,7 +161,7 @@ namespace Horde.Commands.Compute
 		{
 			await using IBlobWriter writer = storage.CreateBlobWriter(cancellationToken: cancellationToken);
 
-			IHashedBlobRef<DirectoryNode> sandbox = await writer.WriteFilesAsync(taskFile.Directory, cancellationToken: cancellationToken);
+			IHashedBlobRef<DirectoryNode> sandbox = await writer.AddFilesAsync(taskFile.Directory, cancellationToken: cancellationToken);
 			await writer.FlushAsync(cancellationToken);
 
 			return sandbox.GetLocator();
diff --git a/Engine/Source/Programs/Horde/HordeCmd/Commands/Tools/ToolUpload.cs b/Engine/Source/Programs/Horde/HordeCmd/Commands/Tools/ToolUpload.cs
--- a/Engine/Source/Programs/Horde/HordeCmd/Commands/Tools/ToolUpload.cs
+++ b/Engine/Source/Programs/Horde/HordeCmd/Commands/Tools/ToolUpload.cs
@@ -33,7 +33,7 @@ class ToolUpload(IHordeClient hordeClient) : Command
 		IHashedBlobRef<DirectoryNode> target;
 		await using (IBlobWriter writer = storageNamespace.CreateBlobWriter(ToolId.ToString(), null))
 		{
-			target = await writer.WriteFilesAsync(InputDir);
+			target = await writer.AddFilesAsync(InputDir);
 		}
 		
 		using HordeHttpClient httpClient = hordeClient.CreateHttpClient();
diff --git a/Engine/Source/Programs/Horde/Samples/RemoteClient/RemoteClientApp.cs b/Engine/Source/Programs/Horde/Samples/RemoteClient/RemoteClientApp.cs
--- a/Engine/Source/Programs/Horde/Samples/RemoteClient/RemoteClientApp.cs
+++ b/Engine/Source/Programs/Horde/Samples/RemoteClient/RemoteClientApp.cs
@@ -113,7 +113,7 @@ namespace RemoteClient
 
 				await using (IBlobWriter writer = storage.CreateBlobWriter())
 				{
-					IHashedBlobRef<DirectoryNode> sandbox = await writer.WriteFilesAsync(uploadDir);
+					IHashedBlobRef<DirectoryNode> sandbox = await writer.AddFilesAsync(uploadDir);
 					await writer.FlushAsync();
 					await channel.UploadFilesAsync("", sandbox.GetLocator(), storage.Backend);
 				}
diff --git a/Engine/Source/Programs/Shared/EpicGames.Horde.Tests/Storage/WorkspaceTests.cs b/Engine/Source/Programs/Shared/EpicGames.Horde.Tests/Storage/WorkspaceTests.cs
--- a/Engine/Source/Programs/Shared/EpicGames.Horde.Tests/Storage/WorkspaceTests.cs
+++ b/Engine/Source/Programs/Shared/EpicGames.Horde.Tests/Storage/WorkspaceTests.cs
@@ -74,7 +74,7 @@ namespace EpicGames.Horde.Tests.Storage
 			DirectoryInfo tempDir = CreateTempDirectory();
 			await input.ExtractAsync(tempDir, new ExtractOptions(), NullLogger.Instance, default);
 
-			IHashedBlobRef<DirectoryNode> output = await blobWriter.WriteFilesAsync(new DirectoryReference(tempDir), null, null, default);
+			IHashedBlobRef<DirectoryNode> output = await blobWriter.AddFilesAsync(new DirectoryReference(tempDir), null, null, default);
 			await AssertContentsEqualAsync(input, output);
 		}
 
@@ -101,7 +101,7 @@ namespace EpicGames.Horde.Tests.Storage
 			Assert.AreEqual(0, updateStats.NumBytesReused);
 			Assert.AreEqual(0, updateStats.NumBytesDownloaded);
 
-			IHashedBlobRef<DirectoryNode> output = await blobWriter.WriteFilesAsync(new DirectoryReference(tempDir), null, null, default);
+			IHashedBlobRef<DirectoryNode> output = await blobWriter.AddFilesAsync(new DirectoryReference(tempDir), null, null, default);
 			await AssertContentsEqualAsync(initialState, output);
 		}
 
@@ -128,7 +128,7 @@ namespace EpicGames.Horde.Tests.Storage
 			}
 			await ws.SyncAsync(WorkspaceLayerId.Default, updatedState, null);
 
-			IHashedBlobRef<DirectoryNode> output = await blobWriter.WriteFilesAsync(new DirectoryReference(tempDir), null, null, default);
+			IHashedBlobRef<DirectoryNode> output = await blobWriter.AddFilesAsync(new DirectoryReference(tempDir), null, null, default);
 			await AssertContentsEqualAsync(updatedState, output);
 		}
 
@@ -156,7 +156,7 @@ namespace EpicGames.Horde.Tests.Storage
 			}
 			await ws.SyncAsync(WorkspaceLayerId.Default, updatedState, null);
 
-			IHashedBlobRef<DirectoryNode> output = await blobWriter.WriteFilesAsync(new DirectoryReference(tempDir), null, null, default);
+			IHashedBlobRef<DirectoryNode> output = await blobWriter.AddFilesAsync(new DirectoryReference(tempDir), null, null, default);
 			await AssertContentsEqualAsync(updatedState, output);
 		}
 
@@ -210,7 +210,7 @@ namespace EpicGames.Horde.Tests.Storage
 			Assert.AreEqual(stats.NumBytesDownloaded, 0);
 
 			// Check it's still the same
-			IHashedBlobRef<DirectoryNode> output = await blobWriter.WriteFilesAsync(new DirectoryReference(tempDir), null, null, default);
+			IHashedBlobRef<DirectoryNode> output = await blobWriter.AddFilesAsync(new DirectoryReference(tempDir), null, null, default);
 			await AssertContentsEqualAsync(initialState, output);
 
 			// Remove layer A
@@ -220,7 +220,7 @@ namespace EpicGames.Horde.Tests.Storage
 			Assert.AreEqual(stats.NumBytesDownloaded, 0);
 
 			// Check it's empty
-			output = await blobWriter.WriteFilesAsync(new DirectoryReference(tempDir), null, null, default);
+			output = await blobWriter.AddFilesAsync(new DirectoryReference(tempDir), null, null, default);
 			await AssertContentsEqualAsync(emptyState, output);
 		}
 
@@ -261,7 +261,7 @@ namespace EpicGames.Horde.Tests.Storage
 			WorkspaceSyncStats syncStats = new WorkspaceSyncStats();
 			await ws.SyncAsync(WorkspaceLayerId.Default, updatedState, stats: syncStats);
 
-			IHashedBlobRef<DirectoryNode> output = await blobWriter.WriteFilesAsync(new DirectoryReference(tempDir), null, null, default);
+			IHashedBlobRef<DirectoryNode> output = await blobWriter.AddFilesAsync(new DirectoryReference(tempDir), null, null, default);
 			await AssertContentsEqualAsync(updatedState, output);
 
 			Assert.AreEqual(3, syncStats.NumFilesMoved);
diff --git a/Engine/Source/Programs/Shared/EpicGames.Horde/Storage/Nodes/DirectoryNode.Extract.cs b/Engine/Source/Programs/Shared/EpicGames.Horde/Storage/Nodes/DirectoryNode.Extract.cs
--- a/Engine/Source/Programs/Shared/EpicGames.Horde/Storage/Nodes/DirectoryNode.Extract.cs
+++ b/Engine/Source/Programs/Shared/EpicGames.Horde/Storage/Nodes/DirectoryNode.Extract.cs
@@ -178,7 +178,7 @@ namespace EpicGames.Horde.Storage.Nodes
 		/// <param name="logger"></param>
 		/// <param name="cancellationToken"></param>
 		[Obsolete("Use an ExtractAsync overload that takes an ExtractOptions object")]
-		public static Task ExtractAsync(this IBlobRef<DirectoryNode> directoryRef, DirectoryInfo directoryInfo, ILogger logger, CancellationToken cancellationToken)
+		public static Task ExtractAsync(this IBlobRef<DirectoryNode> directoryRef, DirectoryInfo directoryInfo, ILogger logger, CancellationToken cancellationToken = default)
 			=> ExtractAsync(directoryRef, directoryInfo, new ExtractOptions(), logger, cancellationToken);
 
 		/// <summary>
@@ -190,7 +190,7 @@ namespace EpicGames.Horde.Storage.Nodes
 		/// <param name="logger">Logger for output</param>
 		/// <param name="cancellationToken">Cancellation token for the operation</param>
 		[Obsolete("Use an ExtractAsync overload that takes an ExtractOptions object")]
-		public static Task ExtractAsync(this IBlobRef<DirectoryNode> directoryRef, DirectoryInfo directoryInfo, IProgress<IExtractStats>? progress, ILogger logger, CancellationToken cancellationToken)
+		public static Task ExtractAsync(this IBlobRef<DirectoryNode> directoryRef, DirectoryInfo directoryInfo, IProgress<IExtractStats>? progress, ILogger logger, CancellationToken cancellationToken = default)
 		{
 			return ExtractAsync(directoryRef, directoryInfo, new ExtractOptions { Progress = progress }, logger, cancellationToken);
 		}
@@ -205,7 +205,7 @@ namespace EpicGames.Horde.Storage.Nodes
 		/// <param name="logger">Logger for output</param>
 		/// <param name="cancellationToken">Cancellation token for the operation</param>
 		[Obsolete("Use an ExtractAsync overload that takes an ExtractOptions object")]
-		public static async Task ExtractAsync(this IBlobRef<DirectoryNode> directoryRef, DirectoryInfo directoryInfo, IProgress<IExtractStats>? progress, TimeSpan frequency, ILogger logger, CancellationToken cancellationToken)
+		public static async Task ExtractAsync(this IBlobRef<DirectoryNode> directoryRef, DirectoryInfo directoryInfo, IProgress<IExtractStats>? progress, TimeSpan frequency, ILogger logger, CancellationToken cancellationToken = default)
 		{
 			DirectoryNode directoryNode = await directoryRef.ReadBlobAsync(cancellationToken);
 			await ExtractAsync(directoryNode, directoryInfo, new ExtractOptions { Progress = progress, ProgressUpdateFrequency = frequency }, logger, cancellationToken);
@@ -218,8 +218,7 @@ namespace EpicGames.Horde.Storage.Nodes
 		/// <param name="directoryInfo">Direcotry to write to</param>
 		/// <param name="logger">Logger for output</param>
 		/// <param name="cancellationToken">Cancellation token for the operation</param>
-		[Obsolete("Use an ExtractAsync overload that takes an ExtractOptions object")]
-		public static Task ExtractAsync(this DirectoryNode directoryNode, DirectoryInfo directoryInfo, ILogger logger, CancellationToken cancellationToken)
+		public static Task ExtractAsync(this DirectoryNode directoryNode, DirectoryInfo directoryInfo, ILogger logger, CancellationToken cancellationToken = default)
 			=> ExtractAsync(directoryNode, directoryInfo, new ExtractOptions(), logger, cancellationToken);
 
 		/// <summary>
@@ -230,7 +229,7 @@ namespace EpicGames.Horde.Storage.Nodes
 		/// <param name="options">Options for the download</param>
 		/// <param name="logger">Logger for output</param>
 		/// <param name="cancellationToken">Cancellation token for the operation</param>
-		public static async Task ExtractAsync(this IBlobRef<DirectoryNode> directoryRef, DirectoryInfo directoryInfo, ExtractOptions? options, ILogger logger, CancellationToken cancellationToken)
+		public static async Task ExtractAsync(this IBlobRef<DirectoryNode> directoryRef, DirectoryInfo directoryInfo, ExtractOptions? options, ILogger logger, CancellationToken cancellationToken = default)
 		{
 			DirectoryNode directoryNode = await directoryRef.ReadBlobAsync(cancellationToken);
 			await ExtractAsync(directoryNode, directoryInfo, options, logger, cancellationToken);
@@ -244,7 +243,7 @@ namespace EpicGames.Horde.Storage.Nodes
 		/// <param name="options">Options for the download</param>
 		/// <param name="logger">Logger for output</param>
 		/// <param name="cancellationToken">Cancellation token for the operation</param>
-		public static async Task ExtractAsync(this DirectoryNode directoryNode, DirectoryInfo directoryInfo, ExtractOptions? options, ILogger logger, CancellationToken cancellationToken)
+		public static async Task ExtractAsync(this DirectoryNode directoryNode, DirectoryInfo directoryInfo, ExtractOptions? options, ILogger logger, CancellationToken cancellationToken = default)
 		{
 			options ??= new ExtractOptions();
 
diff --git a/Engine/Source/Programs/Shared/EpicGames.Horde/Storage/Nodes/DirectoryNode.Update.cs b/Engine/Source/Programs/Shared/EpicGames.Horde/Storage/Nodes/DirectoryNode.Update.cs
--- a/Engine/Source/Programs/Shared/EpicGames.Horde/Storage/Nodes/DirectoryNode.Update.cs
+++ b/Engine/Source/Programs/Shared/EpicGames.Horde/Storage/Nodes/DirectoryNode.Update.cs
@@ -351,7 +351,7 @@ namespace EpicGames.Horde.Storage.Nodes
 		/// <summary>
 		/// Writes a tree of files to a storage writer
 		/// </summary>
-		public static async Task<IHashedBlobRef<DirectoryNode>> WriteFilesAsync(this IBlobWriter writer, DirectoryReference baseDir, ChunkingOptions? options = null, IProgress<IUpdateStats>? progress = null, CancellationToken cancellationToken = default)
+		public static async Task<IHashedBlobRef<DirectoryNode>> AddFilesAsync(this IBlobWriter writer, DirectoryReference baseDir, ChunkingOptions? options = null, IProgress<IUpdateStats>? progress = null, CancellationToken cancellationToken = default)
 		{
 			DirectoryNode outputNode = new DirectoryNode();
 			await outputNode.AddFilesAsync(baseDir, DirectoryReference.EnumerateFiles(baseDir, "*", SearchOption.AllDirectories), writer, options, progress, cancellationToken);
@@ -361,7 +361,7 @@ namespace EpicGames.Horde.Storage.Nodes
 		/// <summary>
 		/// Writes a tree of files to a storage writer
 		/// </summary>
-		public static async Task<IHashedBlobRef<DirectoryNode>> WriteFilesAsync(this IBlobWriter writer, DirectoryInfo baseDir, IReadOnlyList<FileInfo> files, ChunkingOptions? options = null, IProgress<IUpdateStats>? progress = null, CancellationToken cancellationToken = default)
+		public static async Task<IHashedBlobRef<DirectoryNode>> AddFilesAsync(this IBlobWriter writer, DirectoryInfo baseDir, IReadOnlyList<FileInfo> files, ChunkingOptions? options = null, IProgress<IUpdateStats>? progress = null, CancellationToken cancellationToken = default)
 		{
 			DirectoryNode outputNode = new DirectoryNode();
 			await outputNode.AddFilesAsync(baseDir, files, writer, options, progress, cancellationToken);
@@ -371,7 +371,7 @@ namespace EpicGames.Horde.Storage.Nodes
 		/// <summary>
 		/// Writes a tree of files to a storage writer
 		/// </summary>
-		public static async Task<IHashedBlobRef<DirectoryNode>> WriteFilesAsync(this IBlobWriter writer, DirectoryReference baseDir, IReadOnlyList<FileReference> files, ChunkingOptions? options = null, IProgress<IUpdateStats>? progress = null, CancellationToken cancellationToken = default)
+		public static async Task<IHashedBlobRef<DirectoryNode>> AddFilesAsync(this IBlobWriter writer, DirectoryReference baseDir, IReadOnlyList<FileReference> files, ChunkingOptions? options = null, IProgress<IUpdateStats>? progress = null, CancellationToken cancellationToken = default)
 		{
 			DirectoryNode outputNode = new DirectoryNode();
 			await outputNode.AddFilesAsync(baseDir, files, writer, options, progress, cancellationToken);
