Horde: Fix unused parameter warning.

#rnx

diff --git a/Engine/Source/Programs/Shared/EpicGames.Horde/Storage/Bundles/V2/BundleWriter.cs b/Engine/Source/Programs/Shared/EpicGames.Horde/Storage/Bundles/V2/BundleWriter.cs
--- a/Engine/Source/Programs/Shared/EpicGames.Horde/Storage/Bundles/V2/BundleWriter.cs
+++ b/Engine/Source/Programs/Shared/EpicGames.Horde/Storage/Bundles/V2/BundleWriter.cs
@@ -349,7 +349,7 @@ namespace EpicGames.Horde.Storage.Bundles.V2
 
 			try
 			{
-				await bundle._lockTask;
+				await bundle._lockTask.WaitAsync(cancellationToken);
 			}
 			catch
 			{
