UBT: GenerateProjects - Optimize CommandLineArguments.ApplyTo by caching GetArgumentTargetsForType - 69s to 61s.
CommandLineArguments.ApplyTo thread time on TP workers reduced from 76s to 73ms!
The CoreCLR CMDSemReadWrite::LockRead acquired when getting custom attributes time on TP workers reduced from 78s to 13s!

#rnx

diff --git a/Engine/Source/Programs/Shared/EpicGames.Core/CommandLineArguments.cs b/Engine/Source/Programs/Shared/EpicGames.Core/CommandLineArguments.cs
--- a/Engine/Source/Programs/Shared/EpicGames.Core/CommandLineArguments.cs
+++ b/Engine/Source/Programs/Shared/EpicGames.Core/CommandLineArguments.cs
@@ -858,28 +858,57 @@ namespace EpicGames.Core
 			return obj;
 		}
 
+		private static Dictionary<Type, List<ArgumentTarget>> CachedArgumentTargetsForType = new Dictionary<Type, List<ArgumentTarget>>();
 		static List<ArgumentTarget> GetArgumentTargetsForType(Type? targetType)
 		{
 			List<ArgumentTarget> targets = [];
-			while (targetType != null && targetType != typeof(object))
+			if (targetType == null || targetType == typeof(object))
 			{
-				foreach (FieldInfo fieldInfo in targetType.GetFields(BindingFlags.Instance | BindingFlags.GetField | BindingFlags.Public | BindingFlags.NonPublic | BindingFlags.DeclaredOnly))
+				return targets;
+			}
+
+			lock (CachedArgumentTargetsForType)
+			{
+				List<ArgumentTarget>? cachedTargets;
+				if (CachedArgumentTargetsForType.TryGetValue(targetType, out cachedTargets)) 
 				{
-					IEnumerable<CommandLineAttribute> attributes = fieldInfo.GetCustomAttributes<CommandLineAttribute>();
-					if (attributes.Any())
-					{
-						targets.Add(new ArgumentTarget(fieldInfo, fieldInfo.FieldType, fieldInfo.SetValue, fieldInfo.GetValue, attributes.ToArray()));
-					}
+					return cachedTargets;
 				}
-				foreach (PropertyInfo propertyInfo in targetType.GetProperties(BindingFlags.Instance | BindingFlags.GetProperty | BindingFlags.Public | BindingFlags.NonPublic | BindingFlags.DeclaredOnly))
+			}
+
+			foreach (FieldInfo fieldInfo in targetType.GetFields(BindingFlags.Instance | BindingFlags.GetField | BindingFlags.Public | BindingFlags.NonPublic | BindingFlags.DeclaredOnly))
+			{
+				IEnumerable<CommandLineAttribute> attributes = fieldInfo.GetCustomAttributes<CommandLineAttribute>();
+				if (attributes.Any())
 				{
-					IEnumerable<CommandLineAttribute> attributes = propertyInfo.GetCustomAttributes<CommandLineAttribute>();
-					if (attributes.Any())
-					{
-						targets.Add(new ArgumentTarget(propertyInfo, propertyInfo.PropertyType, propertyInfo.SetValue, propertyInfo.GetValue, attributes.ToArray()));
-					}
+					targets.Add(new ArgumentTarget(fieldInfo, fieldInfo.FieldType, fieldInfo.SetValue, fieldInfo.GetValue, attributes.ToArray()));
 				}
-				targetType = targetType.BaseType;
+			}
+			foreach (PropertyInfo propertyInfo in targetType.GetProperties(BindingFlags.Instance | BindingFlags.GetProperty | BindingFlags.Public | BindingFlags.NonPublic | BindingFlags.DeclaredOnly))
+			{
+				IEnumerable<CommandLineAttribute> attributes = propertyInfo.GetCustomAttributes<CommandLineAttribute>();
+				if (attributes.Any())
+				{
+					targets.Add(new ArgumentTarget(propertyInfo, propertyInfo.PropertyType, propertyInfo.SetValue, propertyInfo.GetValue, attributes.ToArray()));
+				}
+			}
+
+			List<ArgumentTarget> baseTargets = GetArgumentTargetsForType(targetType.BaseType);
+			if (baseTargets.Count > 0)
+			{
+				if (targets.Count > 0)
+				{
+					targets.AddRange(baseTargets);
+				}
+				else
+				{
+					targets = baseTargets;
+				}
+			}
+
+			lock (CachedArgumentTargetsForType)
+			{
+				CachedArgumentTargetsForType[targetType] = targets;
 			}
 			return targets;
 		}
