Horde: Missing file.

#rnx

diff --git a/Engine/Source/Programs/Shared/EpicGames.Horde/IHordeClient.cs b/Engine/Source/Programs/Shared/EpicGames.Horde/IHordeClient.cs
--- a/Engine/Source/Programs/Shared/EpicGames.Horde/IHordeClient.cs
+++ b/Engine/Source/Programs/Shared/EpicGames.Horde/IHordeClient.cs
@@ -94,6 +94,13 @@ namespace EpicGames.Horde
 		IServerLogger CreateServerLogger(LogId logId, LogLevel minimumLevel = LogLevel.Information);
 	}
 
+	/// <summary>
+	/// Interface for a Horde client with a known lifetime
+	/// </summary>
+	public interface IHordeClientWithLifetime : IHordeClient, IAsyncDisposable
+	{
+	}
+
 	/// <summary>
 	/// Extension methods for <see cref="IHordeClient"/>
 	/// </summary>
