EpicGames.Core: Fix static analysis warning about naming convention.

#rnx

diff --git a/Engine/Source/Programs/Shared/EpicGames.Core/CommandLineArguments.cs b/Engine/Source/Programs/Shared/EpicGames.Core/CommandLineArguments.cs
--- a/Engine/Source/Programs/Shared/EpicGames.Core/CommandLineArguments.cs
+++ b/Engine/Source/Programs/Shared/EpicGames.Core/CommandLineArguments.cs
@@ -858,7 +858,7 @@ namespace EpicGames.Core
 			return obj;
 		}
 
-		private static Dictionary<Type, List<ArgumentTarget>> CachedArgumentTargetsForType = new Dictionary<Type, List<ArgumentTarget>>();
+		private static readonly Dictionary<Type, List<ArgumentTarget>> s_cachedArgumentTargetsForType = new Dictionary<Type, List<ArgumentTarget>>();
 		static List<ArgumentTarget> GetArgumentTargetsForType(Type? targetType)
 		{
 			List<ArgumentTarget> targets = [];
@@ -867,10 +867,10 @@ namespace EpicGames.Core
 				return targets;
 			}
 
-			lock (CachedArgumentTargetsForType)
+			lock (s_cachedArgumentTargetsForType)
 			{
 				List<ArgumentTarget>? cachedTargets;
-				if (CachedArgumentTargetsForType.TryGetValue(targetType, out cachedTargets)) 
+				if (s_cachedArgumentTargetsForType.TryGetValue(targetType, out cachedTargets)) 
 				{
 					return cachedTargets;
 				}
@@ -906,9 +906,9 @@ namespace EpicGames.Core
 				}
 			}
 
-			lock (CachedArgumentTargetsForType)
+			lock (s_cachedArgumentTargetsForType)
 			{
-				CachedArgumentTargetsForType[targetType] = targets;
+				s_cachedArgumentTargetsForType[targetType] = targets;
 			}
 			return targets;
 		}
