Added xml docs for CompressedBuffer.cs and made some types private that should not be public. Fixes build error.

diff --git a/Engine/Source/Programs/Shared/EpicGames.Serialization/CompressedBuffer.cs b/Engine/Source/Programs/Shared/EpicGames.Serialization/CompressedBuffer.cs
--- a/Engine/Source/Programs/Shared/EpicGames.Serialization/CompressedBuffer.cs
+++ b/Engine/Source/Programs/Shared/EpicGames.Serialization/CompressedBuffer.cs
@@ -15,7 +15,7 @@ using EpicGames.Core;
 
 namespace EpicGames.Serialization
 {
-	public class CompressedBufferHeader
+	class CompressedBufferHeader
 	{
 		[System.Diagnostics.CodeAnalysis.SuppressMessage("Design", "CA1028:Enum Storage should be Int32", Justification = "Interop requires byte")]
 		public enum CompressionMethod : byte
@@ -67,6 +67,9 @@ namespace EpicGames.Serialization
 		}
 	}
 
+	/// <summary>
+	/// Utilities for working with compressed buffers
+	/// </summary>
 	public static class CompressedBuffer
 	{
 		private static (CompressedBufferHeader, uint[]) ExtractHeader(BinaryReader br)
@@ -172,6 +175,15 @@ namespace EpicGames.Serialization
 			}
 		}
 
+		/// <summary>
+		/// Decompress a compressed buffer into a stream
+		/// </summary>
+		/// <param name="sourceStream">Source stream containing a compressed buffer</param>
+		/// <param name="streamSize">The length of the compressed buffer</param>
+		/// <param name="targetStream">The stream to write the decompressed content to</param>
+		/// <param name="cancellationToken">Cancellation token</param>
+		/// <returns></returns>
+		/// <exception cref="Exception"></exception>
 		public static async Task DecompressContentAsync(Stream sourceStream, ulong streamSize, Stream targetStream, CancellationToken cancellationToken = default)
 		{
 			using BinaryReader br = new BinaryReader(sourceStream);
@@ -249,6 +261,14 @@ namespace EpicGames.Serialization
 			}
 		}
 
+		/// <summary>
+		/// Compress a buffer into a stream
+		/// </summary>
+		/// <param name="s">The stream to write the compressed content into</param>
+		/// <param name="method">Which oodle compressor method to use</param>
+		/// <param name="compressionLevel">The oodle compression level</param>
+		/// <param name="rawContents">The uncompressed content you wish to compress</param>
+		/// <returns>The iohash of the uncompressed content</returns>
 		public static IoHash CompressContent(Stream s, OodleCompressorType method, OodleCompressionLevel compressionLevel, byte[] rawContents)
 		{
 			const long DefaultBlockSize = 256 * 1024;
@@ -373,9 +393,19 @@ namespace EpicGames.Serialization
 		}
 	}
 
+	/// <summary>
+	/// Thrown when the crc32 of the header does not align with the decompressed content, usually indicating that the full content was not present.
+	/// </summary>
+	/// <param name="headerCrc32"></param>
+	/// <param name="calculatedCrc"></param>
 	public class InvalidHashException(uint headerCrc32, uint calculatedCrc)
 		: Exception($"Header specified crc \"{headerCrc32}\" but calculated hash was \"{calculatedCrc}\"");
 
+	/// <summary>
+	/// Thrown when a byte buffer is encountered that is not a compressed buffer
+	/// </summary>
+	/// <param name="headerMagic"></param>
+	/// <param name="expectedMagic"></param>
 	public class InvalidMagicException(uint headerMagic, uint expectedMagic)
 		: Exception($"Header magic \"{headerMagic}\" was incorrect, expected to be {expectedMagic}");
 }
