EpicGames.Perforce: Upgrade P4D binary to 2024.2

#skipci

diff --git a/Engine/Source/Programs/Shared/EpicGames.Perforce.Fixture/Dockerfile b/Engine/Source/Programs/Shared/EpicGames.Perforce.Fixture/Dockerfile
--- a/Engine/Source/Programs/Shared/EpicGames.Perforce.Fixture/Dockerfile
+++ b/Engine/Source/Programs/Shared/EpicGames.Perforce.Fixture/Dockerfile
@@ -5,7 +5,7 @@ WORKDIR /tmp
 RUN apt-get update
 RUN apt-get -qq install -y curl netcat dos2unix
 
-RUN curl https://www.perforce.com/downloads/perforce/r22.1/bin.linux26x86_64/helix-core-server.tgz --location --silent --output helix-core-server.tgz
+RUN curl https://cdist2.perforce.com/perforce/r24.2/bin.linux26x86_64/helix-core-server.tgz --location --silent --output helix-core-server.tgz
 RUN tar zxf helix-core-server.tgz && \
     mv ./p4d /usr/local/bin/p4d && \
     mv ./p4 /usr/local/bin/p4 && \
