Horde: Improve logging for job step log writer errors

diff --git a/Engine/Source/Programs/Shared/EpicGames.Horde/Logs/ServerLogger.cs b/Engine/Source/Programs/Shared/EpicGames.Horde/Logs/ServerLogger.cs
--- a/Engine/Source/Programs/Shared/EpicGames.Horde/Logs/ServerLogger.cs
+++ b/Engine/Source/Programs/Shared/EpicGames.Horde/Logs/ServerLogger.cs
@@ -31,6 +31,7 @@ namespace EpicGames.Horde.Logs
 		readonly IBlobWriter _writer;
 
 		int _bufferLength;
+		string? _closedByStackTrace;
 
 		// Tailing task
 		readonly Task _tailTask;
@@ -82,7 +83,7 @@ namespace EpicGames.Horde.Logs
 		{
 			if (!_dataChannel.Writer.TryWrite(jsonLogEvent))
 			{
-				throw new InvalidOperationException("Expected unbounded writer to complete immediately");
+				_internalLogger.LogWarning(KnownLogEvents.Systemic_Horde, "Channel for sending logs to Horde server is closed. Unable to write {LogEvent}. Closed by {Stacktrace}", jsonLogEvent, _closedByStackTrace);
 			}
 		}
 
@@ -94,6 +95,7 @@ namespace EpicGames.Horde.Logs
 		{
 			if (_dataWriter != null)
 			{
+				_closedByStackTrace = Environment.StackTrace;
 				_dataChannel.Writer.TryComplete();
 				await _dataWriter;
 				_dataWriter = null;
