EpicGames.Core: Fix unsafe type cast in LoggerScopeCollection.Scope constructor

diff --git a/Engine/Source/Programs/Shared/EpicGames.Core/Log.cs b/Engine/Source/Programs/Shared/EpicGames.Core/Log.cs
--- a/Engine/Source/Programs/Shared/EpicGames.Core/Log.cs
+++ b/Engine/Source/Programs/Shared/EpicGames.Core/Log.cs
@@ -798,7 +798,11 @@ namespace EpicGames.Core
 			public Scope(LoggerScopeCollection outer, object? value)
 			{
 				_outer = outer;
-				_properties = value as IEnumerable<KeyValuePair<string, object?>>;
+				_properties = value switch
+				{
+					IEnumerable<KeyValuePair<string, object?>> props => props,
+					_ => null
+				};
 
 				_outer._scopes.Add(this);
 			}
