Fixed ToolTips not working on List type parameters in stream configs

diff --git a/Engine/Source/Programs/Shared/EpicGames.Horde/Jobs/Templates/TemplateMessages.cs b/Engine/Source/Programs/Shared/EpicGames.Horde/Jobs/Templates/TemplateMessages.cs
--- a/Engine/Source/Programs/Shared/EpicGames.Horde/Jobs/Templates/TemplateMessages.cs
+++ b/Engine/Source/Programs/Shared/EpicGames.Horde/Jobs/Templates/TemplateMessages.cs
@@ -306,6 +306,7 @@ namespace EpicGames.Horde.Jobs.Templates
 			Label = parameter.Label;
 			Style = parameter.Style;
 			Items = parameter.Items.ConvertAll(x => new GetTemplateListParameterItemResponse(x));
+			ToolTip = parameter.ToolTip;
 		}
 	}
 
