Refactored buffered token reader cache system

#rnx

diff --git a/Engine/Source/Programs/Shared/EpicGames.UHT/Parsers/UhtHeaderFileParser.cs b/Engine/Source/Programs/Shared/EpicGames.UHT/Parsers/UhtHeaderFileParser.cs
--- a/Engine/Source/Programs/Shared/EpicGames.UHT/Parsers/UhtHeaderFileParser.cs
+++ b/Engine/Source/Programs/Shared/EpicGames.UHT/Parsers/UhtHeaderFileParser.cs
@@ -1133,7 +1133,8 @@ namespace EpicGames.UHT.Parsers
 
 		private bool CheckForConstructor(UhtClass classObj, UhtDeclaration declaration)
 		{
-			IUhtTokenReader replayReader = UhtTokenReplayReader.GetThreadInstance(HeaderFile, HeaderFile.Data.Memory, new ReadOnlyMemory<UhtToken>(declaration.Tokens), UhtTokenType.EndOfDeclaration);
+			using UhtTokenReplayReaderBorrower borrowedReader = new(HeaderFile, HeaderFile.Data.Memory, declaration.Tokens, UhtTokenType.EndOfDeclaration);
+			IUhtTokenReader replayReader = borrowedReader.Reader;
 
 			// Allow explicit constructors
 			{
@@ -1251,7 +1252,8 @@ namespace EpicGames.UHT.Parsers
 				return false;
 			}
 
-			IUhtTokenReader replayReader = UhtTokenReplayReader.GetThreadInstance(HeaderFile, HeaderFile.Data.Memory, new ReadOnlyMemory<UhtToken>(declaration.Tokens), UhtTokenType.EndOfDeclaration);
+			using UhtTokenReplayReaderBorrower borrowedReader = new(HeaderFile, HeaderFile.Data.Memory, declaration.Tokens, UhtTokenType.EndOfDeclaration);
+			IUhtTokenReader replayReader = borrowedReader.Reader;
 
 			SkipVirtualAndAPI(replayReader);
 
@@ -1265,7 +1267,8 @@ namespace EpicGames.UHT.Parsers
 
 		bool CheckForSerialize(UhtClass classObj, UhtDeclaration declaration)
 		{
-			IUhtTokenReader replayReader = UhtTokenReplayReader.GetThreadInstance(HeaderFile, HeaderFile.Data.Memory, new ReadOnlyMemory<UhtToken>(declaration.Tokens), UhtTokenType.EndOfDeclaration);
+			using UhtTokenReplayReaderBorrower borrowedReader = new(HeaderFile, HeaderFile.Data.Memory, declaration.Tokens, UhtTokenType.EndOfDeclaration);
+			IUhtTokenReader replayReader = borrowedReader.Reader;
 
 			SkipVirtualAndAPI(replayReader);
 
diff --git a/Engine/Source/Programs/Shared/EpicGames.UHT/Parsers/UhtPropertyParser.cs b/Engine/Source/Programs/Shared/EpicGames.UHT/Parsers/UhtPropertyParser.cs
--- a/Engine/Source/Programs/Shared/EpicGames.UHT/Parsers/UhtPropertyParser.cs
+++ b/Engine/Source/Programs/Shared/EpicGames.UHT/Parsers/UhtPropertyParser.cs
@@ -538,7 +538,9 @@ namespace EpicGames.UHT.Parsers
 		public static UhtProperty? ResolveProperty(UhtPropertyResolvePhase resolvePhase, UhtPropertySettings propertySettings, ReadOnlyMemory<char> data, ReadOnlyMemory<UhtToken> typeTokens)
 		{
 			UhtSession session = propertySettings.Outer.Session;
-			IUhtTokenReader replayReader = UhtTokenReplayReader.GetThreadInstance(propertySettings.Outer, data, typeTokens, UhtTokenType.EndOfType);
+
+			using UhtTokenReplayReaderBorrower borrowedReader = new(propertySettings.Outer, propertySettings.Outer.HeaderFile.Data.Memory, typeTokens, UhtTokenType.EndOfType);
+			IUhtTokenReader replayReader = borrowedReader.Reader;
 
 			// Loop through the tokens until we find a known property type or the start of a template argument list
 			for (int index = 0; index < typeTokens.Length; ++index)
@@ -943,7 +945,7 @@ namespace EpicGames.UHT.Parsers
 
 					// If we have reached the end
 					if (!tokenReader.TryOptional(','))
-				{
+					{
 						break;
 					}
 
@@ -957,8 +959,8 @@ namespace EpicGames.UHT.Parsers
 					// Get the next property name
 					//NameToken = tokenReader.GetIdentifier();
 					//CheckForOptionalParts(SpecifierContext.PropertySettings, ref NameToken);
+					}
 				}
-			}
 			else if (_options.HasAnyFlags(UhtPropertyParseOptions.CommaSeparatedName))
 			{
 				tokenReader.Require(',');
@@ -979,12 +981,12 @@ namespace EpicGames.UHT.Parsers
 					if (_currentTypeTokens.Count < 2 || !_currentTypeTokens[^1].IsIdentifier())
 					{
 						throw new UhtException(tokenReader, $"{propertySettings.PropertyCategory.GetHintText()}: Expected name");
-				}
+					}
 					nameToken = _currentTypeTokens[^1];
 					_currentTypeTokens.RemoveAt(_currentTypeTokens.Count - 1);
 				}
 				Finalize(topScope, specifierContext, ref nameToken, new ReadOnlyMemory<UhtToken>(_currentTypeTokens.ToArray()), layoutMacroType, propertyDelegate);
-			}
+				}
 			else if (_options.HasAnyFlags(UhtPropertyParseOptions.NameIncluded))
 			{
 				if (_currentTypeTokens.Count < 2 || !_currentTypeTokens[^1].IsIdentifier())
@@ -1239,22 +1241,22 @@ namespace EpicGames.UHT.Parsers
 		}
 
 		private bool GatherTypeTokens(ref UhtToken token)
-	{
+			{
 			if (_currentTemplateDepth == 0 && token.IsSymbol() && (token.IsValue(',') || token.IsValue('(') || token.IsValue(')') || 
 				token.IsValue(';') || token.IsValue('[') || token.IsValue(':') || token.IsValue('=') || token.IsValue('{')))
-		{
+	{
 				return false;
-			}
+		}
 
 			_currentTypeTokens.Add(token);
 			if (token.IsSymbol('<'))
 		{
 				++_currentTemplateDepth;
-			}
+		}
 			else if (token.IsSymbol('>'))
-				{
+		{
 				if (_currentTemplateDepth == 0)
-					{
+			{
 					throw new UhtTokenException(_currentTokenReader!, token, "',' or ')'");
 				}
 				--_currentTemplateDepth;
diff --git a/Engine/Source/Programs/Shared/EpicGames.UHT/Tokenizer/UhtTokenReplayReader.cs b/Engine/Source/Programs/Shared/EpicGames.UHT/Tokenizer/UhtTokenReplayReader.cs
--- a/Engine/Source/Programs/Shared/EpicGames.UHT/Tokenizer/UhtTokenReplayReader.cs
+++ b/Engine/Source/Programs/Shared/EpicGames.UHT/Tokenizer/UhtTokenReplayReader.cs
@@ -14,8 +14,6 @@ namespace EpicGames.UHT.Tokenizer
 	/// </summary>
 	public sealed class UhtTokenReplayReader : IUhtTokenReader, IUhtMessageLineNumber
 	{
-		private static readonly ThreadLocal<UhtTokenReplayReader> s_tls = new(() => new UhtTokenReplayReader());
-
 		const int MaxSavedStates = 2;
 		private struct SavedState
 		{
@@ -80,22 +78,6 @@ namespace EpicGames.UHT.Tokenizer
 			return this;
 		}
 
-		/// <summary>
-		/// Return the replay reader associated with the current thread.  Only one replay reader is cached per thread.
-		/// </summary>
-		/// <param name="messageSite">The message site used to log errors</param>
-		/// <param name="data">Source data where tokens were originally parsed</param>
-		/// <param name="tokens">Collection of tokens to replay</param>
-		/// <param name="endTokenType">Type of end token marker to return when the end of the token list is reached.  This is used to produce errors in the context of the replay</param>
-		/// <returns>The threaded instance of the replay reader</returns>
-		/// <exception cref="UhtIceException">Thrown if the TLS value can not be retrieved.</exception>
-		public static UhtTokenReplayReader GetThreadInstance(IUhtMessageSite messageSite, ReadOnlyMemory<char> data, ReadOnlyMemory<UhtToken> tokens, UhtTokenType endTokenType)
-		{
-			UhtTokenReplayReader? reader = s_tls.Value ?? throw new UhtIceException("Unable to acquire a UhtTokenReplayReader");
-			reader.Reset(messageSite, data, tokens, endTokenType);
-			return reader;
-		}
-
 		#region IUHTMessageSite Implementation
 		IUhtMessageSession IUhtMessageSite.MessageSession => _messageSite.MessageSession;
 		IUhtMessageSource? IUhtMessageSite.MessageSource => _messageSite.MessageSource;
@@ -337,4 +319,40 @@ namespace EpicGames.UHT.Tokenizer
 			return ref _currentToken;
 		}
 	}
+
+	/// <summary>
+	/// Borrow a thread specific token replay reader
+	/// </summary>
+	public readonly struct UhtTokenReplayReaderBorrower : IDisposable
+	{
+		private static readonly ThreadLocal<Stack<UhtTokenReplayReader>> s_tlsStack = new(() => new Stack<UhtTokenReplayReader>());
+
+		private readonly UhtTokenReplayReader _reader;
+
+		/// <summary>
+		/// The borrowed reader
+		/// </summary>
+		public readonly IUhtTokenReader Reader => _reader;
+
+		/// <summary>
+		/// Borrow a thread specific replay reader
+		/// </summary>
+		/// <param name="messageSite">Location for messages</param>
+		/// <param name="data">Source being parsed</param>
+		/// <param name="tokens">Tokens to be replayed</param>
+		/// <param name="endTokenType">Token type when end of tokens reached</param>
+		/// <exception cref="UhtIceException"></exception>
+		public UhtTokenReplayReaderBorrower(IUhtMessageSite messageSite, ReadOnlyMemory<char> data, ReadOnlyMemory<UhtToken> tokens, UhtTokenType endTokenType)
+		{
+			Stack<UhtTokenReplayReader> stack = s_tlsStack.Value!;
+			_reader = stack.Count > 0 ? _reader = stack.Pop() : new();
+			_reader.Reset(messageSite, data, tokens, endTokenType);
+		}
+
+		/// <inheritdoc/>
+		public readonly void Dispose()
+		{
+			s_tlsStack.Value!.Push(_reader);
+		}
+	}
 }
diff --git a/Engine/Source/Programs/Shared/EpicGames.UHT/Types/UhtFunction.cs b/Engine/Source/Programs/Shared/EpicGames.UHT/Types/UhtFunction.cs
--- a/Engine/Source/Programs/Shared/EpicGames.UHT/Types/UhtFunction.cs
+++ b/Engine/Source/Programs/Shared/EpicGames.UHT/Types/UhtFunction.cs
@@ -436,8 +436,8 @@ namespace EpicGames.UHT.Types
 								{
 									// All tokens MUST be consumed from the reader
 									StringBuilder builder = new();
-									IUhtTokenReader defaultValueReader = UhtTokenReplayReader.GetThreadInstance(property, HeaderFile.Data.Memory, property.DefaultValueTokens.ToArray(), UhtTokenType.EndOfDefault);
-									parsed = property.SanitizeDefaultValue(defaultValueReader, builder) && defaultValueReader.IsEOF;
+									using UhtTokenReplayReaderBorrower borrowedReader = new(property, HeaderFile.Data.Memory, property.DefaultValueTokens.ToArray(), UhtTokenType.EndOfDefault);
+									parsed = property.SanitizeDefaultValue(borrowedReader.Reader, builder) && borrowedReader.Reader.IsEOF;
 									if (parsed)
 									{
 										MetaData.Add(key, builder.ToString());
