Cleaning up the property parser class to avoid using class instance properties

#rnx

diff --git a/Engine/Source/Programs/Shared/EpicGames.UHT/Parsers/UhtFunctionParser.cs b/Engine/Source/Programs/Shared/EpicGames.UHT/Parsers/UhtFunctionParser.cs
--- a/Engine/Source/Programs/Shared/EpicGames.UHT/Parsers/UhtFunctionParser.cs
+++ b/Engine/Source/Programs/Shared/EpicGames.UHT/Parsers/UhtFunctionParser.cs
@@ -204,7 +204,7 @@ namespace EpicGames.UHT.Parsers
 					UhtProperty? returnValueProperty = null;
 					if (hasReturnValue)
 					{
-						topScope.HeaderParser.GetCachedPropertyParser().Parse(topScope, EPropertyFlags.None,
+						UhtPropertyParser.Parse(topScope, EPropertyFlags.None,
 							function.GetPropertyParseOptions(true), UhtPropertyCategory.Return,
 							(UhtParsingScope topScope, UhtProperty property, ref UhtToken nameToken, UhtLayoutMacroType layoutMacroType) =>
 							{
@@ -432,7 +432,7 @@ namespace EpicGames.UHT.Parsers
 						// Get return type.  C++ style functions always have a return value type, even if it's void
 						UhtToken funcNameToken = new();
 						UhtProperty? returnValueProperty = null;
-						topScope.HeaderParser.GetCachedPropertyParser().Parse(topScope, EPropertyFlags.None,
+						UhtPropertyParser.Parse(topScope, EPropertyFlags.None,
 							function.GetPropertyParseOptions(true), UhtPropertyCategory.Return,
 							(UhtParsingScope topScope, UhtProperty property, ref UhtToken nameToken, UhtLayoutMacroType layoutMacroType) =>
 							{
@@ -562,7 +562,7 @@ namespace EpicGames.UHT.Parsers
 
 			topScope.TokenReader.RequireList(')', ',', false, () =>
 			{
-				topScope.HeaderParser.GetCachedPropertyParser().Parse(topScope, disallowFlags, options, propertyCategory,
+				UhtPropertyParser.Parse(topScope, disallowFlags, options, propertyCategory,
 					(UhtParsingScope topScope, UhtProperty property, ref UhtToken nameToken, UhtLayoutMacroType layoutMacroType) =>
 					{
 						property.PropertyFlags |= EPropertyFlags.Parm;
diff --git a/Engine/Source/Programs/Shared/EpicGames.UHT/Parsers/UhtHeaderFileParser.cs b/Engine/Source/Programs/Shared/EpicGames.UHT/Parsers/UhtHeaderFileParser.cs
--- a/Engine/Source/Programs/Shared/EpicGames.UHT/Parsers/UhtHeaderFileParser.cs
+++ b/Engine/Source/Programs/Shared/EpicGames.UHT/Parsers/UhtHeaderFileParser.cs
@@ -314,11 +314,6 @@ namespace EpicGames.UHT.Parsers
 		/// </summary>
 		public bool SpottedAutogeneratedHeaderInclude { get; set; } = false;
 
-		/// <summary>
-		/// For a given header file, we share a common property parser to reduce the number of allocations.
-		/// </summary>
-		public UhtPropertyParser? PropertyParser { get; set; } = null;
-
 		/// <summary>
 		/// If set, the preprocessor is run in a C++ UHT compatibility mode where only a subset
 		/// of #if class of preprocessor statements are allowed.
@@ -435,16 +430,6 @@ namespace EpicGames.UHT.Parsers
 			_topScope = scope.ParentScope;
 		}
 
-		/// <summary>
-		/// Get the cached property parser
-		/// </summary>
-		/// <returns>Property parser</returns>
-		public UhtPropertyParser GetCachedPropertyParser()
-		{
-			PropertyParser ??= new();
-			return PropertyParser;
-		}
-
 		/// <summary>
 		/// Return the current compiler directive
 		/// </summary>
diff --git a/Engine/Source/Programs/Shared/EpicGames.UHT/Parsers/UhtPropertyParser.cs b/Engine/Source/Programs/Shared/EpicGames.UHT/Parsers/UhtPropertyParser.cs
--- a/Engine/Source/Programs/Shared/EpicGames.UHT/Parsers/UhtPropertyParser.cs
+++ b/Engine/Source/Programs/Shared/EpicGames.UHT/Parsers/UhtPropertyParser.cs
@@ -440,22 +440,27 @@ namespace EpicGames.UHT.Parsers
 		}
 	}
 
+	internal record struct UhtPropertyCategoryExtraContext(UhtPropertyCategory Category) : IUhtMessageExtraContext
+	{
+		#region IMessageExtraContext implementation
+		/// <inheritdoc/>
+		public IEnumerable<object?>? MessageExtraContext
+		{
+			get
+			{
+				Stack<object?> extraContext = new(1);
+				extraContext.Push(Category.GetHintText());
+				return extraContext;
+			}
+		}
+		#endregion
+	}
+
 	/// <summary>
 	/// Property parser
 	/// </summary>
-	public class UhtPropertyParser : IUhtMessageExtraContext
+	public static class UhtPropertyParser
 	{
-		private UhtPropertyCategory _category = UhtPropertyCategory.Member;
-		private UhtPropertyParseOptions _options;
-
-		// Scratch pad variables used by actions
-		private IUhtTokenReader? _currentTokenReader = null;
-		private List<UhtToken> _currentTypeTokens = new();
-		private int _currentTemplateDepth = 0;
-
-		// Actions/Delegates
-		private readonly UhtTokensUntilDelegate _gatherTypeTokensDelegate;
-
 		private static readonly Dictionary<StringView, UhtLayoutMacroType> s_layoutMacroTypes = new(new[]
 		{
 			UhtLayoutMacroType.Array.MacroNameAndValue(),
@@ -467,14 +472,6 @@ namespace EpicGames.UHT.Parsers
 			UhtLayoutMacroType.FieldInitialized.MacroNameAndValue(),
 		});
 
-		/// <summary>
-		/// Construct a new property parser
-		/// </summary>
-		public UhtPropertyParser()
-		{
-			_gatherTypeTokensDelegate = GatherTypeTokens;
-		}
-
 		/// <summary>
 		/// Parse the property
 		/// </summary>
@@ -484,7 +481,7 @@ namespace EpicGames.UHT.Parsers
 		/// <param name="category">Property category</param>
 		/// <param name="propertyDelegate">Delegate to be invoked after property has been parsed</param>
 		/// <returns>The property parser</returns>
-		public UhtPropertyParser Parse(UhtParsingScope topScope, EPropertyFlags disallowPropertyFlags, UhtPropertyParseOptions options, UhtPropertyCategory category, UhtPropertyDelegate propertyDelegate)
+		public static void Parse(UhtParsingScope topScope, EPropertyFlags disallowPropertyFlags, UhtPropertyParseOptions options, UhtPropertyCategory category, UhtPropertyDelegate propertyDelegate)
 		{
 			// Initialize the property context
 			using UhtThreadBorrower<UhtPropertySpecifierContext> borrower = new(true);
@@ -501,31 +498,11 @@ namespace EpicGames.UHT.Parsers
 			specifierContext.SeenBlueprintReadOnlySpecifier = false;
 			specifierContext.SeenBlueprintGetterSpecifier = false;
 
-			// Initialize the settings
-			_options = options;
-			_category = category;
-			_currentTokenReader = topScope.TokenReader;
-			_currentTypeTokens = new List<UhtToken>();
-			_currentTemplateDepth = 0;
-
-			using UhtMessageContext tokenContext = new(this);
-			ParseInternal(topScope, specifierContext, propertyDelegate);
-			return this;
+			UhtPropertyCategoryExtraContext extraContext = new(category);
+			using UhtMessageContext tokenContext = new(extraContext);
+			ParseInternal(topScope, options, specifierContext, propertyDelegate);
 		}
 
-		#region IMessageExtraContext implementation
-		/// <inheritdoc/>
-		public IEnumerable<object?>? MessageExtraContext
-		{
-			get
-			{
-				Stack<object?> extraContext = new(1);
-				extraContext.Push(_category.GetHintText());
-				return extraContext;
-			}
-		}
-		#endregion
-
 		/// <summary>
 		/// Resolve the given property.  This method will resolve any immediate property during the parsing phase or 
 		/// resolve any previously parsed property to the final version.
@@ -543,19 +520,20 @@ namespace EpicGames.UHT.Parsers
 			IUhtTokenReader replayReader = borrowedReader.Reader;
 
 			// Loop through the tokens until we find a known property type or the start of a template argument list
-			for (int index = 0; index < typeTokens.Length; ++index)
+			ReadOnlySpan<UhtToken> typeTokensSpan = typeTokens.Span;
+			for (int index = 0; index < typeTokensSpan.Length; ++index)
 			{
-				if (typeTokens.Span[index].IsSymbol())
+				if (typeTokensSpan[index].IsSymbol())
 				{
-					ReadOnlySpan<char> span = typeTokens.Span[index].Value.Span;
+					ReadOnlySpan<char> span = typeTokensSpan[index].Value.Span;
 					if (span.Length == 1 && (span[0] == '<' || span[0] == '>' || span[0] == ','))
 					{
 						break;
 					}
 				}
-				else if (typeTokens.Span[index].IsIdentifier())
+				else if (typeTokensSpan[index].IsIdentifier())
 				{
-					UhtToken copy = typeTokens.Span[index];
+					UhtToken copy = typeTokensSpan[index];
 					session.Config!.RedirectTypeIdentifier(ref copy);
 					if (session.TryGetPropertyType(copy.Value, out UhtPropertyType propertyType))
 					{
@@ -787,7 +765,7 @@ namespace EpicGames.UHT.Parsers
 			return outProperty;
 		}
 
-		private void ParseInternal(UhtParsingScope topScope, UhtPropertySpecifierContext specifierContext, UhtPropertyDelegate propertyDelegate)
+		private static void ParseInternal(UhtParsingScope topScope, UhtPropertyParseOptions options, UhtPropertySpecifierContext specifierContext, UhtPropertyDelegate propertyDelegate)
 		{
 			UhtPropertySettings propertySettings = specifierContext.PropertySettings;
 			IUhtTokenReader tokenReader = topScope.TokenReader;
@@ -839,7 +817,7 @@ namespace EpicGames.UHT.Parsers
 			// Handle MemoryLayout.h macros
 			bool hasWrapperBrackets = false;
 			UhtLayoutMacroType layoutMacroType = UhtLayoutMacroType.None;
-			if (_options.HasAnyFlags(UhtPropertyParseOptions.ParseLayoutMacro))
+			if (options.HasAnyFlags(UhtPropertyParseOptions.ParseLayoutMacro))
 			{
 				ref UhtToken layoutToken = ref tokenReader.PeekToken();
 				if (layoutToken.IsIdentifier())
@@ -871,10 +849,10 @@ namespace EpicGames.UHT.Parsers
 			tokenReader.Optional("mutable");
 
 			// Gather the type tokens and possibly the property name.
-			tokenReader.While(_gatherTypeTokensDelegate);
+			ReadOnlyMemory<UhtToken> typeTokens = GatherTypeTokens(tokenReader);
 
 			// Verify we at least have one type
-			if (_currentTypeTokens.Count < 1)
+			if (typeTokens.Length < 1)
 			{
 				throw new UhtException(tokenReader, $"{propertySettings.PropertyCategory.GetHintText()}: Missing variable type or name");
 			}
@@ -892,7 +870,7 @@ namespace EpicGames.UHT.Parsers
 				tokenReader.LogError($"Specified type modifiers not allowed here '{String.Join(" | ", extraFlags.ToStringList())}'");
 			}
 
-			if (_options.HasAnyFlags(UhtPropertyParseOptions.AddModuleRelativePath))
+			if (options.HasAnyFlags(UhtPropertyParseOptions.AddModuleRelativePath))
 			{
 				UhtParsingScope.AddModuleRelativePathToMetaData(propertySettings.MetaData, topScope.HeaderFile);
 			}
@@ -924,53 +902,37 @@ namespace EpicGames.UHT.Parsers
 					tokenReader.Require(')');
 				}
 
-				Finalize(topScope, specifierContext, ref nameToken, new ReadOnlyMemory<UhtToken>(_currentTypeTokens.ToArray()), layoutMacroType, propertyDelegate);
+				Finalize(topScope, options, specifierContext, ref nameToken, typeTokens, layoutMacroType, propertyDelegate);
 			}
-			else if (_options.HasAnyFlags(UhtPropertyParseOptions.List))
+			else if (options.HasAnyFlags(UhtPropertyParseOptions.List))
 			{
 				// Extract the property name from the types
-				if (_currentTypeTokens.Count < 2 || !_currentTypeTokens[^1].IsIdentifier())
+				if (typeTokens.Span.Length < 2 || !typeTokens.Span[^1].IsIdentifier())
 				{
 					throw new UhtException(tokenReader, $"{propertySettings.PropertyCategory.GetHintText()}: Expected name");
 				}
-				UhtToken nameToken = _currentTypeTokens[^1];
-				_currentTypeTokens.RemoveAt(_currentTypeTokens.Count - 1);
+				UhtToken nameToken = typeTokens.Span[^1];
+				typeTokens = typeTokens[..^1];
 				CheckForOptionalParts(tokenReader, propertySettings, ref nameToken);
 
-				ReadOnlyMemory<UhtToken> typeTokens = new(_currentTypeTokens.ToArray());
+				Finalize(topScope, options, specifierContext, ref nameToken, typeTokens, layoutMacroType, propertyDelegate);
 
-				while (true)
+				// Check for unsupported comma delimited properties
+				if (tokenReader.TryOptional(','))
 				{
-					UhtProperty _ = Finalize(topScope, specifierContext, ref nameToken, typeTokens, layoutMacroType, propertyDelegate);
-
-					// If we have reached the end
-					if (!tokenReader.TryOptional(','))
-					{
-						break;
-					}
-
-					// While we could continue parsing, the old UHT would flag this as an error.
-					throw new UhtException(tokenReader, $"Comma delimited properties cannot be converted");
-
-					// we'll need any metadata tags we parsed later on when we call ConvertEOLCommentToTooltip() so the tags aren't clobbered
-					//SpecifierContext.PropertySettings.MetaData = Property.MetaData.Clone();
-					//SpecifierContext.PropertySettings.ResetTrailingSettings();
-
-					// Get the next property name
-					//NameToken = tokenReader.GetIdentifier();
-					//CheckForOptionalParts(SpecifierContext.PropertySettings, ref NameToken);
-					}
+					throw new UhtException(tokenReader, $"Comma delimited properties are not supported");
 				}
-			else if (_options.HasAnyFlags(UhtPropertyParseOptions.CommaSeparatedName))
+			}
+			else if (options.HasAnyFlags(UhtPropertyParseOptions.CommaSeparatedName))
 			{
 				tokenReader.Require(',');
 				UhtToken nameToken = tokenReader.GetIdentifier();
 				CheckForOptionalParts(tokenReader, propertySettings, ref nameToken);
-				Finalize(topScope, specifierContext, ref nameToken, new ReadOnlyMemory<UhtToken>(_currentTypeTokens.ToArray()), layoutMacroType, propertyDelegate);
+				Finalize(topScope, options, specifierContext, ref nameToken, typeTokens, layoutMacroType, propertyDelegate);
 			}
-			else if (_options.HasAnyFlags(UhtPropertyParseOptions.FunctionNameIncluded))
+			else if (options.HasAnyFlags(UhtPropertyParseOptions.FunctionNameIncluded))
 			{
-				UhtToken nameToken = _currentTypeTokens[^1];
+				UhtToken nameToken = typeTokens.Span[^1];
 				nameToken.Value = new StringView("Function");
 				if (CheckForOptionalParts(tokenReader, propertySettings, ref nameToken))
 				{
@@ -978,31 +940,31 @@ namespace EpicGames.UHT.Parsers
 				}
 				else
 				{
-					if (_currentTypeTokens.Count < 2 || !_currentTypeTokens[^1].IsIdentifier())
+					if (typeTokens.Span.Length < 2 || !typeTokens.Span[^1].IsIdentifier())
 					{
 						throw new UhtException(tokenReader, $"{propertySettings.PropertyCategory.GetHintText()}: Expected name");
 					}
-					nameToken = _currentTypeTokens[^1];
-					_currentTypeTokens.RemoveAt(_currentTypeTokens.Count - 1);
+					nameToken = typeTokens.Span[^1];
+					typeTokens = typeTokens[..^1];
 				}
-				Finalize(topScope, specifierContext, ref nameToken, new ReadOnlyMemory<UhtToken>(_currentTypeTokens.ToArray()), layoutMacroType, propertyDelegate);
-				}
-			else if (_options.HasAnyFlags(UhtPropertyParseOptions.NameIncluded))
+				Finalize(topScope, options, specifierContext, ref nameToken, typeTokens, layoutMacroType, propertyDelegate);
+			}
+			else if (options.HasAnyFlags(UhtPropertyParseOptions.NameIncluded))
 			{
-				if (_currentTypeTokens.Count < 2 || !_currentTypeTokens[^1].IsIdentifier())
+				if (typeTokens.Span.Length < 2 || !typeTokens.Span[^1].IsIdentifier())
 				{
 					throw new UhtException(tokenReader, $"{propertySettings.PropertyCategory.GetHintText()}: Expected name");
 				}
-				UhtToken nameToken = _currentTypeTokens[^1];
-				_currentTypeTokens.RemoveAt(_currentTypeTokens.Count - 1);
+				UhtToken nameToken = typeTokens.Span[^1];
+				typeTokens = typeTokens[..^1];
 				CheckForOptionalParts(tokenReader, propertySettings, ref nameToken);
-				Finalize(topScope, specifierContext, ref nameToken, new ReadOnlyMemory<UhtToken>(_currentTypeTokens.ToArray()), layoutMacroType, propertyDelegate);
+				Finalize(topScope, options, specifierContext, ref nameToken, typeTokens, layoutMacroType, propertyDelegate);
 			}
 			else
 			{
 				UhtToken nameToken = new();
 				CheckForOptionalParts(tokenReader, propertySettings, ref nameToken);
-				Finalize(topScope, specifierContext, ref nameToken, new ReadOnlyMemory<UhtToken>(_currentTypeTokens.ToArray()), layoutMacroType, propertyDelegate);
+				Finalize(topScope, options, specifierContext, ref nameToken, typeTokens, layoutMacroType, propertyDelegate);
 			}
 		}
 
@@ -1102,13 +1064,14 @@ namespace EpicGames.UHT.Parsers
 		/// Finish creating the property
 		/// </summary>
 		/// <param name="topScope">Top most scope being parsed</param>
+		/// <param name="options">Parsing options</param>
 		/// <param name="specifierContext">Context of the property being parsed</param>
 		/// <param name="nameToken">The name of the property</param>
 		/// <param name="typeTokens">Series of tokens that represent the type</param>
 		/// <param name="layoutMacroType">Optional layout macro type being parsed</param>
 		/// <param name="propertyDelegate">Delegate to invoke when processing has been completed</param>
 		/// <returns>The newly created property.  During the parsing phase, this will often be a temporary property if the type references engine types.</returns>
-		private UhtProperty Finalize(UhtParsingScope topScope, UhtPropertySpecifierContext specifierContext, ref UhtToken nameToken, 
+		private static UhtProperty Finalize(UhtParsingScope topScope, UhtPropertyParseOptions options, UhtPropertySpecifierContext specifierContext, ref UhtToken nameToken, 
 			ReadOnlyMemory<UhtToken> typeTokens, UhtLayoutMacroType layoutMacroType, UhtPropertyDelegate propertyDelegate)
 		{
 			UhtPropertySettings propertySettings = specifierContext.PropertySettings;
@@ -1195,7 +1158,7 @@ namespace EpicGames.UHT.Parsers
 			propertyDelegate(topScope, newProperty, ref nameToken, layoutMacroType);
 
 			// Void properties don't get added when they are the return value
-			if (newProperty.PropertyCategory != UhtPropertyCategory.Return || !_options.HasAnyFlags(UhtPropertyParseOptions.DontAddReturn))
+			if (newProperty.PropertyCategory != UhtPropertyCategory.Return || !options.HasAnyFlags(UhtPropertyParseOptions.DontAddReturn))
 			{
 				topScope.ScopeType.AddChild(newProperty);
 			}
@@ -1240,28 +1203,42 @@ namespace EpicGames.UHT.Parsers
 			}
 		}
 
-		private bool GatherTypeTokens(ref UhtToken token)
-			{
-			if (_currentTemplateDepth == 0 && token.IsSymbol() && (token.IsValue(',') || token.IsValue('(') || token.IsValue(')') || 
-				token.IsValue(';') || token.IsValue('[') || token.IsValue(':') || token.IsValue('=') || token.IsValue('{')))
-	{
-				return false;
-		}
-
-			_currentTypeTokens.Add(token);
-			if (token.IsSymbol('<'))
-		{
-				++_currentTemplateDepth;
-		}
-			else if (token.IsSymbol('>'))
+		private static ReadOnlyMemory<UhtToken> GatherTypeTokens(IUhtTokenReader tokenReader)
 		{
-				if (_currentTemplateDepth == 0)
+			List<UhtToken> typeTokens = [];
+			int depth = 0;
+
+			while (true)
 			{
-					throw new UhtTokenException(_currentTokenReader!, token, "',' or ')'");
+				if (tokenReader.IsEOF)
+				{
+					throw new UhtTokenException(tokenReader, tokenReader.GetToken(), null, null);
+				}
+				ref UhtToken token = ref tokenReader.PeekToken();
+
+				if (depth == 0 && token.IsSymbol() && (token.IsValue(',') || token.IsValue('(') || token.IsValue(')') ||
+					token.IsValue(';') || token.IsValue('[') || token.IsValue(':') || token.IsValue('=') || token.IsValue('{')))
+				{
+					break;
+				}
+
+				typeTokens.Add(token);
+				tokenReader.ConsumeToken();
+				if (token.IsSymbol('<'))
+				{
+					++depth;
+				}
+				else if (token.IsSymbol('>'))
+				{
+					if (depth == 0)
+					{
+						throw new UhtTokenException(tokenReader!, token, "',' or ')'");
+					}
+					--depth;
 				}
-				--_currentTemplateDepth;
 			}
-			return true;
+
+			return typeTokens.ToArray();
 		}
 	}
 
@@ -1276,7 +1253,7 @@ namespace EpicGames.UHT.Parsers
 		private static UhtParseResult UPROPERTYKeyword(UhtParsingScope topScope, UhtParsingScope actionScope, ref UhtToken token)
 		{
 			UhtPropertyParseOptions options = UhtPropertyParseOptions.ParseLayoutMacro | UhtPropertyParseOptions.List | UhtPropertyParseOptions.AddModuleRelativePath;
-			topScope.HeaderParser.GetCachedPropertyParser().Parse(topScope, EPropertyFlags.ParmFlags, options, UhtPropertyCategory.Member, s_propertyDelegate);
+			UhtPropertyParser.Parse(topScope, EPropertyFlags.ParmFlags, options, UhtPropertyCategory.Member, s_propertyDelegate);
 			topScope.TokenReader.Require(';');
 
 			// C++ UHT TODO - Skip any extra ';'.  This can be removed if we remove UhtHeaderfileParser.ParserStatement generating errors
diff --git a/Engine/Source/Programs/Shared/EpicGames.UHT/Tokenizer/UhtTokenReaderFlowControlExtensions.cs b/Engine/Source/Programs/Shared/EpicGames.UHT/Tokenizer/UhtTokenReaderFlowControlExtensions.cs
--- a/Engine/Source/Programs/Shared/EpicGames.UHT/Tokenizer/UhtTokenReaderFlowControlExtensions.cs
+++ b/Engine/Source/Programs/Shared/EpicGames.UHT/Tokenizer/UhtTokenReaderFlowControlExtensions.cs
@@ -320,32 +320,6 @@ namespace EpicGames.UHT.Tokenizer
 			return tokenReader;
 		}
 
-		/// <summary>
-		/// Read tokens until the delegate return false.  The terminating token is not consumed.
-		/// </summary>
-		/// <param name="tokenReader">Token reader</param>
-		/// <param name="untilDelegate">Invoked with each read token.  Return true to continue tokenizing or false to terminate.</param>
-		/// <returns>Token reader</returns>
-		public static IUhtTokenReader While(this IUhtTokenReader tokenReader, UhtTokensUntilDelegate untilDelegate)
-		{
-			while (true)
-			{
-				if (tokenReader.IsEOF)
-				{
-					throw new UhtTokenException(tokenReader, tokenReader.GetToken(), null, null);
-				}
-				ref UhtToken token = ref tokenReader.PeekToken();
-				if (untilDelegate(ref token))
-				{
-					tokenReader.ConsumeToken();
-				}
-				else
-				{
-					return tokenReader;
-				}
-			}
-		}
-
 		/// <summary>
 		/// Consume tokens until one of the strings are found.  Terminating token will not be consumed.
 		/// </summary>
