[EpicGames.UBA]
* Fixed so uba logger log scope logic works

diff --git a/Engine/Source/Programs/Shared/EpicGames.UBA/Impl/LoggerImpl.cs b/Engine/Source/Programs/Shared/EpicGames.UBA/Impl/LoggerImpl.cs
--- a/Engine/Source/Programs/Shared/EpicGames.UBA/Impl/LoggerImpl.cs
+++ b/Engine/Source/Programs/Shared/EpicGames.UBA/Impl/LoggerImpl.cs
@@ -3,6 +3,7 @@
 using Microsoft.Extensions.Logging;
 using System;
 using System.Runtime.InteropServices;
+using System.Threading;
 
 namespace EpicGames.UBA.Impl
 {
@@ -17,6 +18,7 @@ namespace EpicGames.UBA.Impl
 		readonly BeginScopeCallback _beginScopeCallbackDelegate;
 		readonly EndScopeCallback _endScopeCallbackDelegate;
 		readonly LogCallback _logCallbackDelegate;
+		readonly object _lock = new object();
 
 		#region DllImport
 		[DllImport("UbaHost", CharSet = CharSet.Auto)]
@@ -67,21 +69,26 @@ namespace EpicGames.UBA.Impl
 
 		public void BeginScope()
 		{
+			Monitor.Enter(_lock);
 		}
 
 		public void EndScope()
 		{
+			Monitor.Exit(_lock);
 		}
 
 		public void Log(LogEntryType type, string message)
 		{
-			switch (type)
+			lock (_lock)
 			{
-				case LogEntryType.Error: _logger.LogError("{Message}", message); break;
-				case LogEntryType.Warning: _logger.LogWarning("{Message}", message); break;
-				case LogEntryType.Info: _logger.LogInformation("{Message}", message); break;
-				case LogEntryType.Detail: _logger.LogDebug("{Message}", message); break;
-				case LogEntryType.Debug: _logger.LogDebug("{Message}", message); break;
+				switch (type)
+				{
+					case LogEntryType.Error: _logger.LogError("{Message}", message); break;
+					case LogEntryType.Warning: _logger.LogWarning("{Message}", message); break;
+					case LogEntryType.Info: _logger.LogInformation("{Message}", message); break;
+					case LogEntryType.Detail: _logger.LogDebug("{Message}", message); break;
+					case LogEntryType.Debug: _logger.LogDebug("{Message}", message); break;
+				}
 			}
 		}
 		#endregion
