Extend UHT code generation from UBT plugin

This allows for plugins to register at injection points in the UHT C++ code generation process.
Instructions:
1. create a .ubtplugin.csproj in your plugin that UBT will pickup
2. Register a class in the C# project and annotate it with [UnrealHeaderTool]
3. Implement a static function and annotate it with the [UhtCodeGeneratorInjector] annotation.

The injectors will be discovered by reflection and automatically registered.

Currently, the only injection points are for classes and structs in the header (either in GENERATED_BODY() macro or outside of the class/struct declaration block).  Futher improvements may be considered for .cpp files, UFUNCTIONS, Delegates and UENUM.

diff --git a/Engine/Source/Programs/Shared/EpicGames.UHT/Exporters/CodeGen/UhtHeaderCodeGeneratorHFile.cs b/Engine/Source/Programs/Shared/EpicGames.UHT/Exporters/CodeGen/UhtHeaderCodeGeneratorHFile.cs
--- a/Engine/Source/Programs/Shared/EpicGames.UHT/Exporters/CodeGen/UhtHeaderCodeGeneratorHFile.cs
+++ b/Engine/Source/Programs/Shared/EpicGames.UHT/Exporters/CodeGen/UhtHeaderCodeGeneratorHFile.cs
@@ -6,6 +6,7 @@ using System.IO;
 using System.Linq;
 using System.Text;
 using EpicGames.Core;
+using EpicGames.UHT.Tables;
 using EpicGames.UHT.Types;
 using EpicGames.UHT.Utils;
 
@@ -581,11 +582,15 @@ namespace EpicGames.UHT.Exporters.CodeGen
 						ReadOnlySpan<char> shortName = scriptStruct.SourceName.AsSpan(1, scriptStruct.SourceName.Length - 7);
 						builder.Append("\tDECLARE_STATESTREAM_HANDLE(").Append(shortName).Append(") \\\r\n");
 					}
+					
+					Session.Tables!.CodeGeneratorInjectorTable.Inject(builder, scriptStruct, UhtCodeGeneratorInjectionLocation.GeneratedMacro);
 				}
 
 				// Forward declare the StaticStruct specialization in the header
 				builder.Append("template<> ").Append(Module.Api).Append("UScriptStruct* StaticStruct<struct ").Append(scriptStruct.SourceName).Append(">();\r\n");
 				builder.Append("\r\n");
+				
+				Session.Tables!.CodeGeneratorInjectorTable.Inject(builder, scriptStruct, UhtCodeGeneratorInjectionLocation.Header);
 
 				// Verse declarations
 				if (scriptStruct.IsVerseField)
@@ -824,6 +829,8 @@ namespace EpicGames.UHT.Exporters.CodeGen
 
 			// Forward declare the StaticClass specialization in the header
 			builder.Append("template<> ").Append(Module.Api).Append("UClass* StaticClass<class ").Append(classObj.SourceName).Append(">();\r\n");
+			
+			Session.Tables!.CodeGeneratorInjectorTable.Inject(builder, classObj, UhtCodeGeneratorInjectionLocation.Header);
 			builder.Append("\r\n");
 
 			// Verse declarations
@@ -1712,6 +1719,7 @@ namespace EpicGames.UHT.Exporters.CodeGen
 					builder.Append(" \\\r\n");
 				}
 				builder.Append(EnableDeprecationWarnings).Append(" \\\r\n");
+				Session.Tables!.CodeGeneratorInjectorTable.Inject(builder, classObj, UhtCodeGeneratorInjectionLocation.GeneratedMacro);
 			}
 			return builder;
 		}
diff --git a/Engine/Source/Programs/Shared/EpicGames.UHT/Tables/UhtCodeGeneratorInjectorTable.cs b/Engine/Source/Programs/Shared/EpicGames.UHT/Tables/UhtCodeGeneratorInjectorTable.cs
new file mode 100644
--- /dev/null
+++ b/Engine/Source/Programs/Shared/EpicGames.UHT/Tables/UhtCodeGeneratorInjectorTable.cs
@@ -0,0 +1,153 @@
+﻿// Copyright Epic Games, Inc. All Rights Reserved.
+
+using System;
+using System.Collections.Generic;
+using System.Linq;
+using System.Reflection;
+using System.Text;
+using EpicGames.Core;
+using EpicGames.UHT.Exporters.CodeGen;
+using EpicGames.UHT.Types;
+using EpicGames.UHT.Utils;
+
+namespace EpicGames.UHT.Tables
+{
+	/// <summary>
+	/// Location where code will be injected
+	/// </summary>
+	public enum UhtCodeGeneratorInjectionLocation
+	{
+		/// <summary>
+		/// Injection location will be in the generated.h after built-in code generation for the respective type
+		/// </summary>
+		Header,
+		/// <summary>
+		/// Injection location within the macros for the type.
+		/// </summary>
+		GeneratedMacro
+	}
+
+	/// <summary>
+	/// Method attribute which defines a code generation injection behaviour.
+	/// </summary>
+	[AttributeUsage(AttributeTargets.Method, AllowMultiple = false)]
+	public sealed class UhtCodeGeneratorInjectorAttribute : Attribute
+	{
+		/// <summary>
+		/// The UhtType that the injector will be applied to.  
+		/// </summary>
+		public required Type UhtType { get; set; }
+		/// <summary>
+		/// Specifies the injection location for the injector to run in.
+		/// </summary>
+		public required UhtCodeGeneratorInjectionLocation Location { get; set; }
+	}
+
+	internal record struct UhtCodeGeneratorInjectorDesire(Type UhtType, UhtCodeGeneratorInjectionLocation Location);
+
+	/// <summary>
+	/// Entrypoint for code generation injector
+	/// </summary>
+	internal delegate void UhtCodeGeneratorInjectorDelegate(StringBuilder builder, UhtType type, int leadingTabs, string eolSequence);
+	
+	/// <summary>
+	/// Holds a set of injectors discovered via reflection
+	/// </summary>
+	public sealed class UhtCodeGeneratorInjectorTable
+	{
+		private readonly Dictionary<UhtCodeGeneratorInjectorDesire, List<UhtCodeGeneratorInjectorDelegate>?> _injectorsTable = new();
+		private readonly List<Type> _supportedLeafTypes;
+
+		/// <summary>
+		/// Constructs UhtCodeGeneratorInjectorTable
+		/// </summary>
+		public UhtCodeGeneratorInjectorTable()
+		{
+			_supportedLeafTypes = new List<Type>
+			{
+				typeof(UhtClass),
+				typeof(UhtScriptStruct),
+			};
+		}
+
+		internal void Inject(StringBuilder builder, UhtType type, UhtCodeGeneratorInjectionLocation location)
+		{
+			string eolSequence;
+			int leadingTabs;
+			switch (location)
+			{
+				case UhtCodeGeneratorInjectionLocation.Header:
+					eolSequence = "\r\n";
+					leadingTabs = 0;
+					break;
+				case UhtCodeGeneratorInjectionLocation.GeneratedMacro:
+					eolSequence = " \\\r\n";
+					leadingTabs = 1;
+					break;
+				default:
+					throw new ArgumentOutOfRangeException(nameof(location), location, "Unhandled injection location");
+			}
+			
+			UhtCodeGeneratorInjectorDesire desire = new(type.GetType(), location);
+			if (_injectorsTable.TryGetValue(desire, out List<UhtCodeGeneratorInjectorDelegate>? injectors))
+			{
+				foreach (UhtCodeGeneratorInjectorDelegate injector in injectors!)
+				{
+					injector.Invoke(builder, type, leadingTabs, eolSequence);
+				}
+			}
+		}
+
+		internal void OnUhtCodeGeneratorInjectorAttribute(Type type, MethodInfo method, UhtCodeGeneratorInjectorAttribute attribute)
+		{
+			// Walk back from each support leaf types up to the declared type
+			// If our leaf type IS of the declared type, then add a delegate for the handler of the leaf type
+			// This allows for users to write injectors that pertain to multiple leaf types if required (ie. all UhtNumericProperty)
+			foreach (Type supportedLeafType in _supportedLeafTypes)
+			{
+				if (attribute.UhtType.IsAssignableFrom(supportedLeafType))
+				{
+					UhtCodeGeneratorInjectorDesire desire = new(supportedLeafType, attribute.Location);
+					List<UhtCodeGeneratorInjectorDelegate>? injectors;
+					if (!_injectorsTable.TryGetValue(desire, out injectors))
+					{
+						injectors = [];
+						_injectorsTable.Add(desire, injectors);
+					}
+					injectors!.Add((UhtCodeGeneratorInjectorDelegate)Delegate.CreateDelegate(typeof(UhtCodeGeneratorInjectorDelegate), method));
+				}
+			}
+		}
+
+		internal void Sort()
+		{
+			// Sort the delegates to ensure determinism across multiple runs
+			foreach (List<UhtCodeGeneratorInjectorDelegate>? injectors in _injectorsTable.Values)
+			{
+				if (injectors == null) { continue; }
+				injectors.Sort((Lhs, Rhs) =>
+				{
+					MethodInfo methodLhs = Lhs.Method;
+					MethodInfo methodRhs = Rhs.Method;
+					
+					// Sort by module
+					int moduleComparison = String.Compare(methodLhs.Module.Name, methodLhs.Module.Name, StringComparison.Ordinal);
+					if (moduleComparison != 0)
+					{
+						return moduleComparison;
+					}
+					
+					// Sort by declaring class
+					int declaredClassComparison = String.Compare(methodLhs.DeclaringType!.Name, methodRhs.DeclaringType!.Name, StringComparison.Ordinal);
+					if (declaredClassComparison != 0)
+					{
+						return declaredClassComparison;
+					}
+					
+					// Sort by method name
+					return String.Compare(methodLhs.Name, methodRhs.Name, StringComparison.Ordinal);
+				});
+			}
+		}
+	}
+}
diff --git a/Engine/Source/Programs/Shared/EpicGames.UHT/Utils/UhtTables.cs b/Engine/Source/Programs/Shared/EpicGames.UHT/Utils/UhtTables.cs
--- a/Engine/Source/Programs/Shared/EpicGames.UHT/Utils/UhtTables.cs
+++ b/Engine/Source/Programs/Shared/EpicGames.UHT/Utils/UhtTables.cs
@@ -489,6 +489,11 @@ namespace EpicGames.UHT.Utils
 		/// </summary>
 		public UhtLocTextDefaultValueTable LocTextDefaultValueTable { get; } = new UhtLocTextDefaultValueTable();
 
+		/// <summary>
+		/// 
+		/// </summary>
+		public UhtCodeGeneratorInjectorTable CodeGeneratorInjectorTable { get; } = new UhtCodeGeneratorInjectorTable();
+		
 		/// <summary>
 		/// Construct a new table collection
 		/// </summary>
@@ -618,6 +623,7 @@ namespace EpicGames.UHT.Utils
 			KeywordTables.Merge();
 			SpecifierTables.Merge();
 			SpecifierValidatorTables.Merge();
+			CodeGeneratorInjectorTable.Sort();
 
 			// Invoke this to generate an exception if there is no default
 			_ = PropertyTypeTable.Default;
@@ -675,6 +681,10 @@ namespace EpicGames.UHT.Utils
 					{
 						LocTextDefaultValueTable.OnLocTextDefaultValueAttribute(methodInfo, locTextDefaultValueAttribute);
 					}
+					else if (methodAttribute is UhtCodeGeneratorInjectorAttribute codeGeneratorInjectorAttribute)
+					{
+						CodeGeneratorInjectorTable.OnUhtCodeGeneratorInjectorAttribute(type, methodInfo, codeGeneratorInjectorAttribute);
+					}
 				}
 			}
 		}
