Adding the collection of type tokens to the properties

#rnx

diff --git a/Engine/Source/Programs/Shared/EpicGames.UHT/Parsers/UhtPropertyParser.cs b/Engine/Source/Programs/Shared/EpicGames.UHT/Parsers/UhtPropertyParser.cs
--- a/Engine/Source/Programs/Shared/EpicGames.UHT/Parsers/UhtPropertyParser.cs
+++ b/Engine/Source/Programs/Shared/EpicGames.UHT/Parsers/UhtPropertyParser.cs
@@ -2,6 +2,7 @@
 
 using System;
 using System.Collections.Generic;
+using System.Diagnostics;
 using System.Diagnostics.CodeAnalysis;
 using System.Text;
 using System.Threading;
@@ -376,11 +377,6 @@ namespace EpicGames.UHT.Parsers
 		/// <inheritdoc/>
 		protected override string PGetMacroText => "invalid";
 
-		/// <summary>
-		/// Collection of type tokens 
-		/// </summary>
-		public ReadOnlyMemory<UhtToken> TypeTokens { get; set; }
-
 		/// <summary>
 		/// Property settings being parsed
 		/// </summary>
@@ -390,10 +386,8 @@ namespace EpicGames.UHT.Parsers
 		/// Construct a new property to be resolved
 		/// </summary>
 		/// <param name="propertySettings">Property settings</param>
-		/// <param name="typeTokens">Type tokens</param>
-		public UhtPreResolveProperty(UhtPropertySettings propertySettings, ReadOnlyMemory<UhtToken> typeTokens) : base(propertySettings)
+		public UhtPreResolveProperty(UhtPropertySettings propertySettings) : base(propertySettings)
 		{
-			TypeTokens = typeTokens;
 			PropertySettings = propertySettings;
 		}
 
@@ -508,49 +502,14 @@ namespace EpicGames.UHT.Parsers
 		/// resolve any previously parsed property to the final version.
 		/// </summary>
 		/// <param name="resolvePhase">Used to detect if the property is being parsed or resolved</param>
-		/// <param name="data">Character buffer for header file</param>
 		/// <param name="propertySettings">The property settings.</param>
-		/// <param name="typeTokens">The tokens that represent the type</param>
 		/// <returns></returns>
-		public static UhtProperty? ResolveProperty(UhtPropertyResolvePhase resolvePhase, UhtPropertySettings propertySettings, ReadOnlyMemory<char> data, ReadOnlyMemory<UhtToken> typeTokens)
+		private static UhtProperty? ResolveProperty(UhtPropertyResolvePhase resolvePhase, UhtPropertySettings propertySettings)
 		{
-			UhtSession session = propertySettings.Outer.Session;
-
-			using UhtTokenReplayReaderBorrower borrowedReader = new(propertySettings.Outer, propertySettings.Outer.HeaderFile.Data.Memory, typeTokens, UhtTokenType.EndOfType);
-			IUhtTokenReader replayReader = borrowedReader.Reader;
-
-			// Loop through the tokens until we find a known property type or the start of a template argument list
-			ReadOnlySpan<UhtToken> typeTokensSpan = typeTokens.Span;
-			for (int index = 0; index < typeTokensSpan.Length; ++index)
-			{
-				if (typeTokensSpan[index].IsSymbol())
-				{
-					ReadOnlySpan<char> span = typeTokensSpan[index].Value.Span;
-					if (span.Length == 1 && (span[0] == '<' || span[0] == '>' || span[0] == ','))
-					{
-						break;
-					}
-				}
-				else if (typeTokensSpan[index].IsIdentifier())
-				{
-					UhtToken copy = typeTokensSpan[index];
-					session.Config!.RedirectTypeIdentifier(ref copy);
-					if (session.TryGetPropertyType(copy.Value, out UhtPropertyType propertyType))
-					{
-						if (resolvePhase == UhtPropertyResolvePhase.Resolving || propertyType.Options.HasAnyFlags(UhtPropertyTypeOptions.Immediate))
-						{
-							return ResolvePropertyType(resolvePhase, propertySettings, replayReader, propertyType, typeTokens.Span[index], false);
-						}
-					}
-				}
-			}
-
-			// Try the default processor.  We only use the default processor when trying to resolve something post parsing phase.
-			if (resolvePhase == UhtPropertyResolvePhase.Resolving)
-			{
-				return ResolvePropertyType(resolvePhase, propertySettings, replayReader, session.DefaultPropertyType, new UhtToken(), false);
-			}
-			return null;
+			Debug.Assert(propertySettings.TypeTokens.AllTokens.Length > 0);
+			using UhtTokenReplayReaderBorrower borrowedReader = new(propertySettings.Outer, propertySettings.Outer.HeaderFile.Data.Memory, propertySettings.TypeTokens.AllTokens, UhtTokenType.EndOfType);
+			UhtPropertyResolveArgs args = new(resolvePhase, propertySettings, borrowedReader.Reader);
+			return args.ResolveProperty();
 		}
 
 		private static readonly ThreadLocal<UhtPropertySettings> s_tlsPropertySettings = new(() => { return new UhtPropertySettings(); });
@@ -580,7 +539,7 @@ namespace EpicGames.UHT.Parsers
 				if (type.Children[index] is UhtPreResolveProperty property)
 				{
 					propertySettings.Reset(property, propertyOptions);
-					UhtProperty? resolved = UhtPropertyParser.ResolveProperty(UhtPropertyResolvePhase.Resolving, propertySettings, property.HeaderFile.Data.Memory, property.TypeTokens);
+					UhtProperty? resolved = ResolveProperty(UhtPropertyResolvePhase.Resolving, propertySettings);
 					if (resolved != null)
 					{
 						if (inSymbolTable && resolved != property)
@@ -594,81 +553,14 @@ namespace EpicGames.UHT.Parsers
 		}
 
 		/// <summary>
-		/// Parse a template parameter
+		/// Resolve the given property and finish parsing any postfix settings
 		/// </summary>
-		/// <param name="resolvePhase">Resolution phase</param>
-		/// <param name="parentPropertySettings">Parent property (container) settings</param>
-		/// <param name="paramName">Name of the template parameter</param>
-		/// <param name="tokenReader">Token type</param>
-		/// <returns>Parsed property</returns>
-		public static UhtProperty? ParseTemplateParam(UhtPropertyResolvePhase resolvePhase, UhtPropertySettings parentPropertySettings, StringView paramName, IUhtTokenReader tokenReader)
+		/// <param name="args">Arguments for the resolution process</param>
+		/// <param name="propertyType">Property type used to resolve the property</param>
+		/// <returns>Created property or null if it could not be resolved</returns>
+		/// <exception cref="UhtTokenException"></exception>
+		public static UhtProperty? ResolvePropertyType(UhtPropertyResolveArgs args, UhtPropertyType propertyType)
 		{
-			UhtSession session = parentPropertySettings.Outer.Session;
-
-			// Save the token reader state.  We need this to restore back to the start when invoking the resolve methods.
-			{
-				using UhtThreadBorrower<UhtPropertySpecifierContext> borrower = new(true);
-				UhtPropertySpecifierContext specifierContext = borrower.Instance;
-				specifierContext.Type = parentPropertySettings.Outer;
-				specifierContext.TokenReader = tokenReader;
-				specifierContext.AccessSpecifier = UhtAccessSpecifier.None;
-				specifierContext.MessageSite = tokenReader;
-				specifierContext.PropertySettings.Reset(parentPropertySettings, paramName.ToString(), tokenReader);
-				specifierContext.MetaData = specifierContext.PropertySettings.MetaData;
-				specifierContext.MetaNameIndex = UhtMetaData.IndexNone;
-				specifierContext.SeenEditSpecifier = false;
-				specifierContext.SeenBlueprintWriteSpecifier = false;
-				specifierContext.SeenBlueprintReadOnlySpecifier = false;
-				specifierContext.SeenBlueprintGetterSpecifier = false;
-
-				// Preparse any UPARAM 
-				PreParseTypeInternal(specifierContext, true);
-
-				// Save the state at this point
-				using UhtTokenSaveState saveState = new(tokenReader);
-
-				UhtPropertySettings propertySettings = specifierContext.PropertySettings;
-				UhtPropertyType propertyType = new();
-				while (!tokenReader.IsEOF)
-				{
-					UhtToken token = tokenReader.GetToken();
-
-					if (token.IsSymbol('<') || token.IsSymbol('>') || token.IsSymbol(','))
-					{
-						break;
-					}
-
-					if (!token.IsIdentifier())
-					{
-						continue;
-					}
-
-					UhtToken copy = token;
-					session.Config!.RedirectTypeIdentifier(ref copy);
-					if (session.TryGetPropertyType(copy.Value, out propertyType))
-					{
-						if (resolvePhase == UhtPropertyResolvePhase.Resolving || propertyType.Options.HasAnyFlags(UhtPropertyTypeOptions.Immediate))
-						{
-							saveState.RestoreState();
-							return ResolvePropertyType(resolvePhase, propertySettings, tokenReader, propertyType, token, true);
-						}
-					}
-				}
-
-				// Try the default processor.  We only use the default processor when trying to resolve something post parsing phase.
-				if (resolvePhase == UhtPropertyResolvePhase.Resolving)
-				{
-					saveState.RestoreState();
-					return ResolvePropertyType(resolvePhase, propertySettings, tokenReader, session.DefaultPropertyType, new UhtToken(), true);
-				}
-				return null;
-			}
-		}
-
-		private static UhtProperty? ResolvePropertyType(UhtPropertyResolvePhase resolvePhase,
-			UhtPropertySettings propertySettings, IUhtTokenReader tokenReader, UhtPropertyType propertyType, UhtToken matchedToken, bool isTemplate)
-		{
-			UhtPropertyResolveArgs args = new(resolvePhase, propertySettings, tokenReader, matchedToken);
 			UhtProperty? outProperty = propertyType.Delegate(args);
 			if (outProperty == null)
 			{
@@ -692,7 +584,7 @@ namespace EpicGames.UHT.Parsers
 				if (outerClass != null && outerClass.ClassFlags.HasAnyFlags(EClassFlags.Const))
 				{
 					// Eat a 'not quite truthful' const after the type; autogenerated for member variables of const classes.
-					if (tokenReader.TryOptional("const"))
+					if (args.TokenReader.TryOptional("const"))
 					{
 						outProperty.MetaData.Add(UhtNames.NativeConst, "");
 					}
@@ -700,7 +592,7 @@ namespace EpicGames.UHT.Parsers
 			}
 			else
 			{
-				if (tokenReader.TryOptional("const"))
+				if (args.TokenReader.TryOptional("const"))
 				{
 					outProperty.MetaData.Add(UhtNames.NativeConst, "");
 					outProperty.PropertyFlags |= EPropertyFlags.ConstParm;
@@ -708,13 +600,13 @@ namespace EpicGames.UHT.Parsers
 			}
 
 			// Check for unexpected '*'
-			if (tokenReader.TryOptional('*'))
+			if (args.TokenReader.TryOptional('*'))
 			{
-				tokenReader.LogError($"Inappropriate '*' on variable of type '{outProperty.GetUserFacingDecl()}', cannot have an exposed pointer to this type.");
+				args.TokenReader.LogError($"Inappropriate '*' on variable of type '{outProperty.GetUserFacingDecl()}', cannot have an exposed pointer to this type.");
 			}
 
 			// Arrays are passed by reference but are only implicitly so; setting it explicitly could cause a problem with replicated functions
-			if (tokenReader.TryOptional('&'))
+			if (args.TokenReader.TryOptional('&'))
 			{
 				switch (outProperty.PropertyCategory)
 				{
@@ -747,19 +639,19 @@ namespace EpicGames.UHT.Parsers
 				}
 			}
 
-			if (!isTemplate)
+			if (!args.IsTemplate)
 			{
-				if (!tokenReader.IsEOF)
+				if (!args.TokenReader.IsEOF)
 				{
-					throw new UhtTokenException(tokenReader, tokenReader.PeekToken(), "end of type declaration");
+					throw new UhtTokenException(args.TokenReader, args.TokenReader.PeekToken(), "end of type declaration");
 				}
 			}
 			else
 			{
-				ref UhtToken token = ref tokenReader.PeekToken();
+				ref UhtToken token = ref args.TokenReader.PeekToken();
 				if (!token.IsSymbol(',') && !token.IsSymbol('>'))
 				{
-					throw new UhtTokenException(tokenReader, tokenReader.PeekToken(), "end of type declaration");
+					throw new UhtTokenException(args.TokenReader, args.TokenReader.PeekToken(), "end of type declaration");
 				}
 			}
 			return outProperty;
@@ -803,7 +695,7 @@ namespace EpicGames.UHT.Parsers
 			}
 
 			// Parse type information including UPARAM that might appear in template arguments
-			PreParseTypeInternal(specifierContext, false);
+			PreParseType(specifierContext, false);
 
 			// Swallow inline keywords
 			if (propertySettings.PropertyCategory == UhtPropertyCategory.Return)
@@ -849,10 +741,10 @@ namespace EpicGames.UHT.Parsers
 			tokenReader.Optional("mutable");
 
 			// Gather the type tokens and possibly the property name.
-			ReadOnlyMemory<UhtToken> typeTokens = GatherTypeTokens(tokenReader);
+			UhtTypeTokens typeTokens = UhtTypeTokens.Gather(tokenReader);
 
 			// Verify we at least have one type
-			if (typeTokens.Length < 1)
+			if (typeTokens.AllTokens.Length < 1)
 			{
 				throw new UhtException(tokenReader, $"{propertySettings.PropertyCategory.GetHintText()}: Missing variable type or name");
 			}
@@ -906,13 +798,7 @@ namespace EpicGames.UHT.Parsers
 			}
 			else if (options.HasAnyFlags(UhtPropertyParseOptions.List))
 			{
-				// Extract the property name from the types
-				if (typeTokens.Span.Length < 2 || !typeTokens.Span[^1].IsIdentifier())
-				{
-					throw new UhtException(tokenReader, $"{propertySettings.PropertyCategory.GetHintText()}: Expected name");
-				}
-				UhtToken nameToken = typeTokens.Span[^1];
-				typeTokens = typeTokens[..^1];
+				UhtToken nameToken = typeTokens.ExtractTrailingIdentifier(tokenReader, propertySettings);
 				CheckForOptionalParts(tokenReader, propertySettings, ref nameToken);
 
 				Finalize(topScope, options, specifierContext, ref nameToken, typeTokens, layoutMacroType, propertyDelegate);
@@ -932,7 +818,7 @@ namespace EpicGames.UHT.Parsers
 			}
 			else if (options.HasAnyFlags(UhtPropertyParseOptions.FunctionNameIncluded))
 			{
-				UhtToken nameToken = typeTokens.Span[^1];
+				UhtToken nameToken = typeTokens.AllTokens.Span[^1];
 				nameToken.Value = new StringView("Function");
 				if (CheckForOptionalParts(tokenReader, propertySettings, ref nameToken))
 				{
@@ -940,23 +826,13 @@ namespace EpicGames.UHT.Parsers
 				}
 				else
 				{
-					if (typeTokens.Span.Length < 2 || !typeTokens.Span[^1].IsIdentifier())
-					{
-						throw new UhtException(tokenReader, $"{propertySettings.PropertyCategory.GetHintText()}: Expected name");
-					}
-					nameToken = typeTokens.Span[^1];
-					typeTokens = typeTokens[..^1];
+					nameToken = typeTokens.ExtractTrailingIdentifier(tokenReader, propertySettings);
 				}
 				Finalize(topScope, options, specifierContext, ref nameToken, typeTokens, layoutMacroType, propertyDelegate);
 			}
 			else if (options.HasAnyFlags(UhtPropertyParseOptions.NameIncluded))
 			{
-				if (typeTokens.Span.Length < 2 || !typeTokens.Span[^1].IsIdentifier())
-				{
-					throw new UhtException(tokenReader, $"{propertySettings.PropertyCategory.GetHintText()}: Expected name");
-				}
-				UhtToken nameToken = typeTokens.Span[^1];
-				typeTokens = typeTokens[..^1];
+				UhtToken nameToken = typeTokens.ExtractTrailingIdentifier(tokenReader, propertySettings);
 				CheckForOptionalParts(tokenReader, propertySettings, ref nameToken);
 				Finalize(topScope, options, specifierContext, ref nameToken, typeTokens, layoutMacroType, propertyDelegate);
 			}
@@ -975,7 +851,7 @@ namespace EpicGames.UHT.Parsers
 		/// <param name="isTemplateArgument">If true, this is part of a template argument</param>
 		/// <exception cref="UhtIceException"></exception>
 		/// <exception cref="UhtException"></exception>
-		private static void PreParseTypeInternal(UhtPropertySpecifierContext specifierContext, bool isTemplateArgument)
+		public static void PreParseType(UhtPropertySpecifierContext specifierContext, bool isTemplateArgument)
 		{
 			UhtPropertySettings propertySettings = specifierContext.PropertySettings;
 			IUhtTokenReader tokenReader = specifierContext.TokenReader;
@@ -1071,16 +947,17 @@ namespace EpicGames.UHT.Parsers
 		/// <param name="layoutMacroType">Optional layout macro type being parsed</param>
 		/// <param name="propertyDelegate">Delegate to invoke when processing has been completed</param>
 		/// <returns>The newly created property.  During the parsing phase, this will often be a temporary property if the type references engine types.</returns>
-		private static UhtProperty Finalize(UhtParsingScope topScope, UhtPropertyParseOptions options, UhtPropertySpecifierContext specifierContext, ref UhtToken nameToken, 
-			ReadOnlyMemory<UhtToken> typeTokens, UhtLayoutMacroType layoutMacroType, UhtPropertyDelegate propertyDelegate)
+		private static UhtProperty Finalize(UhtParsingScope topScope, UhtPropertyParseOptions options, UhtPropertySpecifierContext specifierContext, ref UhtToken nameToken,
+			UhtTypeTokens typeTokens, UhtLayoutMacroType layoutMacroType, UhtPropertyDelegate propertyDelegate)
 		{
 			UhtPropertySettings propertySettings = specifierContext.PropertySettings;
+			propertySettings.TypeTokens = new(typeTokens, 0);
 			IUhtTokenReader tokenReader = specifierContext.TokenReader;
 
 			propertySettings.SourceName = propertySettings.PropertyCategory == UhtPropertyCategory.Return ? "ReturnValue" : nameToken.Value.ToString();
 
 			// Try to resolve the property using any immediate mode property types
-			UhtProperty newProperty = ResolveProperty(UhtPropertyResolvePhase.Parsing, propertySettings, propertySettings.Outer.HeaderFile.Data.Memory, typeTokens) ?? new UhtPreResolveProperty(propertySettings, typeTokens);
+			UhtProperty newProperty = ResolveProperty(UhtPropertyResolvePhase.Parsing, propertySettings) ?? new UhtPreResolveProperty(propertySettings);
 
 			// Force the category in non-engine projects
 			if (newProperty.PropertyCategory == UhtPropertyCategory.Member)
@@ -1202,44 +1079,6 @@ namespace EpicGames.UHT.Parsers
 				throw new UhtException(tokenReader, $"{propertySettings.PropertyCategory.GetHintText()} {nameToken.Value}: Missing array dimensions or terminating '{terminator}'");
 			}
 		}
-
-		private static ReadOnlyMemory<UhtToken> GatherTypeTokens(IUhtTokenReader tokenReader)
-		{
-			List<UhtToken> typeTokens = [];
-			int depth = 0;
-
-			while (true)
-			{
-				if (tokenReader.IsEOF)
-				{
-					throw new UhtTokenException(tokenReader, tokenReader.GetToken(), null, null);
-				}
-				ref UhtToken token = ref tokenReader.PeekToken();
-
-				if (depth == 0 && token.IsSymbol() && (token.IsValue(',') || token.IsValue('(') || token.IsValue(')') ||
-					token.IsValue(';') || token.IsValue('[') || token.IsValue(':') || token.IsValue('=') || token.IsValue('{')))
-				{
-					break;
-				}
-
-				typeTokens.Add(token);
-				tokenReader.ConsumeToken();
-				if (token.IsSymbol('<'))
-				{
-					++depth;
-				}
-				else if (token.IsSymbol('>'))
-				{
-					if (depth == 0)
-					{
-						throw new UhtTokenException(tokenReader!, token, "',' or ')'");
-					}
-					--depth;
-				}
-			}
-
-			return typeTokens.ToArray();
-		}
 	}
 
 	[UnrealHeaderTool]
diff --git a/Engine/Source/Programs/Shared/EpicGames.UHT/Tables/UhtPropertyTypeTable.cs b/Engine/Source/Programs/Shared/EpicGames.UHT/Tables/UhtPropertyTypeTable.cs
--- a/Engine/Source/Programs/Shared/EpicGames.UHT/Tables/UhtPropertyTypeTable.cs
+++ b/Engine/Source/Programs/Shared/EpicGames.UHT/Tables/UhtPropertyTypeTable.cs
@@ -87,247 +87,12 @@ namespace EpicGames.UHT.Tables
 		}
 	}
 
-	/// <summary>
-	/// The phase of UHT where the property is being resolved
-	/// </summary>
-	public enum UhtPropertyResolvePhase
-	{
-
-		/// <summary>
-		/// Resolved during the source processing phase.  Immediate property types only.
-		/// </summary>
-		Parsing,
-
-		/// <summary>
-		/// Resolved during the resolve phase.  Non-immediate property types only.
-		/// </summary>
-		Resolving,
-	}
-
-	/// <summary>
-	/// Structure defining the arguments used in the delegate to resolve property types
-	/// </summary>
-	/// <param name="ResolvePhase">Specifies if this is being resolved during the parsing phase or the resolution phase.  Type lookups can not happen during the parsing phase</param>
-	/// <param name="PropertySettings">The configuration of the property</param>
-	/// <param name="TokenReader">The token reader containing the type</param>
-	/// <param name="MatchedToken">The token that matched the delegate unless the delegate is the default resolver.</param>
-	public record struct UhtPropertyResolveArgs(UhtPropertyResolvePhase ResolvePhase, UhtPropertySettings PropertySettings, IUhtTokenReader TokenReader, UhtToken MatchedToken)
-	{
-
-		/// <summary>
-		/// Currently running UHT session
-		/// </summary>
-		public readonly UhtSession Session => PropertySettings.Outer.Session;
-
-		/// <summary>
-		/// UHT configuration
-		/// </summary>
-		public readonly IUhtConfig Config => Session.Config!;
-
-		/// <summary>
-		/// When processing type, make sure that the next token is the expected token
-		/// </summary>
-		/// <returns>true if there could be more header to process, false if the end was reached.</returns>
-		public readonly bool SkipExpectedType()
-		{
-			if (TokenReader.TryOptional(MatchedToken.Value))
-			{
-				return true;
-			}
-			if (PropertySettings.PropertyCategory == UhtPropertyCategory.Member && TokenReader.TryOptional("const"))
-			{
-				TokenReader.LogError("Const properties are not supported.");
-			}
-			else
-			{
-				TokenReader.LogError($"Inappropriate keyword '{TokenReader.PeekToken().Value}' on variable of type '{MatchedToken.Value}'");
-			}
-			return false;
-		}
-
-		/// <summary>
-		/// Parse a template type
-		/// </summary>
-		/// <param name="returnUInterface">If true, return the UInterface instead of the type listed</param>
-		/// <returns>Referenced class</returns>
-		public readonly UhtClass? ParseTemplateObject(bool returnUInterface)
-		{
-			if (TokenReader.TryOptional("const"))
-			{
-				PropertySettings.MetaData.Add(UhtNames.NativeConst, "");
-			}
-
-			if (!SkipExpectedType())
-			{
-				return null;
-			}
-
-			bool isNativeConstTemplateArg = false;
-			UhtToken identifier = new();
-			TokenReader
-				.Require('<')
-				.Optional("const", () => { isNativeConstTemplateArg = true; })
-				.Optional("class")
-				.RequireIdentifier((ref UhtToken token) => { identifier = token; })
-				.Optional("const", () => { isNativeConstTemplateArg = true; })
-				.Require('>');
-
-			if (isNativeConstTemplateArg)
-			{
-				PropertySettings.MetaData.Add(UhtNames.NativeConstTemplateArg, "");
-			}
-
-			Config.RedirectTypeIdentifier(ref identifier);
-			UhtClass? returnClass = PropertySettings.Outer.FindType(UhtFindOptions.SourceName | UhtFindOptions.Class, ref identifier, TokenReader) as UhtClass;
-			if (returnClass != null && returnClass.AlternateObject != null && returnUInterface)
-			{
-				returnClass = returnClass.AlternateObject as UhtClass;
-			}
-			return returnClass;
-		}
-
-		/// <summary>
-		/// Parse a template type
-		/// </summary>
-		/// <returns>Referenced class</returns>
-		public readonly UhtClass? ParseTemplateClass()
-		{
-			UhtToken identifier = new();
-
-			if (TokenReader.TryOptional("const"))
-			{
-				PropertySettings.MetaData.Add(UhtNames.NativeConst, "");
-			}
-
-			TokenReader.Optional("class");
-
-			if (!SkipExpectedType())
-			{
-				return null;
-			}
-
-			TokenReader
-				.Require('<')
-				.Optional("class")
-				.RequireIdentifier((ref UhtToken token) => { identifier = token; })
-				.Require('>');
-
-			Config.RedirectTypeIdentifier(ref identifier);
-			UhtClass? returnClass = PropertySettings.Outer.FindType(UhtFindOptions.SourceName | UhtFindOptions.Class, ref identifier, TokenReader) as UhtClass;
-			if (returnClass != null && returnClass.AlternateObject != null)
-			{
-				returnClass = returnClass.AlternateObject as UhtClass;
-			}
-			return returnClass;
-		}
-
-		/// <summary>
-		/// Parse a template type
-		/// </summary>
-		/// <returns>Referenced class</returns>
-		public readonly UhtScriptStruct? ParseTemplateScriptStruct()
-		{
-			UhtToken identifier = new();
-
-			if (TokenReader.TryOptional("const"))
-			{
-				PropertySettings.MetaData.Add(UhtNames.NativeConst, "");
-			}
-
-			TokenReader.Optional("struct");
-
-			if (!SkipExpectedType())
-			{
-				return null;
-			}
-
-			TokenReader
-				.Require('<')
-				.Optional("struct")
-				.RequireIdentifier((ref UhtToken token) => { identifier = token; })
-				.Require('>');
-
-			Config.RedirectTypeIdentifier(ref identifier);
-			return PropertySettings.Outer.FindType(UhtFindOptions.SourceName | UhtFindOptions.ScriptStruct, ref identifier, TokenReader) as UhtScriptStruct;
-		}
-
-		/// <summary>
-		/// Logs message for Object pointers to convert UObject* to TObjectPtr or the reverse
-		/// </summary>
-		/// <param name="engineBehavior">Expected behavior for engine types</param>
-		/// <param name="enginePluginBehavior">Expected behavior for engine plugin types</param>
-		/// <param name="nonEngineBehavior">Expected behavior for non-engine types</param>
-		/// <param name="pointerTypeDesc">Description of the pointer type</param>
-		/// <param name="typeStartPos">Starting character position of the type</param>
-		/// <param name="alternativeTypeDesc">Suggested alternate declaration</param>
-		/// <exception cref="UhtIceException">Thrown if the behavior type is unexpected</exception>
-		public readonly void ConditionalLogPointerUsage(UhtIssueBehavior engineBehavior, UhtIssueBehavior enginePluginBehavior,
-			UhtIssueBehavior nonEngineBehavior, string pointerTypeDesc, int typeStartPos, string? alternativeTypeDesc)
-		{
-			if (PropertySettings.PropertyCategory != UhtPropertyCategory.Member)
-			{
-				return;
-			}
-
-			UhtModule module = PropertySettings.Outer.Module;
-			UhtIssueBehavior behavior = nonEngineBehavior;
-			if (module.IsPartOfEngine)
-			{
-				if (module.IsPlugin)
-				{
-					behavior = enginePluginBehavior;
-				}
-				else
-				{
-					behavior = engineBehavior;
-				}
-			}
-
-			if (behavior == UhtIssueBehavior.AllowSilently)
-			{
-				return;
-			}
-
-			string type = TokenReader.GetStringView(typeStartPos, TokenReader.InputPos - typeStartPos).ToString();
-			type = type.Replace("\n", "\\n", StringComparison.Ordinal);
-			type = type.Replace("\r", "\\r", StringComparison.Ordinal);
-			type = type.Replace("\t", "\\t", StringComparison.Ordinal);
-
-			switch (behavior)
-			{
-				case UhtIssueBehavior.Disallow:
-					if (!String.IsNullOrEmpty(alternativeTypeDesc))
-					{
-						TokenReader.LogError($"{pointerTypeDesc} usage in member declaration detected [[[{type}]]].  This is disallowed for the target/module, consider {alternativeTypeDesc} as an alternative.");
-					}
-					else
-					{
-						TokenReader.LogError($"{pointerTypeDesc} usage in member declaration detected [[[{type}]]].");
-					}
-					break;
-
-				case UhtIssueBehavior.AllowAndLog:
-					if (!String.IsNullOrEmpty(alternativeTypeDesc))
-					{
-						TokenReader.LogInfo($"{pointerTypeDesc} usage in member declaration detected [[[{type}]]].  Consider {alternativeTypeDesc} as an alternative.");
-					}
-					else
-					{
-						TokenReader.LogInfo("{PointerTypeDesc} usage in member declaration detected [[[{Type}]]].");
-					}
-					break;
-
-				default:
-					throw new UhtIceException("Unknown enum value");
-			}
-		}
-	}
-
 	/// <summary>
 	/// Delegate invoked to resolve a tokenized type into a UHTProperty type
 	/// </summary>
-	/// <returns></returns>
-	public delegate UhtProperty? UhtResolvePropertyDelegate(UhtPropertyResolveArgs Args);
+	/// <param name="args">Collection of arguments for property resolution</param>
+	/// <returns>Created property or null if the resolution must be delayed until after symbol table has been constructed</returns>
+	public delegate UhtProperty? UhtResolvePropertyDelegate(UhtPropertyResolveArgs args);
 
 	/// <summary>
 	/// Property type attribute
diff --git a/Engine/Source/Programs/Shared/EpicGames.UHT/Tokenizer/UhtTokenBufferReader.cs b/Engine/Source/Programs/Shared/EpicGames.UHT/Tokenizer/UhtTokenBufferReader.cs
--- a/Engine/Source/Programs/Shared/EpicGames.UHT/Tokenizer/UhtTokenBufferReader.cs
+++ b/Engine/Source/Programs/Shared/EpicGames.UHT/Tokenizer/UhtTokenBufferReader.cs
@@ -160,6 +160,12 @@ namespace EpicGames.UHT.Tokenizer
 			return token;
 		}
 
+		/// <inheritdoc/>
+		public int GetReplayTokenIndex()
+		{
+			throw new NotImplementedException();
+		}
+
 		/// <inheritdoc/>
 		public StringView GetRawString(char terminator, UhtRawStringOptions options)
 		{
diff --git a/Engine/Source/Programs/Shared/EpicGames.UHT/Tokenizer/UhtTokenReader.cs b/Engine/Source/Programs/Shared/EpicGames.UHT/Tokenizer/UhtTokenReader.cs
--- a/Engine/Source/Programs/Shared/EpicGames.UHT/Tokenizer/UhtTokenReader.cs
+++ b/Engine/Source/Programs/Shared/EpicGames.UHT/Tokenizer/UhtTokenReader.cs
@@ -145,6 +145,12 @@ namespace EpicGames.UHT.Tokenizer
 		/// <returns></returns>
 		public UhtToken GetToken();
 
+		/// <summary>
+		/// Method only implemented by the replay reader that returns the current token index
+		/// </summary>
+		/// <returns>Current token index</returns>
+		public int GetReplayTokenIndex();
+
 		/// <summary>
 		/// Tests to see if the given token is the first token of a line
 		/// </summary>
diff --git a/Engine/Source/Programs/Shared/EpicGames.UHT/Tokenizer/UhtTokenReaderSkipExtensions.cs b/Engine/Source/Programs/Shared/EpicGames.UHT/Tokenizer/UhtTokenReaderSkipExtensions.cs
--- a/Engine/Source/Programs/Shared/EpicGames.UHT/Tokenizer/UhtTokenReaderSkipExtensions.cs
+++ b/Engine/Source/Programs/Shared/EpicGames.UHT/Tokenizer/UhtTokenReaderSkipExtensions.cs
@@ -1,5 +1,6 @@
 // Copyright Epic Games, Inc. All Rights Reserved.
 
+using System;
 using EpicGames.UHT.Utils;
 
 namespace EpicGames.UHT.Tokenizer
@@ -60,6 +61,55 @@ namespace EpicGames.UHT.Tokenizer
 			return tokenReader;
 		}
 
+		/// <summary>
+		/// Skip any block of tokens wrapped by the given token symbols
+		/// </summary>
+		/// <param name="tokenReader">Token reader</param>
+		/// <param name="initiator">Initiating token (i.e. &quot;(&quot;)</param>
+		/// <param name="terminator">Terminating token (i.e. &quot;)&quot;)</param>
+		/// <param name="initialNesting">If true, start with an initial nesting count of one (assume we already parsed an initiator)</param>
+		/// <param name="callback">Invoked for each matching token</param>
+		/// <param name="exceptionContext">Extra context for any errors</param>
+		/// <returns>Token reader</returns>
+		/// <exception cref="UhtTokenException">Throw if end of file is reached</exception>
+		public static IUhtTokenReader SkipBrackets(this IUhtTokenReader tokenReader, char initiator, char terminator, int initialNesting, Action<UhtToken> callback, object? exceptionContext = null)
+		{
+			int nesting = initialNesting;
+			if (nesting == 0)
+			{
+				ref UhtToken token = ref tokenReader.PeekToken();
+				if (token.IsSymbol(initiator))
+				{
+					callback(token);
+					tokenReader.ConsumeToken();
+				}
+				else
+				{
+					throw new UhtTokenException(tokenReader, token, initiator, exceptionContext);
+				}
+				++nesting;
+			}
+
+			do
+			{
+				UhtToken skipToken = tokenReader.GetToken();
+				if (skipToken.TokenType.IsEndType())
+				{
+					throw new UhtTokenException(tokenReader, skipToken, terminator, exceptionContext);
+				}
+				else if (skipToken.IsSymbol(initiator))
+				{
+					++nesting;
+				}
+				else if (skipToken.IsSymbol(terminator))
+				{
+					--nesting;
+				}
+				callback(skipToken);
+			} while (nesting != 0);
+			return tokenReader;
+		}
+
 		/// <summary>
 		/// Skip tokens until the given terminator is found.  The terminator will not be consumed.
 		/// </summary>
diff --git a/Engine/Source/Programs/Shared/EpicGames.UHT/Tokenizer/UhtTokenReplayReader.cs b/Engine/Source/Programs/Shared/EpicGames.UHT/Tokenizer/UhtTokenReplayReader.cs
--- a/Engine/Source/Programs/Shared/EpicGames.UHT/Tokenizer/UhtTokenReplayReader.cs
+++ b/Engine/Source/Programs/Shared/EpicGames.UHT/Tokenizer/UhtTokenReplayReader.cs
@@ -190,6 +190,12 @@ namespace EpicGames.UHT.Tokenizer
 			return token;
 		}
 
+		/// <inheritdoc/>
+		public int GetReplayTokenIndex()
+		{
+			return _hasToken ? _currentTokenIndex : _currentTokenIndex + 1;
+		}
+
 		/// <inheritdoc/>
 		public bool IsFirstTokenInLine(ref UhtToken token)
 		{
diff --git a/Engine/Source/Programs/Shared/EpicGames.UHT/Types/Properties/UhtArrayProperty.cs b/Engine/Source/Programs/Shared/EpicGames.UHT/Types/Properties/UhtArrayProperty.cs
--- a/Engine/Source/Programs/Shared/EpicGames.UHT/Types/Properties/UhtArrayProperty.cs
+++ b/Engine/Source/Programs/Shared/EpicGames.UHT/Types/Properties/UhtArrayProperty.cs
@@ -249,7 +249,7 @@ namespace EpicGames.UHT.Types
 			tokenReader.Require('<');
 
 			// Parse the value type
-			UhtProperty? value = UhtPropertyParser.ParseTemplateParam(args.ResolvePhase, propertySettings, propertySettings.SourceName, tokenReader);
+			UhtProperty? value = args.ParseTemplateParam(propertySettings.SourceName);
 			if (value == null)
 			{
 				return null;
diff --git a/Engine/Source/Programs/Shared/EpicGames.UHT/Types/Properties/UhtMapProperty.cs b/Engine/Source/Programs/Shared/EpicGames.UHT/Types/Properties/UhtMapProperty.cs
--- a/Engine/Source/Programs/Shared/EpicGames.UHT/Types/Properties/UhtMapProperty.cs
+++ b/Engine/Source/Programs/Shared/EpicGames.UHT/Types/Properties/UhtMapProperty.cs
@@ -296,7 +296,7 @@ namespace EpicGames.UHT.Types
 			tokenReader.Require('<');
 
 			// Parse the key type
-			UhtProperty? key = UhtPropertyParser.ParseTemplateParam(args.ResolvePhase, propertySettings, "Key", tokenReader);
+			UhtProperty? key = args.ParseTemplateParam("Key");
 			if (key == null)
 			{
 				return null;
@@ -315,7 +315,7 @@ namespace EpicGames.UHT.Types
 			tokenReader.Require(',');
 
 			// Parse the value type
-			UhtProperty? value = UhtPropertyParser.ParseTemplateParam(args.ResolvePhase, propertySettings, "Value", tokenReader);
+			UhtProperty? value = args.ParseTemplateParam("Value");
 			if (value == null)
 			{
 				return null;
diff --git a/Engine/Source/Programs/Shared/EpicGames.UHT/Types/Properties/UhtOptionalProperty.cs b/Engine/Source/Programs/Shared/EpicGames.UHT/Types/Properties/UhtOptionalProperty.cs
--- a/Engine/Source/Programs/Shared/EpicGames.UHT/Types/Properties/UhtOptionalProperty.cs
+++ b/Engine/Source/Programs/Shared/EpicGames.UHT/Types/Properties/UhtOptionalProperty.cs
@@ -245,7 +245,7 @@ namespace EpicGames.UHT.Types
 			tokenReader.Require('<');
 
 			// Parse the value type
-			UhtProperty? value = UhtPropertyParser.ParseTemplateParam(args.ResolvePhase, propertySettings, propertySettings.SourceName, tokenReader);
+			UhtProperty? value = args.ParseTemplateParam(propertySettings.SourceName);
 			if (value == null)
 			{
 				return null;
diff --git a/Engine/Source/Programs/Shared/EpicGames.UHT/Types/Properties/UhtSetProperty.cs b/Engine/Source/Programs/Shared/EpicGames.UHT/Types/Properties/UhtSetProperty.cs
--- a/Engine/Source/Programs/Shared/EpicGames.UHT/Types/Properties/UhtSetProperty.cs
+++ b/Engine/Source/Programs/Shared/EpicGames.UHT/Types/Properties/UhtSetProperty.cs
@@ -239,7 +239,7 @@ namespace EpicGames.UHT.Types
 			tokenReader.Require('<');
 
 			// Parse the value type
-			UhtProperty? value = UhtPropertyParser.ParseTemplateParam(args.ResolvePhase, propertySettings, "Value", tokenReader);
+			UhtProperty? value = args.ParseTemplateParam("Value");
 			if (value == null)
 			{
 				return null;
diff --git a/Engine/Source/Programs/Shared/EpicGames.UHT/Types/UhtProperty.cs b/Engine/Source/Programs/Shared/EpicGames.UHT/Types/UhtProperty.cs
--- a/Engine/Source/Programs/Shared/EpicGames.UHT/Types/UhtProperty.cs
+++ b/Engine/Source/Programs/Shared/EpicGames.UHT/Types/UhtProperty.cs
@@ -2,6 +2,7 @@
 
 using System;
 using System.Collections.Generic;
+using System.Diagnostics;
 using System.Diagnostics.CodeAnalysis;
 using System.Globalization;
 using System.Linq;
@@ -642,6 +643,7 @@ namespace EpicGames.UHT.Types
 	/// <summary>
 	/// Property settings is a transient object used during the parsing of properties
 	/// </summary>
+	[DebuggerDisplay("Type = {TypeTokens}")]
 	public class UhtPropertySettings
 	{
 
@@ -762,6 +764,12 @@ namespace EpicGames.UHT.Types
 		[SuppressMessage("Usage", "CA2227:Collection properties should be read only", Justification = "")]
 		public List<UhtToken>? DefaultValueTokens { get; set; }
 
+		/// <summary>
+		/// Collection of the tokens used to declare the property type.  The type tokens are only valid for
+		/// the outermost property.
+		/// </summary>
+		public UhtTypeTokensRef TypeTokens { get; set; } = new();
+
 		/// <summary>
 		/// If true, the property is a bit field
 		/// </summary>
@@ -805,6 +813,7 @@ namespace EpicGames.UHT.Types
 			Setter = null;
 			Getter = null;
 			IsBitfield = false;
+			TypeTokens = parentPropertySettings.TypeTokens;
 		}
 
 		/// <summary>
@@ -836,6 +845,7 @@ namespace EpicGames.UHT.Types
 			Setter = null;
 			Getter = null;
 			IsBitfield = false;
+			TypeTokens = new();
 		}
 
 		/// <summary>
@@ -871,6 +881,7 @@ namespace EpicGames.UHT.Types
 			Setter = property.Setter;
 			Getter = property.Getter;
 			IsBitfield = property.IsBitfield;
+			TypeTokens = property.TypeTokens;
 		}
 	}
 
@@ -879,6 +890,7 @@ namespace EpicGames.UHT.Types
 	/// </summary>
 	[UhtEngineClass(Name = "Field", IsProperty = true)] // This is here just so it is defined
 	[UhtEngineClass(Name = "Property", IsProperty = true)]
+	[DebuggerDisplay("Name = {SourceName}, Type = {TypeTokens}")]
 	public abstract class UhtProperty : UhtType
 	{
 		#region Constants
@@ -1004,6 +1016,12 @@ namespace EpicGames.UHT.Types
 		[JsonConverter(typeof(JsonStringEnumConverter))]
 		public UhtPropertyAllocator Allocator { get; set; } = UhtPropertyAllocator.Default;
 
+		/// <summary>
+		/// Collection of tokens that declare the type of the property.
+		/// </summary>
+		[JsonIgnore]
+		public UhtTypeTokensRef TypeTokens { get; init; }
+
 		/// <summary>
 		/// Replication notify name
 		/// </summary>
@@ -1141,6 +1159,7 @@ namespace EpicGames.UHT.Types
 			Setter = propertySettings.Setter;
 			IsBitfield = propertySettings.IsBitfield;
 			PropertyCaps = UhtPropertyCaps.CanBeContainerValue | UhtPropertyCaps.CanBeContainerKey | UhtPropertyCaps.CanHaveConfig;
+			TypeTokens = propertySettings.TypeTokens;
 		}
 
 		#region Text and code generation support
diff --git a/Engine/Source/Programs/Shared/EpicGames.UHT/Types/UhtType.cs b/Engine/Source/Programs/Shared/EpicGames.UHT/Types/UhtType.cs
--- a/Engine/Source/Programs/Shared/EpicGames.UHT/Types/UhtType.cs
+++ b/Engine/Source/Programs/Shared/EpicGames.UHT/Types/UhtType.cs
@@ -2,6 +2,7 @@
 
 using System;
 using System.Collections.Generic;
+using System.Diagnostics;
 using System.Diagnostics.CodeAnalysis;
 using System.Text;
 using System.Text.Json.Serialization;
@@ -708,6 +709,7 @@ namespace EpicGames.UHT.Types
 	/// Base type for all UHT types.
 	/// These mostly map to UnrealEngine types.
 	/// </summary>
+	[DebuggerDisplay("SourceName = {SourceName}")]
 	public abstract class UhtType : IUhtMessageSite, IUhtMessageLineNumber
 	{
 		/// <summary>
diff --git a/Engine/Source/Programs/Shared/EpicGames.UHT/Utils/UhtPropertyResolveArgs.cs b/Engine/Source/Programs/Shared/EpicGames.UHT/Utils/UhtPropertyResolveArgs.cs
new file mode 100644
--- /dev/null
+++ b/Engine/Source/Programs/Shared/EpicGames.UHT/Utils/UhtPropertyResolveArgs.cs
@@ -0,0 +1,343 @@
+// Copyright Epic Games, Inc. All Rights Reserved.
+
+using System;
+using System.Diagnostics;
+using EpicGames.Core;
+using EpicGames.UHT.Parsers;
+using EpicGames.UHT.Tables;
+using EpicGames.UHT.Tokenizer;
+using EpicGames.UHT.Types;
+
+namespace EpicGames.UHT.Utils
+{
+	/// <summary>
+	/// The phase of UHT where the property is being resolved
+	/// </summary>
+	public enum UhtPropertyResolvePhase
+	{
+
+		/// <summary>
+		/// Resolved during the source processing phase.  Immediate property types only.
+		/// </summary>
+		Parsing,
+
+		/// <summary>
+		/// Resolved during the resolve phase.  Non-immediate property types only.
+		/// </summary>
+		Resolving,
+	}
+
+	/// <summary>
+	/// Structure defining the arguments used in the delegate to resolve property types
+	/// </summary>
+	/// <param name="ResolvePhase">Specifies if this is being resolved during the parsing phase or the resolution phase.  Type lookups can not happen during the parsing phase</param>
+	/// <param name="PropertySettings">The configuration of the property</param>
+	/// <param name="TokenReader">The token reader containing the type</param>
+	public record struct UhtPropertyResolveArgs(UhtPropertyResolvePhase ResolvePhase, UhtPropertySettings PropertySettings, IUhtTokenReader TokenReader)
+	{
+		private UhtToken _matchedToken = new();
+
+		/// <summary>
+		/// Currently running UHT session
+		/// </summary>
+		public readonly UhtSession Session => PropertySettings.Outer.Session;
+
+		/// <summary>
+		/// UHT configuration
+		/// </summary>
+		public readonly IUhtConfig Config => Session.Config!;
+
+		/// <summary>
+		/// 
+		/// </summary>
+		public readonly UhtTypeTokensRef TypeTokens => PropertySettings.TypeTokens;
+
+		/// <summary>
+		/// If true, a template argument is being resolved
+		/// </summary>
+		public readonly bool IsTemplate => TypeTokens.SegmentIndex != 0;
+
+		/// <summary>
+		/// When processing type, make sure that the next token is the expected token
+		/// </summary>
+		/// <returns>true if there could be more header to process, false if the end was reached.</returns>
+		public readonly bool SkipExpectedType()
+		{
+			if (TokenReader.TryOptional(_matchedToken.Value))
+			{
+				return true;
+			}
+			if (PropertySettings.PropertyCategory == UhtPropertyCategory.Member && TokenReader.TryOptional("const"))
+			{
+				TokenReader.LogError("Const properties are not supported.");
+			}
+			else
+			{
+				TokenReader.LogError($"Inappropriate keyword '{TokenReader.PeekToken().Value}' on variable of type '{_matchedToken.Value}'");
+			}
+			return false;
+		}
+
+		/// <summary>
+		/// Resolve the given property.  This method will resolve any immediate property during the parsing phase or 
+		/// resolve any previously parsed property to the final version.
+		/// </summary>
+		/// <returns></returns>
+		public UhtProperty? ResolveProperty()
+		{
+
+			// Loop through the tokens until we find a known property type or the start of a template argument list
+			foreach (UhtToken token in TypeTokens.MyType.Span)
+			{
+				if (token.IsIdentifier())
+				{
+					UhtToken copy = token;
+					Config.RedirectTypeIdentifier(ref copy);
+					if (Session.TryGetPropertyType(copy.Value, out UhtPropertyType propertyType))
+					{
+						if (ResolvePhase == UhtPropertyResolvePhase.Resolving || propertyType.Options.HasAnyFlags(UhtPropertyTypeOptions.Immediate))
+						{
+							_matchedToken = token;
+							return UhtPropertyParser.ResolvePropertyType(this, propertyType);
+						}
+					}
+				}
+			}
+
+			// Try the default processor.  We only use the default processor when trying to resolve something post parsing phase.
+			if (ResolvePhase == UhtPropertyResolvePhase.Resolving)
+			{
+				_matchedToken = new();
+				return UhtPropertyParser.ResolvePropertyType(this, Session.DefaultPropertyType);
+			}
+			return null;
+		}
+
+		/// <summary>
+		/// Parse a template parameter
+		/// </summary>
+		/// <param name="paramName">Name of the template parameter</param>
+		/// <returns>Parsed property</returns>
+		public readonly UhtProperty? ParseTemplateParam(StringView paramName)
+		{
+			// From the current segment index, look forward for the template segment being parsed
+			int tokenIndex = TokenReader.GetReplayTokenIndex();
+			int templateSegmentIndex = TypeTokens.SegmentIndex + 1;
+			foreach (UhtTypeSegment segment in TypeTokens.Segments[(TypeTokens.SegmentIndex + 1)..].Span)
+			{
+				if (tokenIndex < segment.Start)
+				{
+					throw new UhtIceException("Error parsing template argument");
+				}
+				if (tokenIndex < segment.End)
+				{
+					break;
+				}
+				templateSegmentIndex++;
+			}
+			if (templateSegmentIndex == TypeTokens.Segments.Length)
+			{
+				throw new UhtIceException("Error parsing template argument");
+			}
+
+			using UhtThreadBorrower<UhtPropertySpecifierContext> borrower = new(true);
+			UhtPropertySpecifierContext specifierContext = borrower.Instance;
+			specifierContext.Type = PropertySettings.Outer;
+			specifierContext.TokenReader = TokenReader;
+			specifierContext.AccessSpecifier = UhtAccessSpecifier.None;
+			specifierContext.MessageSite = TokenReader;
+			specifierContext.PropertySettings.Reset(PropertySettings, paramName.ToString(), TokenReader);
+			specifierContext.PropertySettings.TypeTokens = new(TypeTokens, templateSegmentIndex);
+			specifierContext.MetaData = specifierContext.PropertySettings.MetaData;
+			specifierContext.MetaNameIndex = UhtMetaData.IndexNone;
+			specifierContext.SeenEditSpecifier = false;
+			specifierContext.SeenBlueprintWriteSpecifier = false;
+			specifierContext.SeenBlueprintReadOnlySpecifier = false;
+			specifierContext.SeenBlueprintGetterSpecifier = false;
+
+			// Pre-parse any UPARAM 
+			UhtPropertyParser.PreParseType(specifierContext, true);
+
+			// Parse the type
+			UhtPropertyResolveArgs templateArgs = new(ResolvePhase, specifierContext.PropertySettings, TokenReader);
+			return templateArgs.ResolveProperty();
+		}
+
+		/// <summary>
+		/// Parse a template type
+		/// </summary>
+		/// <param name="returnUInterface">If true, return the UInterface instead of the type listed</param>
+		/// <returns>Referenced class</returns>
+		public readonly UhtClass? ParseTemplateObject(bool returnUInterface)
+		{
+			if (TokenReader.TryOptional("const"))
+			{
+				PropertySettings.MetaData.Add(UhtNames.NativeConst, "");
+			}
+
+			if (!SkipExpectedType())
+			{
+				return null;
+			}
+
+			bool isNativeConstTemplateArg = false;
+			UhtToken identifier = new();
+			TokenReader
+				.Require('<')
+				.Optional("const", () => { isNativeConstTemplateArg = true; })
+				.Optional("class")
+				.RequireIdentifier((ref UhtToken token) => { identifier = token; })
+				.Optional("const", () => { isNativeConstTemplateArg = true; })
+				.Require('>');
+
+			if (isNativeConstTemplateArg)
+			{
+				PropertySettings.MetaData.Add(UhtNames.NativeConstTemplateArg, "");
+			}
+
+			Config.RedirectTypeIdentifier(ref identifier);
+			UhtClass? returnClass = PropertySettings.Outer.FindType(UhtFindOptions.SourceName | UhtFindOptions.Class, ref identifier, TokenReader) as UhtClass;
+			if (returnClass != null && returnClass.AlternateObject != null && returnUInterface)
+			{
+				returnClass = returnClass.AlternateObject as UhtClass;
+			}
+			return returnClass;
+		}
+
+		/// <summary>
+		/// Parse a template type
+		/// </summary>
+		/// <returns>Referenced class</returns>
+		public readonly UhtClass? ParseTemplateClass()
+		{
+			UhtToken identifier = new();
+
+			if (TokenReader.TryOptional("const"))
+			{
+				PropertySettings.MetaData.Add(UhtNames.NativeConst, "");
+			}
+
+			TokenReader.Optional("class");
+
+			if (!SkipExpectedType())
+			{
+				return null;
+			}
+
+			TokenReader
+				.Require('<')
+				.Optional("class")
+				.RequireIdentifier((ref UhtToken token) => { identifier = token; })
+				.Require('>');
+
+			Config.RedirectTypeIdentifier(ref identifier);
+			UhtClass? returnClass = PropertySettings.Outer.FindType(UhtFindOptions.SourceName | UhtFindOptions.Class, ref identifier, TokenReader) as UhtClass;
+			if (returnClass != null && returnClass.AlternateObject != null)
+			{
+				returnClass = returnClass.AlternateObject as UhtClass;
+			}
+			return returnClass;
+		}
+
+		/// <summary>
+		/// Parse a template type
+		/// </summary>
+		/// <returns>Referenced class</returns>
+		public readonly UhtScriptStruct? ParseTemplateScriptStruct()
+		{
+			UhtToken identifier = new();
+
+			if (TokenReader.TryOptional("const"))
+			{
+				PropertySettings.MetaData.Add(UhtNames.NativeConst, "");
+			}
+
+			TokenReader.Optional("struct");
+
+			if (!SkipExpectedType())
+			{
+				return null;
+			}
+
+			TokenReader
+				.Require('<')
+				.Optional("struct")
+				.RequireIdentifier((ref UhtToken token) => { identifier = token; })
+				.Require('>');
+
+			Config.RedirectTypeIdentifier(ref identifier);
+			return PropertySettings.Outer.FindType(UhtFindOptions.SourceName | UhtFindOptions.ScriptStruct, ref identifier, TokenReader) as UhtScriptStruct;
+		}
+
+		/// <summary>
+		/// Logs message for Object pointers to convert UObject* to TObjectPtr or the reverse
+		/// </summary>
+		/// <param name="engineBehavior">Expected behavior for engine types</param>
+		/// <param name="enginePluginBehavior">Expected behavior for engine plugin types</param>
+		/// <param name="nonEngineBehavior">Expected behavior for non-engine types</param>
+		/// <param name="pointerTypeDesc">Description of the pointer type</param>
+		/// <param name="typeStartPos">Starting character position of the type</param>
+		/// <param name="alternativeTypeDesc">Suggested alternate declaration</param>
+		/// <exception cref="UhtIceException">Thrown if the behavior type is unexpected</exception>
+		public readonly void ConditionalLogPointerUsage(UhtIssueBehavior engineBehavior, UhtIssueBehavior enginePluginBehavior,
+			UhtIssueBehavior nonEngineBehavior, string pointerTypeDesc, int typeStartPos, string? alternativeTypeDesc)
+		{
+			if (PropertySettings.PropertyCategory != UhtPropertyCategory.Member)
+			{
+				return;
+			}
+
+			UhtModule module = PropertySettings.Outer.Module;
+			UhtIssueBehavior behavior = nonEngineBehavior;
+			if (module.IsPartOfEngine)
+			{
+				if (module.IsPlugin)
+				{
+					behavior = enginePluginBehavior;
+				}
+				else
+				{
+					behavior = engineBehavior;
+				}
+			}
+
+			if (behavior == UhtIssueBehavior.AllowSilently)
+			{
+				return;
+			}
+
+			string type = TokenReader.GetStringView(typeStartPos, TokenReader.InputPos - typeStartPos).ToString();
+			type = type.Replace("\n", "\\n", StringComparison.Ordinal);
+			type = type.Replace("\r", "\\r", StringComparison.Ordinal);
+			type = type.Replace("\t", "\\t", StringComparison.Ordinal);
+
+			switch (behavior)
+			{
+				case UhtIssueBehavior.Disallow:
+					if (!String.IsNullOrEmpty(alternativeTypeDesc))
+					{
+						TokenReader.LogError($"{pointerTypeDesc} usage in member declaration detected [[[{type}]]].  This is disallowed for the target/module, consider {alternativeTypeDesc} as an alternative.");
+					}
+					else
+					{
+						TokenReader.LogError($"{pointerTypeDesc} usage in member declaration detected [[[{type}]]].");
+					}
+					break;
+
+				case UhtIssueBehavior.AllowAndLog:
+					if (!String.IsNullOrEmpty(alternativeTypeDesc))
+					{
+						TokenReader.LogInfo($"{pointerTypeDesc} usage in member declaration detected [[[{type}]]].  Consider {alternativeTypeDesc} as an alternative.");
+					}
+					else
+					{
+						TokenReader.LogInfo("{PointerTypeDesc} usage in member declaration detected [[[{Type}]]].");
+					}
+					break;
+
+				default:
+					throw new UhtIceException("Unknown enum value");
+			}
+		}
+	}
+}
diff --git a/Engine/Source/Programs/Shared/EpicGames.UHT/Utils/UhtTypeTokens.cs b/Engine/Source/Programs/Shared/EpicGames.UHT/Utils/UhtTypeTokens.cs
new file mode 100644
--- /dev/null
+++ b/Engine/Source/Programs/Shared/EpicGames.UHT/Utils/UhtTypeTokens.cs
@@ -0,0 +1,332 @@
+﻿// Copyright Epic Games, Inc. All Rights Reserved.
+
+using System;
+using System.Collections.Generic;
+using System.Diagnostics;
+using System.Text;
+using EpicGames.UHT.Tokenizer;
+using EpicGames.UHT.Types;
+
+namespace EpicGames.UHT.Utils
+{
+	/// <summary>
+	/// Represents the decoded tokens for part of a type definition
+	/// </summary>
+	public struct UhtTypeSegment
+	{
+
+		/// <summary>
+		/// Starting token index of this type
+		/// </summary>
+		public int Start { get; set; } = -1;
+
+		/// <summary>
+		/// Ending token index of this type (non-inclusive)
+		/// </summary>
+		public int End { get; set; } = -1;
+
+		/// <summary>
+		/// Starting token index of the type
+		/// </summary>
+		public int TypeStart { get; set; } = -1;
+
+		/// <summary>
+		/// Ending token index of the type
+		/// </summary>
+		public int TypeEnd { get; set; } = -1;
+
+		/// <summary>
+		/// Construct a new type segment
+		/// </summary>
+		public UhtTypeSegment()
+		{
+		}
+	}
+
+	/// <summary>
+	/// Information about a parsed type
+	/// </summary>
+	public struct UhtTypeTokens
+	{
+		private readonly Memory<UhtTypeSegment> _segments;
+
+		/// <summary>
+		/// Collection of tokens that describe the type
+		/// </summary>
+		public ReadOnlyMemory<UhtToken> AllTokens { get; internal set; }
+
+		/// <summary>
+		/// Collection of type segments
+		/// </summary>
+		public readonly ReadOnlyMemory<UhtTypeSegment> Segments => _segments;
+
+		private UhtTypeTokens(ReadOnlyMemory<UhtToken> typeTokens, Memory<UhtTypeSegment> segments)
+		{
+			AllTokens = typeTokens;
+			_segments = segments;
+		}
+
+		/// <summary>
+		/// Gather a type definition from the given token reader
+		/// </summary>
+		/// <param name="tokenReader">Token reader containing the type</param>
+		/// <returns></returns>
+		public static UhtTypeTokens Gather(IUhtTokenReader tokenReader)
+		{
+			List<UhtToken> typeTokens = [];
+			List<UhtTypeSegment> segments = [];
+
+			Gather(typeTokens, segments, tokenReader, 0);
+			return new UhtTypeTokens(typeTokens.ToArray().AsMemory(), segments.ToArray().AsMemory());
+		}
+
+		/// <summary>
+		/// Extract the trailing token from the type string since it is the property/function name
+		/// </summary>
+		/// <param name="messageSite">Message site for errors</param>
+		/// <param name="propertySettings">Property setting for what is being parsed</param>
+		/// <returns>Trailing identifier</returns>
+		/// <exception cref="UhtException"></exception>
+		public UhtToken ExtractTrailingIdentifier(IUhtMessageSite messageSite, UhtPropertySettings propertySettings)
+		{
+			// Extract the property name from the types
+			if (AllTokens.Length < 2 || !AllTokens.Span[^1].IsIdentifier())
+			{
+				throw new UhtException(messageSite, $"{propertySettings.PropertyCategory.GetHintText()}: Expected name");
+			}
+
+			Debug.Assert(_segments.Span[0].End == AllTokens.Length);
+			UhtToken nameToken = AllTokens.Span[^1];
+			AllTokens = AllTokens[..^1];
+			ref UhtTypeSegment segment = ref _segments.Span[0];
+			segment.End = AllTokens.Length;
+			segment.TypeEnd = Math.Min(segment.TypeEnd, segment.End);
+			return nameToken;
+		}
+
+		/// <inheritdoc/>
+		public override readonly string ToString()
+		{
+			return ToString(AllTokens.Span);
+		}
+
+		/// <summary>
+		/// Convert the token span to a string
+		/// </summary>
+		/// <param name="tokens">Tokens to convert</param>
+		/// <returns>Resulting string</returns>
+		public static string ToString(ReadOnlySpan<UhtToken> tokens)
+		{
+			StringBuilder builder = new StringBuilder();
+			if (tokens.IsEmpty)
+			{
+				builder.Append("(empty type)");
+			}
+			else
+			{
+				bool priorIsIdentifier = false;
+				foreach (UhtToken token in tokens)
+				{
+					bool isIdentifier = token.IsIdentifier();
+					if (priorIsIdentifier && token.IsIdentifier())
+					{
+						builder.Append(' ');
+					}
+					builder.Append(token.ToString());
+					if (token.IsSymbol(','))
+					{
+						builder.Append(' ');
+					}
+					priorIsIdentifier = isIdentifier;
+				}
+			}
+			return builder.ToString();
+		}
+
+
+
+		private static int Gather(List<UhtToken> typeTokens, List<UhtTypeSegment> segments, IUhtTokenReader tokenReader, int depth)
+		{
+			int argumentCount = 1;
+			int segmentIndex = segments.Count;
+			UhtTypeSegment segment = new();
+			segment.Start = typeTokens.Count;
+			segments.Add(segment);
+
+			while (true)
+			{
+				if (tokenReader.IsEOF)
+				{
+					throw new UhtTokenException(tokenReader, tokenReader.GetToken(), null, null);
+				}
+				ref UhtToken token = ref tokenReader.PeekToken();
+
+				char symbolChar = token.IsSymbol() && token.Value.Length == 1 ? token.Value.Span[0] : '\0';
+				switch (symbolChar)
+				{
+					case '(':
+					case ')':
+					case ';':
+					case '[':
+					case ':':
+					case '=':
+					case '{':
+					case ',':
+						if (depth == 0)
+						{
+							// Don't consume the token
+							segment.End = typeTokens.Count;
+							segments[segmentIndex] = segment;
+							return argumentCount;
+						}
+						else if (symbolChar == ',')
+						{
+							segment.End = typeTokens.Count;
+							segments[segmentIndex] = segment;
+
+							tokenReader.ConsumeToken();
+							typeTokens.Add(token);
+
+							segment = new();
+							segment.Start = typeTokens.Count;
+							segmentIndex = segments.Count;
+							segments.Add(segment);
+							argumentCount++;
+						}
+						else
+						{
+							tokenReader.ConsumeToken();
+							typeTokens.Add(token);
+						}
+						break;
+
+					case '>':
+						if (depth == 0)
+						{
+							throw new UhtTokenException(tokenReader, token, "',' or ')'");
+						}
+						else
+						{
+							segment.End = typeTokens.Count;
+							segments[segmentIndex] = segment;
+							tokenReader.ConsumeToken(); // go on ahead and add the closing '>' to the token stream but not included in current segment
+							typeTokens.Add(token);
+							return argumentCount;
+						}
+
+					case '<':
+						tokenReader.ConsumeToken();
+						typeTokens.Add(token);
+						Gather(typeTokens, segments, tokenReader, depth + 1);
+						break;
+
+					default:
+						tokenReader.ConsumeToken();
+						typeTokens.Add(token);
+						if (token.IsIdentifier("UPARAM"))
+						{
+							tokenReader.SkipBrackets('(', ')', 0, (UhtToken x) => typeTokens.Add(x));
+						}
+						else if (token.IsIdentifier("const"))
+						{
+						}
+						else if (token.IsSymbol("::") || token.IsIdentifier())
+						{
+							if (segment.TypeEnd == -1)
+							{
+								segment.TypeStart = typeTokens.Count - 1;
+								segment.TypeEnd = typeTokens.Count;
+							}
+							else if (segment.TypeEnd == typeTokens.Count - 1)
+							{
+								segment.TypeEnd = typeTokens.Count;
+							}
+						}
+						break;
+				}
+			}
+		}
+	}
+
+	/// <summary>
+	/// Reference to a token segment
+	/// </summary>
+	public readonly struct UhtTypeTokensRef
+	{
+		/// <summary>
+		/// Collection of tokens that describe the type
+		/// </summary>
+		public readonly ReadOnlyMemory<UhtToken> AllTokens { get; }
+
+		/// <summary>
+		/// Collection of type segments
+		/// </summary>
+		public readonly ReadOnlyMemory<UhtTypeSegment> Segments { get; }
+
+		/// <summary>
+		/// Index of the segment being referenced
+		/// </summary>
+		public readonly int SegmentIndex { get; }
+
+		/// <summary>
+		/// The type segment for this reference
+		/// </summary>
+		public readonly UhtTypeSegment Segment => Segments.Span[SegmentIndex];
+
+		/// <summary>
+		/// The collection of my types for this declaration
+		/// </summary>
+		public readonly ReadOnlyMemory<UhtToken> MyType
+		{
+			get
+			{
+				int typeStart = Segment.TypeStart;
+				int typeEnd = Segment.TypeEnd;
+				return typeStart >= 0 ? AllTokens[typeStart..typeEnd] : new();
+			}
+		}
+
+		/// <summary>
+		/// The collection of tokens that define this type
+		/// </summary>
+		public readonly ReadOnlyMemory<UhtToken> MyTokens
+		{
+			get
+			{
+				int start = Segment.Start;
+				int end = Segment.End;
+				return AllTokens[start..end];
+			}
+		}
+
+		/// <summary>
+		/// Construct a new instance of a type tokens
+		/// </summary>
+		/// <param name="typeTokens">Complete type tokens</param>
+		/// <param name="segmentIndex">Segment index</param>
+		public UhtTypeTokensRef(UhtTypeTokens typeTokens, int segmentIndex)
+		{
+			AllTokens = typeTokens.AllTokens;
+			Segments = typeTokens.Segments;
+			SegmentIndex = segmentIndex;
+		}
+
+		/// <summary>
+		/// Construct a new instance of a type tokens
+		/// </summary>
+		/// <param name="typeTokens">Complete type tokens</param>
+		/// <param name="segmentIndex">Segment index</param>
+		public UhtTypeTokensRef(UhtTypeTokensRef typeTokens, int segmentIndex)
+		{
+			AllTokens = typeTokens.AllTokens;
+			Segments = typeTokens.Segments;
+			SegmentIndex = segmentIndex;
+		}
+
+		/// <inheritdoc/>
+		public override readonly string ToString()
+		{
+			return UhtTypeTokens.ToString(MyTokens.Span);
+		}
+	}
+}
