TargetRules Unit tests & basic framework to assist with TargetRules mocking.

diff --git a/Engine/Source/Programs/Shared/UnrealBuildTool.Tests/CppCompileWarningsTests.cs b/Engine/Source/Programs/Shared/UnrealBuildTool.Tests/CppCompileWarningsTests.cs
--- a/Engine/Source/Programs/Shared/UnrealBuildTool.Tests/CppCompileWarningsTests.cs
+++ b/Engine/Source/Programs/Shared/UnrealBuildTool.Tests/CppCompileWarningsTests.cs
@@ -1,6 +1,14 @@
 // Copyright Epic Games, Inc. All Rights Reserved.
 
+using System;
+using System.Collections.Generic;
+using System.IO;
+using System.Linq;
+using EpicGames.Core;
+using Microsoft.Extensions.Logging;
 using Microsoft.VisualStudio.TestTools.UnitTesting;
+using UnrealBuildBase;
+using UnrealBuildTool.Tests.TestUtilities;
 
 namespace UnrealBuildTool.Tests
 {
@@ -18,28 +26,28 @@ namespace UnrealBuildTool.Tests
 		}
 
 		[DataTestMethod]
-		[DataRow(WarningLevel.Default, null, true, WarningLevel.Default)]					// Null applied, no overwrite
-		[DataRow(WarningLevel.Default, WarningLevel.Off, true, WarningLevel.Off)]			// Overwrite with Off
-		[DataRow(WarningLevel.Default, WarningLevel.Warning, true, WarningLevel.Warning)]	// Overwrite with Warning
-		[DataRow(WarningLevel.Default, WarningLevel.Error, true, WarningLevel.Error)]		// Overwrite with Error
-		[DataRow(WarningLevel.Default, WarningLevel.Default, true, WarningLevel.Default)]	// Applied is Default, no change
-
-		[DataRow(WarningLevel.Default, WarningLevel.Warning, false, WarningLevel.Warning)]	// Default, non-overwrite, change
-		[DataRow(WarningLevel.Default, WarningLevel.Error, false, WarningLevel.Error)]		// Default, non-overwrite, change
-
-		[DataRow(WarningLevel.Off, null, true, WarningLevel.Off)]							// Null applied, no change
-		[DataRow(WarningLevel.Off, WarningLevel.Default, true, WarningLevel.Off)]			// Applied Default, no change
-		[DataRow(WarningLevel.Off, WarningLevel.Warning, true, WarningLevel.Warning)]		// Overwrite with Warning
-		[DataRow(WarningLevel.Off, WarningLevel.Error, true, WarningLevel.Error)]			// Overwrite with Error
-		[DataRow(WarningLevel.Off, WarningLevel.Warning, false, WarningLevel.Off)]			// No overwrite, no change
-
-		[DataRow(WarningLevel.Warning, WarningLevel.Error, true, WarningLevel.Error)]		// Overwrite with Error
-		[DataRow(WarningLevel.Warning, WarningLevel.Off, true, WarningLevel.Off)]			// Overwrite with Off
-		[DataRow(WarningLevel.Warning, WarningLevel.Off, false, WarningLevel.Warning)]		// No overwrite, remains Warning
-
-		[DataRow(WarningLevel.Error, WarningLevel.Off, true, WarningLevel.Off)]				// Overwrite with Off
-		[DataRow(WarningLevel.Error, WarningLevel.Warning, true, WarningLevel.Warning)]		// Overwrite with Warning
-		[DataRow(WarningLevel.Error, WarningLevel.Off, false, WarningLevel.Error)]			// No overwrite, remains Error
+		[DataRow(WarningLevel.Default, null, true, WarningLevel.Default)]                   // Null applied, no overwrite
+		[DataRow(WarningLevel.Default, WarningLevel.Off, true, WarningLevel.Off)]           // Overwrite with Off
+		[DataRow(WarningLevel.Default, WarningLevel.Warning, true, WarningLevel.Warning)]   // Overwrite with Warning
+		[DataRow(WarningLevel.Default, WarningLevel.Error, true, WarningLevel.Error)]       // Overwrite with Error
+		[DataRow(WarningLevel.Default, WarningLevel.Default, true, WarningLevel.Default)]   // Applied is Default, no change
+
+		[DataRow(WarningLevel.Default, WarningLevel.Warning, false, WarningLevel.Warning)]  // Default, non-overwrite, change
+		[DataRow(WarningLevel.Default, WarningLevel.Error, false, WarningLevel.Error)]      // Default, non-overwrite, change
+
+		[DataRow(WarningLevel.Off, null, true, WarningLevel.Off)]                           // Null applied, no change
+		[DataRow(WarningLevel.Off, WarningLevel.Default, true, WarningLevel.Off)]           // Applied Default, no change
+		[DataRow(WarningLevel.Off, WarningLevel.Warning, true, WarningLevel.Warning)]       // Overwrite with Warning
+		[DataRow(WarningLevel.Off, WarningLevel.Error, true, WarningLevel.Error)]           // Overwrite with Error
+		[DataRow(WarningLevel.Off, WarningLevel.Warning, false, WarningLevel.Off)]          // No overwrite, no change
+
+		[DataRow(WarningLevel.Warning, WarningLevel.Error, true, WarningLevel.Error)]       // Overwrite with Error
+		[DataRow(WarningLevel.Warning, WarningLevel.Off, true, WarningLevel.Off)]           // Overwrite with Off
+		[DataRow(WarningLevel.Warning, WarningLevel.Off, false, WarningLevel.Warning)]      // No overwrite, remains Warning
+
+		[DataRow(WarningLevel.Error, WarningLevel.Off, true, WarningLevel.Off)]             // Overwrite with Off
+		[DataRow(WarningLevel.Error, WarningLevel.Warning, true, WarningLevel.Warning)]     // Overwrite with Warning
+		[DataRow(WarningLevel.Error, WarningLevel.Off, false, WarningLevel.Error)]          // No overwrite, remains Error
 		public void TestApplyWarnings(WarningLevel inSourceWarningLevel, WarningLevel? appliedWarningLevel, bool overwriteSourceValue, WarningLevel expected)
 		{
 			WarningLevel sourceWarningLevel = inSourceWarningLevel;
@@ -65,7 +73,7 @@ namespace UnrealBuildTool.Tests
 			defaultBuildSettings.ShadowVariableWarningLevel = WarningLevel.Error;
 			defaultBuildSettings.UndefinedIdentifierWarningLevel = WarningLevel.Warning;
 			defaultBuildSettings.UnsafeTypeCastWarningLevel = WarningLevel.Error;
-			defaultBuildSettings.SwitchUnhandledEnumeratorWarningLevel= WarningLevel.Warning;
+			defaultBuildSettings.SwitchUnhandledEnumeratorWarningLevel = WarningLevel.Warning;
 
 			ReadOnlyCppCompileWarnings wrapper = new ReadOnlyCppCompileWarnings(defaultBuildSettings);
 
@@ -229,5 +237,156 @@ namespace UnrealBuildTool.Tests
 				Assert.AreEqual(WarningLevel.Warning, newWarnings.UndefinedIdentifierWarningLevel);
 			}
 		}
+
+		[TestMethod]
+		public void TestGenerateWarningCommandLineAgs()
+		{
+			DirectoryReference programsDirectory = DirectoryReference.Combine(Unreal.EngineSourceDirectory, "Programs");
+			FileReference testStubsDirectory = FileReference.Combine(programsDirectory, "Shared", "UnrealBuildTool.Tests", "UBT");
+
+			// Generate ephemeral TargetRules & uproject files
+			using (SafeTestDirectory testRootDirectory = SafeTestDirectory.CreateTestDirectory(Path.Combine(StaticTargetRulesInitializer.UBTTestFolderRoot, $"{StaticTargetRulesInitializer.UBTTestFolderPrefix}{Guid.NewGuid().ToString()}")))
+			using (SafeTestDirectory testSourceDirectory = SafeTestDirectory.CreateTestDirectory(Path.Combine(testRootDirectory.TemporaryDirectory, "Source")))
+			using (SafeTestFile testEditorTargetTestFile = new SafeTestFile(File.ReadAllText(Path.Combine(testStubsDirectory.FullName, "TargetStubs", "TestStubEditor.Target.ubttest")), "TestStubEditor.Target.cs", testSourceDirectory.TemporaryDirectory))
+			using (SafeTestFile testTargetTestFile = new SafeTestFile(File.ReadAllText(Path.Combine(testStubsDirectory.FullName, "TargetStubs", "TestStub.Target.ubttest")), "TestStub.Target.cs", testSourceDirectory.TemporaryDirectory))
+			using (SafeTestFile testUprojectTestFile = new SafeTestFile(File.ReadAllText(Path.Combine(testStubsDirectory.FullName, "UprojectStubs", "TestStub.uproject.ubttest")), "TestStub.uproject", testRootDirectory.TemporaryDirectory))
+			{
+				CppCompileEnvironment newCompileEnvironment = new CppCompileEnvironment(UnrealTargetPlatform.Win64, CppConfiguration.Development, new UnrealArchitectures(UnrealArch.X64), null);
+
+				FileReference testUproject = new FileReference(testUprojectTestFile.TemporaryFile);
+
+				// Generate rules assembly & subsequent targets
+				RulesAssembly generatedTestRulesAssembly = RulesCompiler.CreateTargetRulesAssembly(testUproject, "TestStubEditor", false, false, false, null, false, new TestLogger());
+				TargetRules testTargetRules = generatedTestRulesAssembly.CreateTargetRules("TestStubEditor", UnrealTargetPlatform.Win64, UnrealTargetConfiguration.Development, null, testUproject, null, new TestLogger());
+				CppCompileWarnings baseWarnings = new CppCompileWarnings(testTargetRules, null);
+
+				IEnumerable<string> args = null;
+
+				// No args, base check
+				{
+					args = baseWarnings.GenerateWarningCommandLineArgs(newCompileEnvironment, typeof(VCToolChain));
+					Assert.IsFalse(args.Any());
+				}
+
+				// VC toolchain check;
+				{
+					baseWarnings.UnsafeTypeCastWarningLevel = WarningLevel.Warning;
+					args = baseWarnings.GenerateWarningCommandLineArgs(newCompileEnvironment, typeof(VCToolChain));
+					Assert.IsTrue(args.Count() == 2);
+					Assert.IsTrue(args.Contains("/w44244"));
+					Assert.IsTrue(args.Contains("/w44838"));
+
+					baseWarnings.UnsafeTypeCastWarningLevel = WarningLevel.Default;
+				}
+
+				// MSVC compiler check;
+				{
+					baseWarnings.ShadowVariableWarningLevel = WarningLevel.Error;
+					args = baseWarnings.GenerateWarningCommandLineArgs(newCompileEnvironment, typeof(VCToolChain));
+
+					Assert.IsTrue(args.Count() == 3);
+					Assert.IsTrue(args.Contains("/we4456"));
+					Assert.IsTrue(args.Contains("/we4458"));
+					Assert.IsTrue(args.Contains("/we4459"));
+
+					WindowsCompiler previousCompiler = testTargetRules.WindowsPlatform.Compiler;
+					testTargetRules.WindowsPlatform.Compiler = WindowsCompiler.Intel;
+
+					args = baseWarnings.GenerateWarningCommandLineArgs(newCompileEnvironment, typeof(VCToolChain));
+
+					Assert.IsFalse(args.Any());
+
+					testTargetRules.WindowsPlatform.Compiler = previousCompiler;
+					baseWarnings.ShadowVariableWarningLevel = WarningLevel.Default;
+				}
+
+				// Specialized VC toolchain check; UndefinedIdentifier
+				{
+					baseWarnings.UndefinedIdentifierWarningLevel = WarningLevel.Warning;
+					args = baseWarnings.GenerateWarningCommandLineArgs(newCompileEnvironment, typeof(VCToolChain));
+
+					Assert.IsTrue(args.Count() == 1);
+					Assert.IsTrue(args.Contains("/w44668"));
+
+
+					CppCompileEnvironment preprocessOnlyCompileEnvironment = new CppCompileEnvironment(newCompileEnvironment);
+					preprocessOnlyCompileEnvironment.bPreprocessOnly = true;
+					args = baseWarnings.GenerateWarningCommandLineArgs(preprocessOnlyCompileEnvironment, typeof(VCToolChain));
+
+					Assert.IsFalse(args.Any());
+
+					baseWarnings.UndefinedIdentifierWarningLevel = WarningLevel.Default;
+				}
+
+				// Specialized VC toolchain check;  UnsafeTypeCast
+				{
+					RulesAssembly generatedTestGameRulesAssembly = RulesCompiler.CreateTargetRulesAssembly(testUproject, "TestStub", false, false, false, null, false, new TestLogger());
+					TargetRules testTargetRulesNonWin64 = generatedTestGameRulesAssembly.CreateTargetRules("TestStub", UnrealTargetPlatform.Android, UnrealTargetConfiguration.Development, null, testUproject, null, new TestLogger());
+					CppCompileWarnings specializedWarnings = new CppCompileWarnings(testTargetRulesNonWin64, null);
+
+					specializedWarnings.UnsafeTypeCastWarningLevel = WarningLevel.Warning;
+					args = specializedWarnings.GenerateWarningCommandLineArgs(newCompileEnvironment, typeof(VCToolChain));
+
+					Assert.IsTrue(args.Count() == 2);
+					Assert.IsTrue(args.Contains("/wd4244"));
+					Assert.IsTrue(args.Contains("/wd4838"));
+
+					baseWarnings.UnsafeTypeCastWarningLevel = WarningLevel.Warning;
+					args = baseWarnings.GenerateWarningCommandLineArgs(newCompileEnvironment, typeof(VCToolChain));
+
+					Assert.IsTrue(args.Count() == 2);
+					Assert.IsTrue(args.Contains("/w44244"));
+					Assert.IsTrue(args.Contains("/w44838"));
+
+					baseWarnings.UnsafeTypeCastWarningLevel = WarningLevel.Default;
+				}
+
+				// Clang toolchain version check - no version; should disable Shadow
+				{
+					testTargetRules.WindowsPlatform.Compiler = WindowsCompiler.Clang;
+					baseWarnings.ShadowVariableWarningLevel = WarningLevel.Error;
+
+					args = baseWarnings.GenerateWarningCommandLineArgs(newCompileEnvironment, typeof(ClangToolChain));
+					Assert.IsTrue(args.Count() == 1);
+					Assert.IsTrue(args.Contains("-Wno-shadow"));
+
+					baseWarnings.ShadowVariableWarningLevel = WarningLevel.Default;
+				}
+
+				// Clang toolchain version check - valid versions;
+				{
+					testTargetRules.WindowsPlatform.Compiler = WindowsCompiler.Clang;
+					baseWarnings.ShadowVariableWarningLevel = WarningLevel.Error;
+
+					VersionNumber newVersionNumber = new VersionNumber(16, 9, 5);
+
+					args = baseWarnings.GenerateWarningCommandLineArgs(newCompileEnvironment, typeof(ClangToolChain), newVersionNumber);
+					Assert.IsTrue(args.Count() == 1);
+					Assert.IsTrue(args.Contains("-Wshadow"));
+
+					newVersionNumber = new VersionNumber(17);
+					args = baseWarnings.GenerateWarningCommandLineArgs(newCompileEnvironment, typeof(ClangToolChain), newVersionNumber);
+					Assert.IsTrue(args.Count() == 1);
+					Assert.IsTrue(args.Contains("-Wno-shadow"));
+
+					newVersionNumber = new VersionNumber(17, 1);
+					args = baseWarnings.GenerateWarningCommandLineArgs(newCompileEnvironment, typeof(ClangToolChain), newVersionNumber);
+					Assert.IsTrue(args.Count() == 1);
+					Assert.IsTrue(args.Contains("-Wno-shadow"));
+
+					newVersionNumber = new VersionNumber(18, 1, 2);
+					args = baseWarnings.GenerateWarningCommandLineArgs(newCompileEnvironment, typeof(ClangToolChain), newVersionNumber);
+					Assert.IsTrue(args.Count() == 1);
+					Assert.IsTrue(args.Contains("-Wno-shadow"));
+
+					newVersionNumber = new VersionNumber(18, 1, 3);
+					args = baseWarnings.GenerateWarningCommandLineArgs(newCompileEnvironment, typeof(ClangToolChain), newVersionNumber);
+					Assert.IsTrue(args.Count() == 1);
+					Assert.IsTrue(args.Contains("-Wshadow"));
+
+					baseWarnings.ShadowVariableWarningLevel = WarningLevel.Default;
+				}
+			}
+		}
 	}
 }
\ No newline at end of file
diff --git a/Engine/Source/Programs/Shared/UnrealBuildTool.Tests/TargetRulesTests.cs b/Engine/Source/Programs/Shared/UnrealBuildTool.Tests/TargetRulesTests.cs
new file mode 100644
--- /dev/null
+++ b/Engine/Source/Programs/Shared/UnrealBuildTool.Tests/TargetRulesTests.cs
@@ -0,0 +1,67 @@
+// Copyright Epic Games, Inc. All Rights Reserved.
+
+using System;
+using System.Collections.Generic;
+using System.IO;
+using EpicGames.Core;
+using Microsoft.VisualStudio.TestTools.UnitTesting;
+using UnrealBuildTool.Tests.TestUtilities;
+using UnrealBuildBase;
+
+namespace UnrealBuildTool.Tests
+{
+	[TestClass]
+	public class TargetRulesTests
+	{
+		[TestMethod]
+		public void TestTargetRequiresUniqueEnvironmentFirstOrder()
+		{
+			DirectoryReference programsDirectory = DirectoryReference.Combine(Unreal.EngineSourceDirectory, "Programs");
+			DirectoryReference testStubsDirectory = DirectoryReference.Combine(programsDirectory, "Shared", "UnrealBuildTool.Tests", "UBT");
+
+			// Generate ephemeral TargetRules & uproject files
+			using (SafeTestDirectory testRootDirectory = SafeTestDirectory.CreateTestDirectory(Path.Combine(StaticTargetRulesInitializer.UBTTestFolderRoot, $"{StaticTargetRulesInitializer.UBTTestFolderPrefix}{Guid.NewGuid().ToString()}")))
+			using (SafeTestDirectory testSourceDirectory = SafeTestDirectory.CreateTestDirectory(Path.Combine(testRootDirectory.TemporaryDirectory, "Source")))
+			using (SafeTestFile testEditorTargetTestFile = new SafeTestFile(File.ReadAllText(Path.Combine(testStubsDirectory.FullName, "TargetStubs", "TestStubEditor.Target.ubttest")), "TestStubEditor.Target.cs", testSourceDirectory.TemporaryDirectory))
+			using (SafeTestFile testUprojectTestFile = new SafeTestFile(File.ReadAllText(Path.Combine(testStubsDirectory.FullName, "UprojectStubs", "TestStub.uproject.ubttest")), "TestStub.uproject", testRootDirectory.TemporaryDirectory))
+			{
+				FileReference testUproject = new FileReference(testUprojectTestFile.TemporaryFile);
+
+				// Generate rules assembly & subsequent targets
+				RulesAssembly generatedTestRulesAssembly = RulesCompiler.CreateTargetRulesAssembly(testUproject, "TestStubEditor", false, false, false, null, false, new TestLogger());
+				TargetRules testTargetRules = generatedTestRulesAssembly.CreateTargetRules("TestStubEditor", UnrealTargetPlatform.Win64, UnrealTargetConfiguration.Development, null, testUproject, null, new TestLogger());
+				string? baseTarget = null;
+
+				// Do not mutate the object
+				{
+					bool result = testTargetRules.RequiresUniqueEnvironment(generatedTestRulesAssembly, null, new Dictionary<string, (string?, string?)>(), out baseTarget);
+
+					Assert.IsFalse(result);
+				}
+
+				// Mutate the object against a RequiresUniqueEnvironment
+				{
+					WarningLevel preWarningLevel = testTargetRules.CppCompileWarningSettings.ShadowVariableWarningLevel;
+					testTargetRules.CppCompileWarningSettings.ShadowVariableWarningLevel = WarningLevel.Warning;
+					bool result = testTargetRules.RequiresUniqueEnvironment(generatedTestRulesAssembly, null, new Dictionary<string, (string?, string?)>(), out baseTarget);
+
+					Assert.IsTrue(result);
+
+					testTargetRules.CppCompileWarningSettings.ShadowVariableWarningLevel = preWarningLevel;
+				}
+
+				// Mutate the object against a RequiresUniqueEnvironment, and set bOverrideBuildEnvironment=true
+				{
+					WarningLevel preWarningLevel = testTargetRules.CppCompileWarningSettings.ShadowVariableWarningLevel;
+					testTargetRules.CppCompileWarningSettings.ShadowVariableWarningLevel = WarningLevel.Warning;
+					testTargetRules.bOverrideBuildEnvironment = true;
+					bool result = testTargetRules.RequiresUniqueEnvironment(generatedTestRulesAssembly, null, new Dictionary<string, (string?, string?)>(), out baseTarget);
+
+					Assert.IsFalse(result);
+
+					testTargetRules.CppCompileWarningSettings.ShadowVariableWarningLevel = preWarningLevel;
+				}
+			}
+		}
+	}
+}
\ No newline at end of file
diff --git a/Engine/Source/Programs/Shared/UnrealBuildTool.Tests/TestUtilities/SafeTestDirectory.cs b/Engine/Source/Programs/Shared/UnrealBuildTool.Tests/TestUtilities/SafeTestDirectory.cs
new file mode 100644
--- /dev/null
+++ b/Engine/Source/Programs/Shared/UnrealBuildTool.Tests/TestUtilities/SafeTestDirectory.cs
@@ -0,0 +1,46 @@
+// Copyright Epic Games, Inc. All Rights Reserved.
+
+using System;
+using System.IO;
+
+namespace UnrealBuildTool.Tests.TestUtilities
+{
+	internal class SafeTestDirectory : IDisposable
+	{
+		internal string TemporaryDirectory { get; private set; }
+		private bool _wasFolderCreated = false;
+
+		private SafeTestDirectory(string temporaryDirectory, bool wasFolderCreated = false)
+		{
+			TemporaryDirectory = temporaryDirectory;
+			_wasFolderCreated = wasFolderCreated;
+		}
+
+		internal static SafeTestDirectory CreateTestDirectory(string absoluteDirectoryPath)
+		{
+			SafeTestDirectory returnDirectory = new SafeTestDirectory(absoluteDirectoryPath);
+			if (!Directory.Exists(absoluteDirectoryPath))
+			{
+				Directory.CreateDirectory(absoluteDirectoryPath);
+				returnDirectory._wasFolderCreated = true;
+			}
+
+			return returnDirectory;
+		}
+
+		public void Dispose()
+		{
+			if (_wasFolderCreated && !String.IsNullOrEmpty(TemporaryDirectory) && Directory.Exists(TemporaryDirectory))
+			{
+				try
+				{
+					Directory.Delete(TemporaryDirectory, true);
+				}
+				catch (Exception)
+				{
+					// Swallow exception
+				}
+			}
+		}
+	}
+}
diff --git a/Engine/Source/Programs/Shared/UnrealBuildTool.Tests/TestUtilities/SafeTestFile.cs b/Engine/Source/Programs/Shared/UnrealBuildTool.Tests/TestUtilities/SafeTestFile.cs
new file mode 100644
--- /dev/null
+++ b/Engine/Source/Programs/Shared/UnrealBuildTool.Tests/TestUtilities/SafeTestFile.cs
@@ -0,0 +1,37 @@
+// Copyright Epic Games, Inc. All Rights Reserved.
+
+using System;
+using System.IO;
+
+namespace UnrealBuildTool.Tests.TestUtilities
+{
+	internal class SafeTestFile : IDisposable
+	{
+		internal string TemporaryFile { get; private set; }
+
+		internal SafeTestFile(string content, string outputFileName, string outputFolder)
+		{
+			try
+			{
+				TemporaryFile = Path.Combine(outputFolder, outputFileName);
+
+				using (StreamWriter writer = new StreamWriter(TemporaryFile))
+				{
+					writer.Write(content);
+				}
+			}
+			catch (Exception)
+			{
+				throw new InvalidOperationException("Unable to create test file. Test not feasible.");
+			}
+		}
+
+		public void Dispose()
+		{
+			if (!String.IsNullOrEmpty(TemporaryFile) && File.Exists(TemporaryFile))
+			{
+				File.Delete(TemporaryFile);
+			}
+		}
+	}
+}
diff --git a/Engine/Source/Programs/Shared/UnrealBuildTool.Tests/TestUtilities/StaticTargetRulesInitializer.cs b/Engine/Source/Programs/Shared/UnrealBuildTool.Tests/TestUtilities/StaticTargetRulesInitializer.cs
new file mode 100644
--- /dev/null
+++ b/Engine/Source/Programs/Shared/UnrealBuildTool.Tests/TestUtilities/StaticTargetRulesInitializer.cs
@@ -0,0 +1,64 @@
+// Copyright Epic Games, Inc. All Rights Reserved.
+
+using System;
+using System.IO;
+using Microsoft.Extensions.Logging;
+using Microsoft.VisualStudio.TestTools.UnitTesting;
+
+namespace UnrealBuildTool.Tests.TestUtilities
+{
+	[TestClass]
+	public class StaticTargetRulesInitializer
+	{
+		internal const string UBTTestFolderPrefix = "UBT-Test-";
+		internal static readonly string UBTTestFolderRoot = Path.GetTempPath();
+
+		[AssemblyInitialize]
+		public static void InitializeBuildEnvironment(TestContext context)
+		{
+			// Clean any previous runs linger temp folder
+			// Remarks: this is due to target rules modules dll being held by the test process; full cleanup not feasible on existing run
+			string[] subfolders = Directory.GetDirectories(UBTTestFolderRoot, $"{UBTTestFolderPrefix}*", SearchOption.TopDirectoryOnly);
+
+			foreach (string folder in subfolders)
+			{
+				if (Directory.Exists(folder))
+				{
+					try
+					{
+						Directory.Delete(folder, true);
+					}
+					catch (Exception) { }
+				}
+			}
+
+			// Example, on how to register a fully stubbed test environment.
+			/*
+			 	UEBuildPlatform.RegisterBuildPlatform(
+					new TestBuildPlatfrom(
+							UnrealTargetPlatform.Win64,
+							new TestSDK(logger),
+							new UnrealArchitectureConfig(UnrealArch.X64),
+							logger),
+					logger
+			);*/
+
+			ILogger logger = new TestLogger();
+
+			UEBuildPlatform? _;
+			if (!UEBuildPlatform.TryGetBuildPlatform(UnrealTargetPlatform.Win64, out _))
+			{
+				MicrosoftPlatformSDK SDK = new MicrosoftPlatformSDK(logger);
+				UEBuildPlatform.RegisterBuildPlatform(new WindowsPlatform(UnrealTargetPlatform.Win64, SDK, logger), logger);
+			}
+
+			if (!UEBuildPlatform.TryGetBuildPlatform(UnrealTargetPlatform.Android, out _))
+			{
+				AndroidPlatformSDK ASDK = new AndroidPlatformSDK(logger);
+				UEBuildPlatform.RegisterBuildPlatform(new AndroidPlatform(ASDK, logger), logger);
+			}
+
+			XmlConfig.ReadConfigFiles(null, null, logger);
+		}
+	}
+}
diff --git a/Engine/Source/Programs/Shared/UnrealBuildTool.Tests/TestUtilities/TestBuildPlatform.cs b/Engine/Source/Programs/Shared/UnrealBuildTool.Tests/TestUtilities/TestBuildPlatform.cs
new file mode 100644
--- /dev/null
+++ b/Engine/Source/Programs/Shared/UnrealBuildTool.Tests/TestUtilities/TestBuildPlatform.cs
@@ -0,0 +1,39 @@
+// Copyright Epic Games, Inc. All Rights Reserved.
+
+using EpicGames.Core;
+using Microsoft.Extensions.Logging;
+
+namespace UnrealBuildTool.Tests.TestUtilities
+{
+	internal class TestBuildPlatfrom : UEBuildPlatform
+	{
+		public TestBuildPlatfrom(UnrealTargetPlatform InPlatform, UEBuildPlatformSDK SDK, UnrealArchitectureConfig ArchitectureConfig, ILogger InLogger) : base(InPlatform, SDK, ArchitectureConfig, InLogger)
+		{
+
+		}
+		public override UEToolChain CreateToolChain(ReadOnlyTargetRules Target)
+		{
+			throw new System.NotImplementedException();
+		}
+
+		public override void Deploy(TargetReceipt Receipt)
+		{
+
+		}
+
+		public override bool IsBuildProduct(string FileName, string[] NamePrefixes, string[] NameSuffixes)
+		{
+			return false;
+		}
+
+		public override void SetUpEnvironment(ReadOnlyTargetRules Target, CppCompileEnvironment CompileEnvironment, LinkEnvironment LinkEnvironment)
+		{
+
+		}
+
+		public override bool ShouldCreateDebugInfo(ReadOnlyTargetRules Target)
+		{
+			return false;
+		}
+	}
+}
diff --git a/Engine/Source/Programs/Shared/UnrealBuildTool.Tests/TestUtilities/TestBuildPlatformFactory.cs b/Engine/Source/Programs/Shared/UnrealBuildTool.Tests/TestUtilities/TestBuildPlatformFactory.cs
new file mode 100644
--- /dev/null
+++ b/Engine/Source/Programs/Shared/UnrealBuildTool.Tests/TestUtilities/TestBuildPlatformFactory.cs
@@ -0,0 +1,15 @@
+// Copyright Epic Games, Inc. All Rights Reserved.
+
+using Microsoft.Extensions.Logging;
+
+namespace UnrealBuildTool.Tests.TestUtilities
+{
+	internal class TestBuildPlatformFactory : UEBuildPlatformFactory
+	{
+		public override UnrealTargetPlatform TargetPlatform => UnrealTargetPlatform.Win64;
+
+		public override void RegisterBuildPlatforms(ILogger Logger)
+		{
+		}
+	}
+}
diff --git a/Engine/Source/Programs/Shared/UnrealBuildTool.Tests/TestUtilities/TestBuildPlatformSDK.cs b/Engine/Source/Programs/Shared/UnrealBuildTool.Tests/TestUtilities/TestBuildPlatformSDK.cs
new file mode 100644
--- /dev/null
+++ b/Engine/Source/Programs/Shared/UnrealBuildTool.Tests/TestUtilities/TestBuildPlatformSDK.cs
@@ -0,0 +1,25 @@
+// Copyright Epic Games, Inc. All Rights Reserved.
+
+using EpicGames.Core;
+using Microsoft.Extensions.Logging;
+
+namespace UnrealBuildTool.Tests.TestUtilities
+{
+	internal class TestSDK : UEBuildPlatformSDK
+	{
+		public TestSDK(ILogger InLogger) : base(InLogger)
+		{
+		}
+
+		public override bool TryConvertVersionToInt(string? StringValue, out ulong OutValue, string? Hint = null)
+		{
+			OutValue = 1;
+			return true;
+		}
+
+		protected override string? GetInstalledSDKVersion()
+		{
+			return "1.0.0";
+		}
+	}
+}
diff --git a/Engine/Source/Programs/Shared/UnrealBuildTool.Tests/TestUtilities/TestLogger.cs b/Engine/Source/Programs/Shared/UnrealBuildTool.Tests/TestUtilities/TestLogger.cs
new file mode 100644
--- /dev/null
+++ b/Engine/Source/Programs/Shared/UnrealBuildTool.Tests/TestUtilities/TestLogger.cs
@@ -0,0 +1,24 @@
+// Copyright Epic Games, Inc. All Rights Reserved.
+
+using System;
+using Microsoft.Extensions.Logging;
+
+namespace UnrealBuildTool.Tests.TestUtilities
+{
+	internal class TestLogger : ILogger
+	{
+		public IDisposable? BeginScope<TState>(TState state) where TState : notnull
+		{
+			return null;
+		}
+
+		public bool IsEnabled(LogLevel logLevel)
+		{
+			return false;
+		}
+
+		public void Log<TState>(LogLevel logLevel, EventId eventId, TState state, Exception? exception, Func<TState, Exception?, string> formatter)
+		{
+		}
+	}
+}
\ No newline at end of file
diff --git a/Engine/Source/Programs/Shared/UnrealBuildTool.Tests/TestUtilities/TestToolChain.cs b/Engine/Source/Programs/Shared/UnrealBuildTool.Tests/TestUtilities/TestToolChain.cs
new file mode 100644
--- /dev/null
+++ b/Engine/Source/Programs/Shared/UnrealBuildTool.Tests/TestUtilities/TestToolChain.cs
@@ -0,0 +1,36 @@
+// Copyright Epic Games, Inc. All Rights Reserved.
+
+using System.Collections.Generic;
+using EpicGames.Core;
+using Microsoft.Extensions.Logging;
+using UnrealBuildBase;
+
+namespace UnrealBuildTool.Tests.TestUtilities
+{
+	internal class TestToolChain : UEToolChain
+	{
+		public TestToolChain(ILogger inLogger) : base(inLogger)
+		{
+		}
+
+		public override FileItem? LinkFiles(LinkEnvironment LinkEnvironment, bool bBuildImportLibraryOnly, IActionGraphBuilder Graph)
+		{
+			return null;
+		}
+
+		protected override CPPOutput CompileCPPFiles(CppCompileEnvironment CompileEnvironment, IEnumerable<FileItem> InputFiles, DirectoryReference OutputDir, string ModuleName, IActionGraphBuilder Graph)
+		{
+			return null;
+		}
+
+		protected override CPPOutput CompileISPCFiles(CppCompileEnvironment CompileEnvironment, IEnumerable<FileItem> InputFiles, DirectoryReference OutputDir, IActionGraphBuilder Graph)
+		{
+			return null;
+		}
+
+		protected override CPPOutput GenerateISPCHeaders(CppCompileEnvironment CompileEnvironment, IEnumerable<FileItem> InputFiles, DirectoryReference OutputDir, IActionGraphBuilder Graph)
+		{
+			return null;
+		}
+	}
+}
