PR #11833: Add throw if Refresh token was not returned in response.

diff --git a/Engine/Source/Programs/Shared/EpicGames.OIDC/OidcTokenManager.cs b/Engine/Source/Programs/Shared/EpicGames.OIDC/OidcTokenManager.cs
--- a/Engine/Source/Programs/Shared/EpicGames.OIDC/OidcTokenManager.cs
+++ b/Engine/Source/Programs/Shared/EpicGames.OIDC/OidcTokenManager.cs
@@ -376,6 +376,11 @@ namespace EpicGames.OIDC
 				throw new LoginFailedException("Failed to login due to error: " + loginResult.Error, loginResult.ErrorDescription);
 			}
 
+   			if (String.IsNullOrEmpty(loginResult.RefreshToken))
+			{
+				throw new Exception("No refresh token was provided in response.");
+			}
+
 			_refreshToken = loginResult.RefreshToken;
 			_accessToken = loginResult.AccessToken;
 			_tokenExpiry = loginResult.AccessTokenExpiration;
@@ -581,6 +586,10 @@ namespace EpicGames.OIDC
 			_tokenExpiry = refreshTokenResult.AccessTokenExpiration;
 
 			// refresh tokens are always one time use only so we need to store this new refresh token we got so it can be used the next time
+   			if (String.IsNullOrEmpty(refreshTokenResult.RefreshToken))
+			{
+				throw new Exception("No refresh token was provided in response to refresh.");
+			}
 			_tokenStore.AddRefreshToken(_name, _refreshToken);
 			_tokenStore.Save();
 
@@ -725,4 +734,4 @@ namespace EpicGames.OIDC
 		public string Scopes { get; set; } = "openid profile offline_access email";
 		public string? GenericErrorInformation { get; set; }
 	}
-}
\ No newline at end of file
+}
