Add UHT support for verse::TVal and verse::TPtr

#rnx

diff --git a/Engine/Source/Programs/Shared/EpicGames.UHT/Exporters/CodeGen/UhtCodeGenerator.cs b/Engine/Source/Programs/Shared/EpicGames.UHT/Exporters/CodeGen/UhtCodeGenerator.cs
--- a/Engine/Source/Programs/Shared/EpicGames.UHT/Exporters/CodeGen/UhtCodeGenerator.cs
+++ b/Engine/Source/Programs/Shared/EpicGames.UHT/Exporters/CodeGen/UhtCodeGenerator.cs
@@ -46,6 +46,8 @@ namespace EpicGames.UHT.Exporters.CodeGen
 			public bool NeedsVerseStruct { get; set; }
 			public bool NeedsVerseEnum { get; set; }
 			public bool NeedsVerseInterop { get; set; }
+			public bool NeedsVerseField { get; set; }
+			public bool NeedsVersePointer { get; set; }
 		}
 		public HeaderInfo[] HeaderInfos { get; set; }
 
@@ -442,13 +444,27 @@ namespace EpicGames.UHT.Exporters.CodeGen
 			{
 				headerInfo.NeedsVerseStruct = true;
 			}
-			if (structObj.Children.Any(x => x is UhtVerseValueProperty))
+			foreach (UhtType childType in structObj.Children)
 			{
-				headerInfo.NeedsVerseValue = true;
-			}
-			if (structObj.Children.Any(x => x is UhtVerseStringProperty))
-			{
-				headerInfo.NeedsVerseString = true;
+				if (childType is UhtProperty childProperty)
+				{
+					if (childProperty.PropertyExportFlags.HasAnyFlags(UhtPropertyExportFlags.TVal))
+					{
+						headerInfo.NeedsVerseField = true;
+					}
+					if (childProperty.PropertyExportFlags.HasAnyFlags(UhtPropertyExportFlags.TPtr))
+					{
+						headerInfo.NeedsVersePointer = true;
+					}
+					if (!headerInfo.NeedsVerseValue && childProperty is UhtVerseValueProperty)
+					{
+						headerInfo.NeedsVerseValue = true;
+					}
+					if (!headerInfo.NeedsVerseString && childProperty is UhtVerseStringProperty)
+					{
+						headerInfo.NeedsVerseString = true;
+					}
+				}
 			}
 		}
 		#endregion
diff --git a/Engine/Source/Programs/Shared/EpicGames.UHT/Exporters/CodeGen/UhtHeaderCodeGeneratorHFile.cs b/Engine/Source/Programs/Shared/EpicGames.UHT/Exporters/CodeGen/UhtHeaderCodeGeneratorHFile.cs
--- a/Engine/Source/Programs/Shared/EpicGames.UHT/Exporters/CodeGen/UhtHeaderCodeGeneratorHFile.cs
+++ b/Engine/Source/Programs/Shared/EpicGames.UHT/Exporters/CodeGen/UhtHeaderCodeGeneratorHFile.cs
@@ -74,6 +74,14 @@ namespace EpicGames.UHT.Exporters.CodeGen
 				{
 					builder.Append("#include \"VerseInteropMacros.h\"\r\n");
 				}
+				if (headerInfo.NeedsVerseField)
+				{
+					builder.Append("#include \"VerseField.h\"\r\n");
+				}
+				if (headerInfo.NeedsVersePointer)
+				{
+					builder.Append("#include \"VersePointer.h\"\r\n");
+				}
 				builder.Append("\r\n");
 				builder.Append(DisableDeprecationWarnings).Append("\r\n");
 
diff --git a/Engine/Source/Programs/Shared/EpicGames.UHT/Parsers/UhtPropertyParser.cs b/Engine/Source/Programs/Shared/EpicGames.UHT/Parsers/UhtPropertyParser.cs
--- a/Engine/Source/Programs/Shared/EpicGames.UHT/Parsers/UhtPropertyParser.cs
+++ b/Engine/Source/Programs/Shared/EpicGames.UHT/Parsers/UhtPropertyParser.cs
@@ -1128,6 +1128,44 @@ namespace EpicGames.UHT.Parsers
 		}
 	}
 
+	[UnrealHeaderTool]
+	static class UhtVerseTypeParser
+	{
+		[UhtPropertyType(Keyword = "verse::TVal")]
+		[SuppressMessage("CodeQuality", "IDE0051:Remove unused private members", Justification = "Attribute accessed method")]
+		[SuppressMessage("Style", "IDE0060:Remove unused parameter", Justification = "Attribute accessed method")]
+		private static UhtProperty? VerseTValType(UhtPropertyResolveArgs args)
+		{
+			if (!args.SkipExpectedType())
+			{
+				return null;
+			}
+
+			args.PropertySettings.PropertyExportFlags |= UhtPropertyExportFlags.TVal;
+			args.TokenReader.Require('<');
+			UhtProperty? inner = args.ParseWrappedType();
+			args.TokenReader.Require('>');
+			return inner;
+		}
+
+		[UhtPropertyType(Keyword = "verse::TPtr")]
+		[SuppressMessage("CodeQuality", "IDE0051:Remove unused private members", Justification = "Attribute accessed method")]
+		[SuppressMessage("Style", "IDE0060:Remove unused parameter", Justification = "Attribute accessed method")]
+		private static UhtProperty? VerseTPtrType(UhtPropertyResolveArgs args)
+		{
+			if (!args.SkipExpectedType())
+			{
+				return null;
+			}
+
+			args.PropertySettings.PropertyExportFlags |= UhtPropertyExportFlags.TPtr;
+			args.TokenReader.Require('<');
+			UhtProperty? inner = args.ParseWrappedType();
+			args.TokenReader.Require('>');
+			return inner;
+		}
+	}
+
 	[UnrealHeaderTool]
 	static class UhtDefaultPropertyParser
 	{
diff --git a/Engine/Source/Programs/Shared/EpicGames.UHT/Types/UhtProperty.cs b/Engine/Source/Programs/Shared/EpicGames.UHT/Types/UhtProperty.cs
--- a/Engine/Source/Programs/Shared/EpicGames.UHT/Types/UhtProperty.cs
+++ b/Engine/Source/Programs/Shared/EpicGames.UHT/Types/UhtProperty.cs
@@ -87,6 +87,16 @@ namespace EpicGames.UHT.Types
 		/// If true, the property should have a generated setter function
 		/// </summary>
 		SetterSpecifiedAuto = 0x00002000,
+
+		/// <summary>
+		/// If true, the type is a TVal
+		/// </summary>
+		TVal = 0x00004000,
+
+		/// <summary>
+		/// If true, the type is a TPtr
+		/// </summary>
+		TPtr = 0x00008000,
 	};
 
 	/// <summary>
@@ -816,6 +826,36 @@ namespace EpicGames.UHT.Types
 			TypeTokens = parentPropertySettings.TypeTokens;
 		}
 
+		/// <summary>
+		/// Copy property settings
+		/// </summary>
+		/// <param name="propertySettings">Setting to copy</param>
+		public void Copy(UhtPropertySettings propertySettings)
+		{
+			ParentSettings = propertySettings.ParentSettings;
+			SourceName = propertySettings.SourceName;
+			EngineName = propertySettings.EngineName;
+			VerseName = propertySettings.VerseName;
+			MetaData = propertySettings.MetaData;
+			Outer = propertySettings.Outer;
+			LineNumber = propertySettings.LineNumber;
+			PropertyCategory = propertySettings.PropertyCategory;
+			PropertyFlags = propertySettings.PropertyFlags;
+			DisallowPropertyFlags = propertySettings.DisallowPropertyFlags;
+			PropertyExportFlags = propertySettings.PropertyExportFlags;
+			DefineScope = propertySettings.DefineScope;
+			RepNotifyName = propertySettings.RepNotifyName;
+			Allocator = propertySettings.Allocator;
+			Options = propertySettings.Options;
+			PointerType = propertySettings.PointerType;
+			ArrayDimensions = propertySettings.ArrayDimensions;
+			DefaultValueTokens = propertySettings.DefaultValueTokens;
+			Setter = propertySettings.Setter;
+			Getter = propertySettings.Getter;
+			IsBitfield = propertySettings.IsBitfield;
+			TypeTokens = propertySettings.TypeTokens;
+		}
+
 		/// <summary>
 		/// Reset the property settings.  Used on a cached property settings object
 		/// </summary>
diff --git a/Engine/Source/Programs/Shared/EpicGames.UHT/Utils/UhtPropertyResolveArgs.cs b/Engine/Source/Programs/Shared/EpicGames.UHT/Utils/UhtPropertyResolveArgs.cs
--- a/Engine/Source/Programs/Shared/EpicGames.UHT/Utils/UhtPropertyResolveArgs.cs
+++ b/Engine/Source/Programs/Shared/EpicGames.UHT/Utils/UhtPropertyResolveArgs.cs
@@ -107,26 +107,6 @@ namespace EpicGames.UHT.Utils
 		/// <returns>Parsed property</returns>
 		public readonly UhtProperty? ParseTemplateParam(StringView paramName)
 		{
-			// From the current segment index, look forward for the template segment being parsed
-			int tokenIndex = TokenReader.GetReplayTokenIndex();
-			int templateSegmentIndex = TypeTokens.SegmentIndex + 1;
-			foreach (UhtTypeSegment segment in TypeTokens.Segments[(TypeTokens.SegmentIndex + 1)..].Span)
-			{
-				if (tokenIndex < segment.Start)
-				{
-					throw new UhtIceException("Error parsing template argument");
-				}
-				if (tokenIndex < segment.End)
-				{
-					break;
-				}
-				templateSegmentIndex++;
-			}
-			if (templateSegmentIndex == TypeTokens.Segments.Length)
-			{
-				throw new UhtIceException("Error parsing template argument");
-			}
-
 			using UhtThreadBorrower<UhtPropertySpecifierContext> borrower = new(true);
 			UhtPropertySpecifierContext specifierContext = borrower.Instance;
 			specifierContext.Type = PropertySettings.Outer;
@@ -134,7 +114,7 @@ namespace EpicGames.UHT.Utils
 			specifierContext.AccessSpecifier = UhtAccessSpecifier.None;
 			specifierContext.MessageSite = TokenReader;
 			specifierContext.PropertySettings.Reset(PropertySettings, paramName.ToString(), TokenReader);
-			specifierContext.PropertySettings.TypeTokens = new(TypeTokens, templateSegmentIndex);
+			specifierContext.PropertySettings.TypeTokens = new(TypeTokens, GetTypeTokenSegment());
 			specifierContext.MetaData = specifierContext.PropertySettings.MetaData;
 			specifierContext.MetaNameIndex = UhtMetaData.IndexNone;
 			specifierContext.SeenEditSpecifier = false;
@@ -150,6 +130,35 @@ namespace EpicGames.UHT.Utils
 			return templateArgs.ResolveProperty();
 		}
 
+		/// <summary>
+		/// For types such as TPtr or TVal, this method parses the template parameter and extracts the type
+		/// </summary>
+		/// <returns></returns>
+		public readonly UhtProperty? ParseWrappedType()
+		{
+			using UhtThreadBorrower<UhtPropertySpecifierContext> borrower = new(true);
+			UhtPropertySpecifierContext specifierContext = borrower.Instance;
+			specifierContext.Type = PropertySettings.Outer;
+			specifierContext.TokenReader = TokenReader;
+			specifierContext.AccessSpecifier = UhtAccessSpecifier.None;
+			specifierContext.MessageSite = TokenReader;
+			specifierContext.PropertySettings.Copy(PropertySettings);
+			specifierContext.PropertySettings.TypeTokens = new(TypeTokens, GetTypeTokenSegment());
+			specifierContext.MetaData = PropertySettings.MetaData;
+			specifierContext.MetaNameIndex = UhtMetaData.IndexNone;
+			specifierContext.SeenEditSpecifier = false;
+			specifierContext.SeenBlueprintWriteSpecifier = false;
+			specifierContext.SeenBlueprintReadOnlySpecifier = false;
+			specifierContext.SeenBlueprintGetterSpecifier = false;
+
+			// Pre-parse any UPARAM 
+			UhtPropertyParser.PreParseType(specifierContext, true);
+
+			// Parse the type
+			UhtPropertyResolveArgs templateArgs = new(ResolvePhase, specifierContext.PropertySettings, TokenReader);
+			return templateArgs.ResolveProperty();
+		}
+
 		/// <summary>
 		/// Parse a template type
 		/// </summary>
@@ -371,5 +380,35 @@ namespace EpicGames.UHT.Utils
 			string fullType = UhtTypeTokens.ToString(copiedArray);
 			return Session.TryGetPropertyType(fullType, out propertyType);
 		}
+
+
+		/// <summary>
+		/// Given the current position of the token reader, return the segment being parsed
+		/// </summary>
+		/// <returns>Segment index being parsed</returns>
+		/// <exception cref="UhtIceException"></exception>
+		private readonly int GetTypeTokenSegment()
+		{
+			// From the current segment index, look forward for the template segment being parsed
+			int tokenIndex = TokenReader.GetReplayTokenIndex();
+			int templateSegmentIndex = TypeTokens.SegmentIndex + 1;
+			foreach (UhtTypeSegment segment in TypeTokens.Segments[(TypeTokens.SegmentIndex + 1)..].Span)
+			{
+				if (tokenIndex < segment.Start)
+				{
+					throw new UhtIceException("Error parsing template argument");
+				}
+				if (tokenIndex < segment.End)
+				{
+					break;
+				}
+				templateSegmentIndex++;
+			}
+			if (templateSegmentIndex == TypeTokens.Segments.Length)
+			{
+				throw new UhtIceException("Error parsing template argument");
+			}
+			return templateSegmentIndex;
+		}
 	}
 }
