Cloud Artifacts - Support uploading files larger then 2GB - implemented support for compressed buffer serialization in C# to read from a stream thus supporting compressing buffers larger then 2GB, still allocating this all in memory so will use a lot of memory for large buffers.

diff --git a/Engine/Source/Programs/AutomationTool/BuildGraph/Tasks/CreateCloudArtifactTask.cs b/Engine/Source/Programs/AutomationTool/BuildGraph/Tasks/CreateCloudArtifactTask.cs
--- a/Engine/Source/Programs/AutomationTool/BuildGraph/Tasks/CreateCloudArtifactTask.cs
+++ b/Engine/Source/Programs/AutomationTool/BuildGraph/Tasks/CreateCloudArtifactTask.cs
@@ -575,7 +575,7 @@ namespace AutomationTool.Tasks
 						FileInfo tempFileInfo = new FileInfo(tempFileName!);
 						{
 							await using Stream s = tempFileInfo.OpenWrite();
-							CompressedBuffer.CompressContent(s, OodleCompressorType.Kraken, OodleCompressionLevel.VeryFast, await sourceStream.ReadAllBytesAsync());
+							CompressedBuffer.CompressContent(s, OodleCompressorType.Kraken, OodleCompressionLevel.VeryFast, sourceStream);
 						}
 
 						tempFileInfo.Refresh();
@@ -585,7 +585,7 @@ namespace AutomationTool.Tasks
 					else
 					{
 						MemoryStream ms = new MemoryStream();
-						CompressedBuffer.CompressContent(ms, OodleCompressorType.Kraken, OodleCompressionLevel.VeryFast, await sourceStream.ReadAllBytesAsync());
+						CompressedBuffer.CompressContent(ms, OodleCompressorType.Kraken, OodleCompressionLevel.VeryFast, sourceStream);
 						ms.Position = 0;
 						
 						streamToSend = ms;
diff --git a/Engine/Source/Programs/Shared/EpicGames.Serialization/CompressedBuffer.cs b/Engine/Source/Programs/Shared/EpicGames.Serialization/CompressedBuffer.cs
--- a/Engine/Source/Programs/Shared/EpicGames.Serialization/CompressedBuffer.cs
+++ b/Engine/Source/Programs/Shared/EpicGames.Serialization/CompressedBuffer.cs
@@ -288,6 +288,33 @@ namespace EpicGames.Serialization
 			return CompressContent(s, method, compressionLevel, blocks, blockSize);
 		}
 
+		/// <summary>
+		/// Compress a stream into a compressed buffer stream
+		/// </summary>
+		/// <param name="s">The stream to write the compressed content into</param>
+		/// <param name="method">Which oodle compressor method to use</param>
+		/// <param name="compressionLevel">The oodle compression level</param>
+		/// <param name="sourceStream">The uncompressed content in a stream that you wish to compress</param>
+		/// <returns>The iohash of the uncompressed content</returns>
+		public static IoHash CompressContent(Stream s, OodleCompressorType method, OodleCompressionLevel compressionLevel, Stream sourceStream)
+		{
+			const long DefaultBlockSize = 256 * 1024;
+			long blockSize = DefaultBlockSize;
+			long blockCount = (sourceStream.Length + blockSize - 1) / blockSize;
+			List<byte[]> blocks = new List<byte[]>();
+
+			for (int i = 0; i < blockCount; i++)
+			{
+				long rawBlockSize = Math.Min(sourceStream.Length - (i * blockSize), blockSize);
+				byte[] bufferToCompress = new byte[(int)rawBlockSize];
+
+				sourceStream.ReadFixedLengthBytes(bufferToCompress);
+				blocks.Add(bufferToCompress);
+			}
+
+			return CompressContent(s, method, compressionLevel, blocks, blockSize);
+		}
+
 		private static IoHash CompressContent(Stream s, OodleCompressorType method, OodleCompressionLevel compressionLevel, List<byte[]> blocks, long blockSize)
 		{
 			long blockCount = blocks.Count;
