OidcToken - Refactored token manager to allow for tests and added some functional tests

diff --git a/Engine/Source/Programs/OidcToken/OidcToken.Tests/OidcToken.Tests.csproj b/Engine/Source/Programs/OidcToken/OidcToken.Tests/OidcToken.Tests.csproj
new file mode 100644
--- /dev/null
+++ b/Engine/Source/Programs/OidcToken/OidcToken.Tests/OidcToken.Tests.csproj
@@ -0,0 +1,29 @@
+<Project Sdk="Microsoft.NET.Sdk">
+
+  <PropertyGroup>
+    <TargetFramework>net8.0</TargetFramework>
+    <ImplicitUsings>enable</ImplicitUsings>
+    <Nullable>enable</Nullable>
+
+    <IsPackable>false</IsPackable>
+    <IsTestProject>true</IsTestProject>
+  </PropertyGroup>
+
+  <ItemGroup>
+    <PackageReference Include="coverlet.collector" Version="6.0.0" />
+    <PackageReference Include="Microsoft.NET.Test.Sdk" Version="17.8.0" />
+    <PackageReference Include="Moq" Version="4.20.72" />
+    <PackageReference Include="Moq.Contrib.HttpClient" Version="1.4.0" />
+    <PackageReference Include="MSTest.TestAdapter" Version="3.1.1" />
+    <PackageReference Include="MSTest.TestFramework" Version="3.1.1" />
+  </ItemGroup>
+
+  <ItemGroup>
+    <ProjectReference Include="..\OidcToken.csproj" />
+  </ItemGroup>
+
+  <ItemGroup>
+    <Using Include="Microsoft.VisualStudio.TestTools.UnitTesting" />
+  </ItemGroup>
+
+</Project>
diff --git a/Engine/Source/Programs/OidcToken/OidcToken.Tests/OidcTokenTests.cs b/Engine/Source/Programs/OidcToken/OidcToken.Tests/OidcTokenTests.cs
new file mode 100644
--- /dev/null
+++ b/Engine/Source/Programs/OidcToken/OidcToken.Tests/OidcTokenTests.cs
@@ -0,0 +1,310 @@
+using System.Diagnostics;
+using System.Net;
+using System.Net.Mime;
+using System.Text;
+using System.Text.Json;
+using EpicGames.OIDC;
+using IdentityModel.Client;
+using IdentityModel.OidcClient;
+using Microsoft.Extensions.DependencyInjection;
+using Microsoft.Extensions.Hosting;
+using Moq;
+using Serilog;
+
+namespace OidcToken.Tests
+{
+	[TestClass]
+	public class OidcTokenTests
+	{
+		private IHost SetupMocks(Action<TokenServiceOptions> configureCliOptions, Action<OidcTokenOptions> configureOidcOptions, Func<OidcTokenClient> mockTokenClient, ITokenStore tokenStore)
+		{
+			Mock<IOidcTokenClientFactory> clientFactory = new Mock<IOidcTokenClientFactory>();
+			clientFactory.Setup(factory =>  factory.CreateTokenClient(It.IsAny<string>(), It.IsAny<ProviderInfo>(), It.IsAny<TimeSpan>())).Returns(mockTokenClient);
+
+			IHost host = Host.CreateDefaultBuilder()
+				.UseEnvironment("Testing")
+				.ConfigureServices(collection => collection.AddSerilog())
+				.ConfigureServices(
+					(content, services) =>
+					{
+						services.AddOptions<TokenServiceOptions>().Configure(configureCliOptions).ValidateDataAnnotations();
+						services.AddOptions<OidcTokenOptions>().Configure(configureOidcOptions).ValidateDataAnnotations();
+
+						services.AddSingleton<IOidcTokenManager, OidcTokenManager>();
+						services.AddSingleton<IOidcTokenClientFactory>(clientFactory.Object);
+						services.AddSingleton<ITokenStore>(tokenStore);
+
+						services.AddSingleton<TokenService>();
+						services.AddHostedService<TokenService>(provider => provider.GetService<TokenService>()!);
+					})
+					.Build();
+
+			return host;
+		}
+
+		[TestMethod]
+		public async Task AllocateTokenToJsonFile()
+		{
+			const string providerName = "Test-Service";
+			ProviderInfo providerInfo = new ProviderInfo()
+			{
+				ClientId = "test-client",
+				DisplayName = "test-client",
+				RedirectUri = null,
+				Scopes = "offline profile",
+			};
+
+			const string mockRefreshToken = "random-refresh-token";
+			// TODO: Make the mocked access token into a valid JWT?
+			const string mockAccessToken = "random-token";
+			DateTimeOffset mockExpiryTime = DateTimeOffset.Now.AddDays(2);
+
+			Mock<InMemoryTokenStore> tokenStoreMock = new Mock<InMemoryTokenStore>()
+			{
+				CallBase = true
+			};
+			tokenStoreMock.Setup(store => store.AddRefreshToken(It.Is(providerName, StringComparer.InvariantCultureIgnoreCase), It.Is(mockRefreshToken, StringComparer.InvariantCultureIgnoreCase))).CallBase().Verifiable(Times.Once);
+			tokenStoreMock.Setup(store => store.Save()).CallBase().Verifiable(Times.Once);
+
+			string tempOutputFileName = Path.GetTempFileName();
+			try
+			{
+				using IHost host = SetupMocks(cliOptions =>
+				{
+					cliOptions.Service = providerName;
+					cliOptions.OutFile = tempOutputFileName;
+				}, oidcTokenOptions =>
+				{
+					oidcTokenOptions.Providers = new Dictionary<string, ProviderInfo>()
+					{
+						{
+							providerName, providerInfo
+						}
+					};
+				},
+				() =>
+				{
+					Microsoft.Extensions.Logging.ILogger<OidcTokenClient>? logger = null;
+					Mock<OidcTokenClient> tokenClientMock = new Mock<OidcTokenClient>(providerName, providerInfo, TimeSpan.FromMinutes(15) , tokenStoreMock.Object, logger!)
+					{
+						// call into the real OidcTokenClient unless the method has been mocked
+						CallBase = true
+					};
+
+					tokenClientMock.Setup(client => client.DoLoginAsync(It.IsAny<CancellationToken>())).ReturnsAsync(() => new OidcTokenInfo() {AccessToken = mockAccessToken, RefreshToken = mockRefreshToken, TokenExpiry = mockExpiryTime} );
+					return tokenClientMock.Object;
+				}, tokenStoreMock.Object);
+
+				TokenService tokenService = host.Services.GetService<TokenService>()!;
+				Assert.IsNotNull(tokenService);
+
+				await tokenService.Main();
+
+				Assert.IsTrue(File.Exists(tempOutputFileName));
+				await using FileStream fs = File.OpenRead(tempOutputFileName);
+				TokenResultFile? result = JsonSerializer.Deserialize<TokenResultFile>(fs);
+
+				Assert.AreEqual(mockAccessToken, result!.Token);
+				Assert.AreEqual(mockExpiryTime, result!.ExpiresAt);
+				tokenStoreMock.VerifyAll();
+			}
+			finally
+			{
+				File.Delete(tempOutputFileName);
+			}
+		}
+
+		[TestMethod]
+		public async Task FailToAllocateToken()
+		{
+			const string providerName = "Test-Service";
+			ProviderInfo providerInfo = new ProviderInfo()
+			{
+				ClientId = "test-client",
+				DisplayName = "test-client",
+				RedirectUri = new Uri("http://localhost:1234"),
+				ServerUri = new Uri("https://path-to-idp.com"),
+				Scopes = "offline profile",
+				UseDiscoveryDocument = false,
+			};
+
+			Mock<InMemoryTokenStore> tokenStoreMock = new Mock<InMemoryTokenStore>()
+			{
+				CallBase = true
+			};
+			tokenStoreMock.Setup(store => store.AddRefreshToken(It.Is(providerName, StringComparer.InvariantCultureIgnoreCase), It.IsAny<string>())).CallBase().Verifiable(Times.Never);
+			tokenStoreMock.Setup(store => store.Save()).CallBase().Verifiable(Times.Never);
+
+			const string loginResultErrorMessage = "failed to login";
+
+
+			string tempOutputFileName = Path.GetTempFileName();
+			try
+			{
+				using IHost host = SetupMocks(cliOptions =>
+				{
+					cliOptions.Service = providerName;
+					cliOptions.OutFile = tempOutputFileName;
+				}, oidcTokenOptions =>
+				{
+					oidcTokenOptions.Providers = new Dictionary<string, ProviderInfo>()
+					{
+						{
+							providerName, providerInfo
+						}
+					};
+				},
+				() =>
+				{
+					Microsoft.Extensions.Logging.ILogger<OidcTokenClient>? logger = null;
+					Mock<OidcTokenClient> tokenClientMock = new Mock<OidcTokenClient>(providerName, providerInfo, TimeSpan.FromMinutes(15) , tokenStoreMock.Object, logger!)
+					{
+						CallBase = true
+					};
+
+					tokenClientMock.Setup(client => client.DoCreateOidcClient(It.IsAny<OidcClientOptions>())).Returns((OidcClientOptions options) =>
+					{
+						Mock<OidcClient> oidcClientMock = new Mock<OidcClient>(options);
+						oidcClientMock.Setup(client => client.ProcessResponseAsync(It.IsAny<string>(), It.IsAny<AuthorizeState>(), It.IsAny<Parameters>(), It.IsAny<CancellationToken>())).ReturnsAsync(() => new LoginResult(loginResultErrorMessage));
+						oidcClientMock.Setup(client => client.PrepareLoginAsync( It.IsAny<Parameters>(), It.IsAny<CancellationToken>())).ReturnsAsync(() => new AuthorizeState());
+
+						return oidcClientMock.Object;
+					});
+
+					tokenClientMock.Setup(client => client.CreateHttpServer()).Returns(() =>
+					{
+						Mock<OidcHttpServer> httpMock = new Mock<OidcHttpServer>();
+
+						OidcHttpServer.OidcHttpRequest request = new OidcHttpServer.OidcHttpRequest("", MediaTypeNames.Application.FormUrlEncoded, "POST", new MemoryStream(Encoding.Default.GetBytes("input-object")), hasBody: true, Encoding.ASCII);
+						OidcHttpServer.NullOidcHttpResponse response = new OidcHttpServer.NullOidcHttpResponse(new MemoryStream());
+
+						httpMock.Setup(listener => listener.ProcessNextRequestAsync()).ReturnsAsync(() => (request, response));
+
+						return httpMock.Object;
+					});
+					tokenClientMock.Setup(client => client.OpenBrowser(It.IsAny<string>())).Returns(Mock.Of<Process>());
+					return tokenClientMock.Object;
+				}, tokenStoreMock.Object);
+
+				TokenService tokenService = host.Services.GetService<TokenService>()!;
+				Assert.IsNotNull(tokenService);
+
+				await Assert.ThrowsExceptionAsync<LoginFailedException>(async () =>
+				{
+					await tokenService.Main();
+				});
+				
+				tokenStoreMock.VerifyAll();
+			}
+			finally
+			{
+				File.Delete(tempOutputFileName);
+			}
+		}
+
+		
+		[TestMethod]
+		public async Task SuccessfulToken()
+		{
+			const string providerName = "Test-Service";
+			ProviderInfo providerInfo = new ProviderInfo()
+			{
+				ClientId = "test-client",
+				DisplayName = "test-client",
+				RedirectUri = new Uri("http://localhost:1234"),
+				ServerUri = new Uri("https://path-to-idp.com"),
+				Scopes = "offline profile",
+				UseDiscoveryDocument = false,
+			};
+
+			Mock<InMemoryTokenStore> tokenStoreMock = new Mock<InMemoryTokenStore>()
+			{
+				CallBase = true
+			};
+			tokenStoreMock.Setup(store => store.AddRefreshToken(It.Is(providerName, StringComparer.InvariantCultureIgnoreCase), It.IsAny<string>())).CallBase().Verifiable(Times.Once);
+			tokenStoreMock.Setup(store => store.Save()).CallBase().Verifiable(Times.Once);
+
+			string mockAccessToken = "access-token";
+			string mockRefreshToken = "refresh-token";
+			DateTimeOffset mockExpiryTime = DateTimeOffset.Now.AddHours(4);
+
+			string tempOutputFileName = Path.GetTempFileName();
+			try
+			{
+				Mock<OidcHttpServer> httpMock = new Mock<OidcHttpServer>();
+				// buffer to store the response in
+				byte[] responseBuffer = new byte[16 * 1024];
+				OidcHttpServer.NullOidcHttpResponse response = new OidcHttpServer.NullOidcHttpResponse(new MemoryStream(responseBuffer));
+
+				using IHost host = SetupMocks(cliOptions =>
+				{
+					cliOptions.Service = providerName;
+					cliOptions.OutFile = tempOutputFileName;
+				}, oidcTokenOptions =>
+				{
+					oidcTokenOptions.Providers = new Dictionary<string, ProviderInfo>()
+					{
+						{
+							providerName, providerInfo
+						}
+					};
+				},
+				() =>
+				{
+					Microsoft.Extensions.Logging.ILogger<OidcTokenClient>? logger = null;
+					Mock<OidcTokenClient> tokenClientMock = new Mock<OidcTokenClient>(providerName, providerInfo, TimeSpan.FromMinutes(15) , tokenStoreMock.Object, logger!)
+					{
+						CallBase = true
+					};
+
+					tokenClientMock.Setup(client => client.DoCreateOidcClient(It.IsAny<OidcClientOptions>())).Returns((OidcClientOptions options) =>
+					{
+						Mock<OidcClient> oidcClientMock = new Mock<OidcClient>(options);
+						Mock<LoginResult> mockLoginResult = new Mock<LoginResult>();
+						mockLoginResult.Setup(m => m.RefreshToken).Returns(mockRefreshToken);
+						mockLoginResult.Setup(m => m.AccessToken).Returns(mockAccessToken);
+						mockLoginResult.Setup(m => m.AccessTokenExpiration).Returns(mockExpiryTime);
+						oidcClientMock.Setup(client => client.ProcessResponseAsync(It.IsAny<string>(), It.IsAny<AuthorizeState>(), It.IsAny<Parameters>(), It.IsAny<CancellationToken>())).ReturnsAsync(() => mockLoginResult.Object);
+						oidcClientMock.Setup(client => client.PrepareLoginAsync( It.IsAny<Parameters>(), It.IsAny<CancellationToken>())).ReturnsAsync(() => new AuthorizeState());
+
+						return oidcClientMock.Object;
+					});
+
+					tokenClientMock.Setup(client => client.CreateHttpServer()).Returns(() =>
+					{
+						OidcHttpServer.OidcHttpRequest request = new OidcHttpServer.OidcHttpRequest("", MediaTypeNames.Application.FormUrlEncoded, "POST", new MemoryStream(Encoding.Default.GetBytes("input-object")), hasBody: true, Encoding.ASCII);
+						httpMock.Setup(listener => listener.ProcessNextRequestAsync()).ReturnsAsync(() => (request, response)).Verifiable(Times.Once);
+
+						return httpMock.Object;
+					});
+					tokenClientMock.Setup(client => client.OpenBrowser(It.IsAny<string>())).Returns(Mock.Of<Process>());
+					return tokenClientMock.Object;
+				}, tokenStoreMock.Object);
+
+				TokenService tokenService = host.Services.GetService<TokenService>()!;
+				Assert.IsNotNull(tokenService);
+
+				await tokenService.Main();
+
+				// check the result page to see what got written
+				string s = Encoding.ASCII.GetString(responseBuffer);
+				// no point in checking the exact html returned just make sure something was written
+				Assert.IsFalse(string.IsNullOrEmpty(s));
+				
+				Assert.IsTrue(File.Exists(tempOutputFileName));
+				await using FileStream fs = File.OpenRead(tempOutputFileName);
+				TokenResultFile? result = JsonSerializer.Deserialize<TokenResultFile>(fs);
+
+				// this does not actually verify that a token is correctly allocated  as that requires reaching out to a idp, and we mock away all of those parts. This can only check that the parts we do are working as such these are just the same values as we specified in the mock
+				Assert.AreEqual(mockAccessToken, result!.Token);
+				Assert.AreEqual(mockExpiryTime, result!.ExpiresAt);
+
+				tokenStoreMock.VerifyAll();
+			}
+			finally
+			{
+				File.Delete(tempOutputFileName);
+			}
+		}
+	}
+}
\ No newline at end of file
diff --git a/Engine/Source/Programs/OidcToken/OidcToken.csproj b/Engine/Source/Programs/OidcToken/OidcToken.csproj
--- a/Engine/Source/Programs/OidcToken/OidcToken.csproj
+++ b/Engine/Source/Programs/OidcToken/OidcToken.csproj
@@ -26,6 +26,12 @@
     <PackageReference Include="Serilog.Sinks.File" Version="6.0.0" />
   </ItemGroup>
 
+  <ItemGroup>
+    <Compile Remove="OidcToken.Tests\**" />
+    <EmbeddedResource Remove="OidcToken.Tests\**" />
+    <None Remove="OidcToken.Tests\**" />
+  </ItemGroup>
+
   <ItemGroup>
     <ProjectReference Include="..\Shared\EpicGames.OIDC\EpicGames.OIDC.csproj" />
   </ItemGroup>
diff --git a/Engine/Source/Programs/OidcToken/Program.cs b/Engine/Source/Programs/OidcToken/Program.cs
--- a/Engine/Source/Programs/OidcToken/Program.cs
+++ b/Engine/Source/Programs/OidcToken/Program.cs
@@ -137,7 +137,8 @@ namespace OidcToken
 					}
 					services.AddOptions<OidcTokenOptions>().Bind(serviceConfig).ValidateDataAnnotations();
 
-					services.AddSingleton<OidcTokenManager>();
+					services.AddSingleton<IOidcTokenManager, OidcTokenManager>();
+					services.AddSingleton<IOidcTokenClientFactory, OidcTokenClientFactory>();
 					services.AddSingleton<ITokenStore>(TokenStoreFactory.CreateTokenStore);
 
 					services.AddHostedService<TokenService>();
diff --git a/Engine/Source/Programs/OidcToken/TokenService.cs b/Engine/Source/Programs/OidcToken/TokenService.cs
--- a/Engine/Source/Programs/OidcToken/TokenService.cs
+++ b/Engine/Source/Programs/OidcToken/TokenService.cs
@@ -27,10 +27,10 @@ namespace OidcToken
 		private readonly ILogger<TokenService> Logger;
 		private readonly IHostApplicationLifetime AppLifetime;
 		private readonly IOptionsMonitor<TokenServiceOptions> Settings;
-		private readonly OidcTokenManager TokenManager;
+		private readonly IOidcTokenManager TokenManager;
 		private int? ExitCode;
 
-		public TokenService(ILogger<TokenService> logger, IHostApplicationLifetime appLifetime, IOptionsMonitor<TokenServiceOptions> settings, OidcTokenManager tokenManager)
+		public TokenService(ILogger<TokenService> logger, IHostApplicationLifetime appLifetime, IOptionsMonitor<TokenServiceOptions> settings, IOidcTokenManager tokenManager)
 		{
 			this.Logger = logger;
 			this.AppLifetime = appLifetime;
@@ -73,7 +73,7 @@ namespace OidcToken
 			return Task.CompletedTask;
 		}
 
-		private async Task Main()
+		public async Task Main()
 		{
 			string providerName = Settings.CurrentValue.Service;
 
@@ -227,7 +227,7 @@ namespace OidcToken
 	{
 	}
 
-	internal class TokenResultFile
+	public class TokenResultFile
 	{
 		public string Token { get; set; }
 		public DateTimeOffset ExpiresAt { get; set; }
@@ -239,8 +239,7 @@ namespace OidcToken
 		}
 	}
 
-
-	internal class TokenStatusFile
+	public class TokenStatusFile
 	{
 		public string Service { get; set; }
 
diff --git a/Engine/Source/Programs/Shared/EpicGames.OIDC/InMemoryTokenStore.cs b/Engine/Source/Programs/Shared/EpicGames.OIDC/InMemoryTokenStore.cs
--- a/Engine/Source/Programs/Shared/EpicGames.OIDC/InMemoryTokenStore.cs
+++ b/Engine/Source/Programs/Shared/EpicGames.OIDC/InMemoryTokenStore.cs
@@ -8,12 +8,13 @@ namespace EpicGames.OIDC;
 /// <summary>
 /// An in-memory token store without encryption, intended for testing
 /// </summary>
-public sealed class InMemoryTokenStore : ITokenStore
+// ReSharper disable once ClassNeverInstantiated.Global
+public class InMemoryTokenStore : ITokenStore
 {
 	private readonly Dictionary<string, string> _providerToRefreshToken = new ();
 
 	/// <inheritdoc/>
-	public bool TryGetRefreshToken(string oidcProvider, out string refreshToken)
+	public virtual bool TryGetRefreshToken(string oidcProvider, out string refreshToken)
 	{
 		if (!_providerToRefreshToken.TryGetValue(oidcProvider, out string? result))
 		{
@@ -26,20 +27,28 @@ public sealed class InMemoryTokenStore : ITokenStore
 	}
 
 	/// <inheritdoc/>
-	public void AddRefreshToken(string providerIdentifier, string refreshToken)
+	public virtual void AddRefreshToken(string providerIdentifier, string refreshToken)
 	{
 		_providerToRefreshToken[providerIdentifier] = refreshToken;
 	}
 
 	/// <inheritdoc/>
-	public void Save()
+	public virtual void Save()
 	{
 		// No need - data exists only in-memory
 	}
 
-	/// <inheritdoc/>
+	protected virtual void Dispose(bool disposing)
+	{
+		if (disposing)
+		{
+			// no resources to release
+		}
+	}
+
 	public void Dispose()
 	{
+		Dispose(true);
 		GC.SuppressFinalize(this);
 	}
 }
diff --git a/Engine/Source/Programs/Shared/EpicGames.OIDC/OidcTokenManager.cs b/Engine/Source/Programs/Shared/EpicGames.OIDC/OidcTokenManager.cs
--- a/Engine/Source/Programs/Shared/EpicGames.OIDC/OidcTokenManager.cs
+++ b/Engine/Source/Programs/Shared/EpicGames.OIDC/OidcTokenManager.cs
@@ -62,7 +62,7 @@ namespace EpicGames.OIDC
 			}
 		}
 
-		public OidcTokenManager(IServiceProvider provider, IOptionsMonitor<OidcTokenOptions> settings, ITokenStore tokenStore, List<string>? allowedProviders = null)
+		public OidcTokenManager(IServiceProvider provider, IOptionsMonitor<OidcTokenOptions> settings, ITokenStore tokenStore, IOidcTokenClientFactory tokenClientFactory, List<string>? allowedProviders = null)
 			: this(tokenStore)
 		{
 			Dictionary<string, string> refreshTokens = new Dictionary<string, string>();
@@ -82,7 +82,7 @@ namespace EpicGames.OIDC
 					continue;
 				}
 
-				OidcTokenClient tokenClient = ActivatorUtilities.CreateInstance<OidcTokenClient>(provider, key, providerInfo, settings.CurrentValue.LoginTimeout);
+				OidcTokenClient tokenClient = tokenClientFactory.CreateTokenClient(key, providerInfo, settings.CurrentValue.LoginTimeout);
 
 				if (refreshTokens.TryGetValue(key, out string? refreshToken))
 				{
@@ -140,18 +140,9 @@ namespace EpicGames.OIDC
 		}
 
 		/// <inheritdoc/>
-		public async Task<OidcTokenInfo> LoginAsync(string providerIdentifier, CancellationToken cancellationToken = default)
+		public Task<OidcTokenInfo> LoginAsync(string providerIdentifier, CancellationToken cancellationToken = default)
 		{
-			OidcTokenClient tokenClient = _tokenClients[providerIdentifier];
-
-			OidcTokenInfo tokenInfo = await tokenClient.LoginAsync(cancellationToken);
-
-			if (!String.IsNullOrEmpty(tokenInfo.RefreshToken))
-			{
-				_tokenStore.AddRefreshToken(providerIdentifier, tokenInfo.RefreshToken);
-			}
-
-			return tokenInfo;
+			return _tokenClients[providerIdentifier].LoginAsync(cancellationToken);
 		}
 
 		/// <inheritdoc/>
@@ -173,6 +164,19 @@ namespace EpicGames.OIDC
 		}
 	}
 
+	public interface IOidcTokenClientFactory
+	{
+		OidcTokenClient CreateTokenClient(string key, ProviderInfo providerInfo, TimeSpan loginTimeout);
+	}
+
+	public class OidcTokenClientFactory(IServiceProvider provider) : IOidcTokenClientFactory
+	{
+		public OidcTokenClient CreateTokenClient(string key, ProviderInfo providerInfo, TimeSpan loginTimeout)
+		{
+			return ActivatorUtilities.CreateInstance<OidcTokenClient>(provider, key, providerInfo, loginTimeout);
+		}
+	}
+
 	public enum OidcStatus
 	{
 		/// <summary>
@@ -218,6 +222,116 @@ namespace EpicGames.OIDC
 		}
 	}
 
+	/// <summary>
+	/// Wrapper around HttpListener to make it possible to inject a mock and thus test the Oidc flow without starting a http server
+	/// </summary>
+	public class OidcHttpServer : IDisposable
+	{
+		public class OidcHttpRequest
+		{
+			public OidcHttpRequest(HttpListenerContext context)
+			{
+				HttpListenerRequest request = context.Request;
+				HasBody = request.HasEntityBody;
+				HttpMethod = request.HttpMethod;
+				RawUrl = request.RawUrl;
+				ContentType = request.ContentType;
+				Body = request.InputStream;
+				BodyEncoding = request.ContentEncoding;
+			}
+
+			public OidcHttpRequest(string rawUrl, string contentType, string httpMethod, Stream body, bool hasBody, Encoding bodyEncoding)
+			{
+				RawUrl = rawUrl;
+				ContentType = contentType;
+				HttpMethod = httpMethod;
+				Body = body;
+				HasBody = hasBody;
+				BodyEncoding = bodyEncoding;
+			}
+
+
+			public bool HasBody { get; init; }
+			public string HttpMethod { get; init; }
+			public string? RawUrl { get; init; }
+			public string? ContentType { get; init; }
+			public Stream Body { get; init;}
+			public Encoding BodyEncoding { get; init; }
+		}
+
+		public interface IOidcHttpResponse : IDisposable
+		{
+			public long ContentLength64 { get; set; }
+			public int StatusCode { get; set; }
+			public Stream OutputStream { get; }
+			public void Close()
+			{
+				Dispose();
+			}
+		}
+
+		// Implementation of a response writer for http listener
+		public class OidcHttpResponse(HttpListenerContext context) : IOidcHttpResponse
+		{
+			private readonly HttpListenerResponse _response = context.Response;
+
+			public long ContentLength64 { get => _response.ContentLength64; set => _response.ContentLength64 = value; }
+			public int StatusCode  { get => _response.StatusCode; set => _response.StatusCode = value; }
+			public Stream OutputStream => _response.OutputStream;
+
+			public void Close()
+			{
+				_response.Close();
+			}
+
+			public void Dispose()
+			{
+				((IDisposable)_response).Dispose();
+			}
+		}
+
+		// a response writer that does not actually use the response in any way, just buffering it internally
+		public class NullOidcHttpResponse(Stream outputStream) : IOidcHttpResponse
+		{
+			public void Dispose()
+			{
+				// nothing to dispose
+			}
+
+			public long ContentLength64 { get; set; }
+			public int StatusCode { get; set; }
+			public Stream OutputStream { get; } = outputStream;
+		}
+
+		private readonly HttpListener _http = new();
+
+		public virtual void Start()
+		{
+			_http.Start();
+		}
+
+		public virtual void Dispose()
+		{
+			((IDisposable)_http).Dispose();
+		}
+
+		public virtual void Stop()
+		{
+			_http.Stop();
+		}
+
+		public virtual void AddPrefixes(string prefix)
+		{
+			_http.Prefixes.Add(prefix);
+		}
+
+		public virtual async Task<(OidcHttpRequest, IOidcHttpResponse)> ProcessNextRequestAsync()
+		{
+			HttpListenerContext context = await _http.GetContextAsync();
+			return (new OidcHttpRequest(context), new OidcHttpResponse(context));
+		}
+	}
+	
 	public class OidcTokenClient
 	{
 		public string Name => _name;
@@ -238,6 +352,7 @@ namespace EpicGames.OIDC
 
 		private readonly List<Uri> _redirectUris;
 		private readonly string? _genericErrorInformation;
+		private readonly bool _useDiscoveryDocument;
 
 		public OidcTokenClient(string name, ProviderInfo providerInfo, TimeSpan loginTimeout, ITokenStore tokenStore, ILogger<OidcTokenClient>? logger = null)
 		{
@@ -262,6 +377,7 @@ namespace EpicGames.OIDC
 			_redirectUris = possibleRedirectUris;
 
 			_scopes = providerInfo.Scopes;
+			_useDiscoveryDocument = providerInfo.UseDiscoveryDocument;
 
 			_genericErrorInformation = providerInfo.GenericErrorInformation;
 			_tokenStore.TryGetRefreshToken(name, out _refreshToken);
@@ -278,29 +394,53 @@ namespace EpicGames.OIDC
 				FilterClaims = false,
 				RedirectUri = redirectUri.ToString(),
 				LoadProfile = _providerInfo.LoadClaimsFromUserProfile,
-				//Flow = OidcClientOptions.AuthenticationFlow.AuthorizationCode,
-				//ResponseMode = OidcClientOptions.AuthorizeResponseMode.Redirect
 			};
 
 			// we need to fetch the discovery document ourselves to support OIDC Authorities which have a subresource for it
-			// with Okta has for authorization servers for instance.
-			DiscoveryDocumentResponse discoveryDocument = await GetDiscoveryDocumentAsync(cancellationToken);
-			options.ProviderInformation = new ProviderInformation
-			{
-				IssuerName = discoveryDocument.Issuer,
-				KeySet = discoveryDocument.KeySet,
-
-				AuthorizeEndpoint = discoveryDocument.AuthorizeEndpoint,
-				TokenEndpoint = discoveryDocument.TokenEndpoint,
-				EndSessionEndpoint = discoveryDocument.EndSessionEndpoint,
-				UserInfoEndpoint = discoveryDocument.UserInfoEndpoint,
-				TokenEndPointAuthenticationMethods = discoveryDocument.TokenEndpointAuthenticationMethodsSupported
-			};
+			// which Okta has for authorization servers for instance.
+
+			if (_useDiscoveryDocument)
+			{
+				DiscoveryDocumentResponse discoveryDocument = await GetDiscoveryDocumentAsync(cancellationToken);
+				options.ProviderInformation = new ProviderInformation
+				{
+					IssuerName = discoveryDocument.Issuer,
+					KeySet = discoveryDocument.KeySet,
+
+					AuthorizeEndpoint = discoveryDocument.AuthorizeEndpoint,
+					TokenEndpoint = discoveryDocument.TokenEndpoint,
+					EndSessionEndpoint = discoveryDocument.EndSessionEndpoint,
+					UserInfoEndpoint = discoveryDocument.UserInfoEndpoint,
+					TokenEndPointAuthenticationMethods = discoveryDocument.TokenEndpointAuthenticationMethodsSupported
+				};
+			}
 
 			return options;
 		}
 
 		public async Task<OidcTokenInfo> LoginAsync(CancellationToken cancellationToken)
+		{
+			OidcTokenInfo tokenInfo = await DoLoginAsync(cancellationToken);
+
+   			if (String.IsNullOrEmpty(tokenInfo.RefreshToken))
+			{
+				throw new Exception("No refresh token was provided in response.");
+			}
+
+			_refreshToken = tokenInfo.RefreshToken;
+			_accessToken = tokenInfo.AccessToken;
+			_tokenExpiry = tokenInfo.TokenExpiry;
+
+			if (_refreshToken != null)
+			{
+				_tokenStore.AddRefreshToken(_name, _refreshToken);
+				_tokenStore.Save();
+			}
+
+			return tokenInfo;
+		}
+
+		public virtual async Task<OidcTokenInfo> DoLoginAsync(CancellationToken cancellationToken)
 		{
 			// setup a local http server to listen for the result of the login
 			LoginResult? loginResult = null;
@@ -309,16 +449,16 @@ namespace EpicGames.OIDC
 				try
 				{
 #pragma warning disable CA2000 // Dispose objects before losing scope <-- FALSE POSITIVE?
-					using HttpListener http = new HttpListener();
+					using OidcHttpServer http = CreateHttpServer();
 #pragma warning restore CA2000 // Dispose objects before losing scope
 
 					// build the url the server should be hosted at
 					string prefix = $"{uri.Scheme}{Uri.SchemeDelimiter}{uri.Authority}/";
-					http.Prefixes.Add(prefix);
+					http.AddPrefixes(prefix);
 					http.Start();
 
 					OidcClientOptions options = await BuildClientOptionsAsync(uri, cancellationToken);
-					OidcClient oidcClient = new OidcClient(options);
+					OidcClient oidcClient = DoCreateOidcClient(options);
 					// generate the appropriate codes we need to login
 					AuthorizeState loginState = await oidcClient!.PrepareLoginAsync(cancellationToken: cancellationToken);
 
@@ -376,18 +516,6 @@ namespace EpicGames.OIDC
 				throw new LoginFailedException("Failed to login due to error: " + loginResult.Error, loginResult.ErrorDescription);
 			}
 
-   			if (String.IsNullOrEmpty(loginResult.RefreshToken))
-			{
-				throw new Exception("No refresh token was provided in response.");
-			}
-
-			_refreshToken = loginResult.RefreshToken;
-			_accessToken = loginResult.AccessToken;
-			_tokenExpiry = loginResult.AccessTokenExpiration;
-
-			_tokenStore.AddRefreshToken(_name, _refreshToken);
-			_tokenStore.Save();
-
 			return new OidcTokenInfo
 			{
 				RefreshToken = loginResult.RefreshToken,
@@ -396,11 +524,19 @@ namespace EpicGames.OIDC
 			};
 		}
 
-		private async Task<LoginResult> ProcessHttpRequestAsync(HttpListener http, AuthorizeState loginState, OidcClient oidcClient, CancellationToken cancellationToken)
+		public virtual OidcHttpServer CreateHttpServer()
+		{
+			return new OidcHttpServer();
+		}
+
+		public virtual OidcClient DoCreateOidcClient(OidcClientOptions options)
+		{
+			return new OidcClient(options);
+		}
+
+		private async Task<LoginResult> ProcessHttpRequestAsync(OidcHttpServer http, AuthorizeState loginState, OidcClient oidcClient, CancellationToken cancellationToken)
 		{
-			LoginResult? loginResult = null;
 			const int MaxAttempts = 5;
-			HttpListenerContext? context = null;
 			for (int i = 0; i < MaxAttempts; i++)
 			{
 				if (cancellationToken.IsCancellationRequested)
@@ -408,82 +544,76 @@ namespace EpicGames.OIDC
 					return new LoginResult("Login cancelled");
 				}
 
-				context = await http.GetContextAsync();
+				(OidcHttpServer.OidcHttpRequest request, OidcHttpServer.IOidcHttpResponse response) = await http.ProcessNextRequestAsync();
 
 				if (cancellationToken.IsCancellationRequested)
 				{
 					return new LoginResult("Login cancelled");
 				}
 
+				LoginResult loginResult;
 				string? responseData;
-				switch (context.Request.HttpMethod)
+				switch (request.HttpMethod)
 				{
 					case "GET":
-						responseData = context.Request.RawUrl;
+						responseData = request.RawUrl;
 
 						// parse the returned url for the tokens needed to complete the login
 						loginResult = await oidcClient!.ProcessResponseAsync(responseData, loginState, cancellationToken: cancellationToken);
 						break;
 					case "POST":
 						{
-							HttpListenerRequest request = context.Request;
 							if (request.ContentType != null && !request.ContentType.Equals("application/x-www-form-urlencoded", StringComparison.OrdinalIgnoreCase))
 							{
 								// we do not support url encoded return types
-								context.Response.StatusCode = 415;
+								response.StatusCode = 415;
 								throw new Exception("URL encoded responses not support");
 							}
 
 							// attempt to parse the body
 
 							// if there is no body we can not handle the post
-							if (!context.Request.HasEntityBody)
+							if (!request.HasBody)
 							{
-								context.Response.StatusCode = 415;
+								response.StatusCode = 415;
 								throw new Exception("Empty body not supported");
 							}
 
-							await using Stream body = request.InputStream;
-							using StreamReader reader = new StreamReader(body, request.ContentEncoding);
+							await using Stream body = request.Body;
+							using StreamReader reader = new StreamReader(body, request.BodyEncoding);
 							responseData = await reader.ReadToEndAsync(cancellationToken);
 
 							loginResult = await oidcClient!.ProcessResponseAsync(responseData, loginState, cancellationToken: cancellationToken);
 							break;
 						}
 					case "OPTIONS":
-						context.Response.StatusCode = 200;
-						context.Response.Close();
+						response.StatusCode = 200;
+						response.Close();
 						continue;
 					default:
 						// if we receive any other http method something is very odd. Tell them to use a different method.
-						context.Response.StatusCode = 415;
-						throw new Exception("Unsupported method used: " + context.Request.HttpMethod);
+						response.StatusCode = 415;
+						throw new Exception("Unsupported method used: " + request.HttpMethod);
 				}
 
-				break;
-			}
+				if (cancellationToken.IsCancellationRequested)
+				{
+					return new LoginResult("Login cancelled");
+				}
 
-			if (cancellationToken.IsCancellationRequested)
-			{
-				return new LoginResult("Login cancelled");
-			}
+				// generate a simple http page to show the user
+				string httpPage = loginResult.IsError ? GetFailurePage(loginResult, _genericErrorInformation) : GetSuccessPage();
+				byte[] buffer = Encoding.UTF8.GetBytes(httpPage);
+				response.ContentLength64 = buffer.Length;
+				Stream responseOutput = response.OutputStream;
+				await responseOutput.WriteAsync(buffer, 0, buffer.Length, cancellationToken);
+				responseOutput.Close();
 
-			if (context == null || loginResult == null)
-			{
-				throw new Exception("Context or loginResult not set");
+				return loginResult;
 			}
 
-			// generate a simple http page to show the user
-			HttpListenerResponse response = context.Response;
-			string httpPage = loginResult.IsError ? GetFailurePage(loginResult, _genericErrorInformation) : GetSuccessPage();
-			byte[] buffer = Encoding.UTF8.GetBytes(httpPage);
-			response.ContentLength64 = buffer.Length;
-			Stream responseOutput = response.OutputStream;
-			await responseOutput.WriteAsync(buffer, 0, buffer.Length, cancellationToken);
-			responseOutput.Close();
-
-			return loginResult;
-		}	
+			throw new Exception("Failed to process login after multiple attempts");
+		}
 
 		private static string GetSuccessPage()
 		{
@@ -507,7 +637,7 @@ namespace EpicGames.OIDC
 			return httpPageFailure;
 		}
 
-		private Process? OpenBrowser(string url)
+		public virtual Process? OpenBrowser(string url)
 		{
 			_logger?.LogDebug("Opening browser.");
 			if (RuntimeInformation.IsOSPlatform(OSPlatform.Windows))
@@ -684,7 +814,6 @@ namespace EpicGames.OIDC
 			_refreshToken = inRefreshToken;
 		}
 	}
-
 	public class LoginFailedException : Exception
 	{
 		public LoginFailedException(string message, string errorDescription) : base(message)
@@ -733,5 +862,8 @@ namespace EpicGames.OIDC
 
 		public string Scopes { get; set; } = "openid profile offline_access email";
 		public string? GenericErrorInformation { get; set; }
+		public bool UseDiscoveryDocument { get; set; } = true;
+		public string? ClientSecret { get; set; } = null;
+
 	}
 }
