Horde: Adding artifact type expiry information to dashboard config query

diff --git a/Engine/Source/Programs/Horde/HordeServer/Dashboard/DashboardController.cs b/Engine/Source/Programs/Horde/HordeServer/Dashboard/DashboardController.cs
--- a/Engine/Source/Programs/Horde/HordeServer/Dashboard/DashboardController.cs
+++ b/Engine/Source/Programs/Horde/HordeServer/Dashboard/DashboardController.cs
@@ -160,6 +160,8 @@ namespace HordeServer.Dashboard
 
 			dashboardConfigResponse.ArtifactTypes = _globalConfig.Value.Plugins.GetBuildConfig().ArtifactTypes.Select(a => a.Type.ToString()).ToList();
 
+			dashboardConfigResponse.ArtifactTypeInfo = _globalConfig.Value.Plugins.GetBuildConfig().ArtifactTypes.Select(a => new GetDashboardArtifactTypeResponse { Type = a.Type.ToString(), KeepDays = a.KeepDays}).ToList();
+
 			foreach (IPluginResponseFilter responseFilter in _responseFilters)
 			{
 				responseFilter.Apply(HttpContext, dashboardConfigResponse);
diff --git a/Engine/Source/Programs/Shared/EpicGames.Horde/Dashboard/DashboardMessages.cs b/Engine/Source/Programs/Shared/EpicGames.Horde/Dashboard/DashboardMessages.cs
--- a/Engine/Source/Programs/Shared/EpicGames.Horde/Dashboard/DashboardMessages.cs
+++ b/Engine/Source/Programs/Shared/EpicGames.Horde/Dashboard/DashboardMessages.cs
@@ -70,6 +70,27 @@ namespace EpicGames.Horde.Dashboard
 		/// </summary>
 		public List<string> ArtifactTypes { get; set; } = new List<string>();
 
+		/// <summary>
+		/// Configured artifact type info
+		/// </summary>
+		public List<GetDashboardArtifactTypeResponse> ArtifactTypeInfo { get; set; } = new List<GetDashboardArtifactTypeResponse>();
+
+	}
+
+	/// <summary>
+	/// Describes an artifact type
+	/// </summary>
+	public class GetDashboardArtifactTypeResponse
+	{
+		/// <summary>
+		/// Type of the artifact
+		/// </summary>
+		public string Type { get; set; } = "UnnamedArtifactType";
+
+		/// <summary>
+		/// Number of days to retain artifacts of this type
+		/// </summary>
+		public int? KeepDays { get; set; }
 	}
 
 	/// <summary>
