OidcToken - Added missing copyright to tests and fixed other analyze errors in oidc token

diff --git a/Engine/Source/Programs/OidcToken/OidcToken.Tests/OidcTokenTests.cs b/Engine/Source/Programs/OidcToken/OidcToken.Tests/OidcTokenTests.cs
--- a/Engine/Source/Programs/OidcToken/OidcToken.Tests/OidcTokenTests.cs
+++ b/Engine/Source/Programs/OidcToken/OidcToken.Tests/OidcTokenTests.cs
@@ -1,5 +1,6 @@
+// Copyright Epic Games, Inc. All Rights Reserved.
+
 using System.Diagnostics;
-using System.Net;
 using System.Net.Mime;
 using System.Text;
 using System.Text.Json;
@@ -105,7 +106,7 @@ namespace OidcToken.Tests
 				TokenResultFile? result = JsonSerializer.Deserialize<TokenResultFile>(fs);
 
 				Assert.AreEqual(mockAccessToken, result!.Token);
-				Assert.AreEqual(mockExpiryTime, result!.ExpiresAt);
+				Assert.AreEqual(mockExpiryTime, result.ExpiresAt);
 				tokenStoreMock.VerifyAll();
 			}
 			finally
@@ -175,8 +176,8 @@ namespace OidcToken.Tests
 					{
 						Mock<OidcHttpServer> httpMock = new Mock<OidcHttpServer>();
 
-						OidcHttpServer.OidcHttpRequest request = new OidcHttpServer.OidcHttpRequest("", MediaTypeNames.Application.FormUrlEncoded, "POST", new MemoryStream(Encoding.Default.GetBytes("input-object")), hasBody: true, Encoding.ASCII);
-						OidcHttpServer.NullOidcHttpResponse response = new OidcHttpServer.NullOidcHttpResponse(new MemoryStream());
+						OidcHttpRequest request = new OidcHttpRequest("", MediaTypeNames.Application.FormUrlEncoded, "POST", new MemoryStream(Encoding.Default.GetBytes("input-object")), hasBody: true, Encoding.ASCII);
+						NullOidcHttpResponse response = new NullOidcHttpResponse(new MemoryStream());
 
 						httpMock.Setup(listener => listener.ProcessNextRequestAsync()).ReturnsAsync(() => (request, response));
 
@@ -234,7 +235,7 @@ namespace OidcToken.Tests
 				Mock<OidcHttpServer> httpMock = new Mock<OidcHttpServer>();
 				// buffer to store the response in
 				byte[] responseBuffer = new byte[16 * 1024];
-				OidcHttpServer.NullOidcHttpResponse response = new OidcHttpServer.NullOidcHttpResponse(new MemoryStream(responseBuffer));
+				NullOidcHttpResponse response = new NullOidcHttpResponse(new MemoryStream(responseBuffer));
 
 				using IHost host = SetupMocks(cliOptions =>
 				{
@@ -272,7 +273,7 @@ namespace OidcToken.Tests
 
 					tokenClientMock.Setup(client => client.CreateHttpServer()).Returns(() =>
 					{
-						OidcHttpServer.OidcHttpRequest request = new OidcHttpServer.OidcHttpRequest("", MediaTypeNames.Application.FormUrlEncoded, "POST", new MemoryStream(Encoding.Default.GetBytes("input-object")), hasBody: true, Encoding.ASCII);
+						OidcHttpRequest request = new OidcHttpRequest("", MediaTypeNames.Application.FormUrlEncoded, "POST", new MemoryStream(Encoding.Default.GetBytes("input-object")), hasBody: true, Encoding.ASCII);
 						httpMock.Setup(listener => listener.ProcessNextRequestAsync()).ReturnsAsync(() => (request, response)).Verifiable(Times.Once);
 
 						return httpMock.Object;
@@ -297,7 +298,7 @@ namespace OidcToken.Tests
 
 				// this does not actually verify that a token is correctly allocated  as that requires reaching out to a idp, and we mock away all of those parts. This can only check that the parts we do are working as such these are just the same values as we specified in the mock
 				Assert.AreEqual(mockAccessToken, result!.Token);
-				Assert.AreEqual(mockExpiryTime, result!.ExpiresAt);
+				Assert.AreEqual(mockExpiryTime, result.ExpiresAt);
 
 				tokenStoreMock.VerifyAll();
 			}
diff --git a/Engine/Source/Programs/Shared/EpicGames.OIDC/InMemoryTokenStore.cs b/Engine/Source/Programs/Shared/EpicGames.OIDC/InMemoryTokenStore.cs
--- a/Engine/Source/Programs/Shared/EpicGames.OIDC/InMemoryTokenStore.cs
+++ b/Engine/Source/Programs/Shared/EpicGames.OIDC/InMemoryTokenStore.cs
@@ -38,6 +38,10 @@ public class InMemoryTokenStore : ITokenStore
 		// No need - data exists only in-memory
 	}
 
+	/// <summary>
+	/// Dispose method
+	/// </summary>
+	/// <param name="disposing"></param>
 	protected virtual void Dispose(bool disposing)
 	{
 		if (disposing)
@@ -46,6 +50,7 @@ public class InMemoryTokenStore : ITokenStore
 		}
 	}
 
+	/// <inheritdoc/>
 	public void Dispose()
 	{
 		Dispose(true);
diff --git a/Engine/Source/Programs/Shared/EpicGames.OIDC/OidcTokenManager.cs b/Engine/Source/Programs/Shared/EpicGames.OIDC/OidcTokenManager.cs
--- a/Engine/Source/Programs/Shared/EpicGames.OIDC/OidcTokenManager.cs
+++ b/Engine/Source/Programs/Shared/EpicGames.OIDC/OidcTokenManager.cs
@@ -62,7 +62,7 @@ namespace EpicGames.OIDC
 			}
 		}
 
-		public OidcTokenManager(IServiceProvider provider, IOptionsMonitor<OidcTokenOptions> settings, ITokenStore tokenStore, IOidcTokenClientFactory tokenClientFactory, List<string>? allowedProviders = null)
+		public OidcTokenManager(IOptionsMonitor<OidcTokenOptions> settings, ITokenStore tokenStore, IOidcTokenClientFactory tokenClientFactory, List<string>? allowedProviders = null)
 			: this(tokenStore)
 		{
 			Dictionary<string, string> refreshTokens = new Dictionary<string, string>();
@@ -223,86 +223,135 @@ namespace EpicGames.OIDC
 	}
 
 	/// <summary>
-	/// Wrapper around HttpListener to make it possible to inject a mock and thus test the Oidc flow without starting a http server
+	/// The data read from a http request into the oidc client, mostly a wrapper around the http.listener context
 	/// </summary>
-	public class OidcHttpServer : IDisposable
+	public class OidcHttpRequest
 	{
-		public class OidcHttpRequest
-		{
-			public OidcHttpRequest(HttpListenerContext context)
-			{
-				HttpListenerRequest request = context.Request;
-				HasBody = request.HasEntityBody;
-				HttpMethod = request.HttpMethod;
-				RawUrl = request.RawUrl;
-				ContentType = request.ContentType;
-				Body = request.InputStream;
-				BodyEncoding = request.ContentEncoding;
-			}
-
-			public OidcHttpRequest(string rawUrl, string contentType, string httpMethod, Stream body, bool hasBody, Encoding bodyEncoding)
-			{
-				RawUrl = rawUrl;
-				ContentType = contentType;
-				HttpMethod = httpMethod;
-				Body = body;
-				HasBody = hasBody;
-				BodyEncoding = bodyEncoding;
-			}
+		public OidcHttpRequest(HttpListenerContext context)
+		{
+			HttpListenerRequest request = context.Request;
+			HasBody = request.HasEntityBody;
+			HttpMethod = request.HttpMethod;
+			RawUrl = request.RawUrl;
+			ContentType = request.ContentType;
+			Body = request.InputStream;
+			BodyEncoding = request.ContentEncoding;
+		}
+
+#pragma warning disable CA1054
+		public OidcHttpRequest(string rawUrl, string contentType, string httpMethod, Stream body, bool hasBody, Encoding bodyEncoding)
+#pragma warning restore CA1054
+		{
+			RawUrl = rawUrl;
+			ContentType = contentType;
+			HttpMethod = httpMethod;
+			Body = body;
+			HasBody = hasBody;
+			BodyEncoding = bodyEncoding;
+		}
+
+		public bool HasBody { get; init; }
+		public string HttpMethod { get; init; }
+#pragma warning disable CA1056 // This does not need to be a Uri
+		public string? RawUrl { get; init; }
+#pragma warning restore CA1056
+		public string? ContentType { get; init; }
+		public Stream Body { get; init;}
+		public Encoding BodyEncoding { get; init; }
+	}
 
+	/// <summary>
+	/// Interface for the response writing that happens from the local http server, usually wraps a http.listener response
+	/// </summary>
+	public interface IOidcHttpResponse : IDisposable
+	{
+		/// <summary>
+		/// The content length header value
+		/// </summary>
+		public long ContentLength64 { get; set; }
+		/// <summary>
+		/// The status code header value
+		/// </summary>
+		public int StatusCode { get; set; }
 
-			public bool HasBody { get; init; }
-			public string HttpMethod { get; init; }
-			public string? RawUrl { get; init; }
-			public string? ContentType { get; init; }
-			public Stream Body { get; init;}
-			public Encoding BodyEncoding { get; init; }
-		}
+		/// <summary>
+		/// The body of the response
+		/// </summary>
+		public Stream OutputStream { get; }
 
-		public interface IOidcHttpResponse : IDisposable
+		/// <summary>
+		/// Used to explicitly close the response when finished
+		/// </summary>
+		public void Close()
 		{
-			public long ContentLength64 { get; set; }
-			public int StatusCode { get; set; }
-			public Stream OutputStream { get; }
-			public void Close()
-			{
-				Dispose();
-			}
+			Dispose();
 		}
+	}
 
-		// Implementation of a response writer for http listener
-		public class OidcHttpResponse(HttpListenerContext context) : IOidcHttpResponse
-		{
-			private readonly HttpListenerResponse _response = context.Response;
+	/// <summary>
+	/// Implementation of a response writer for http listener
+	/// </summary>
+	/// <param name="context"></param>
+	public sealed class OidcHttpResponse(HttpListenerContext context) : IOidcHttpResponse
+	{
+		private readonly HttpListenerResponse _response = context.Response;
 
-			public long ContentLength64 { get => _response.ContentLength64; set => _response.ContentLength64 = value; }
-			public int StatusCode  { get => _response.StatusCode; set => _response.StatusCode = value; }
-			public Stream OutputStream => _response.OutputStream;
+		/// <inheritdoc/>
+		public long ContentLength64 { get => _response.ContentLength64; set => _response.ContentLength64 = value; }
+		/// <inheritdoc/>
+		public int StatusCode  { get => _response.StatusCode; set => _response.StatusCode = value; }
+		/// <inheritdoc/>
+		public Stream OutputStream => _response.OutputStream;
 
-			public void Close()
-			{
-				_response.Close();
-			}
+		/// <inheritdoc/>
+		public void Close()
+		{
+			_response.Close();
+		}
 
-			public void Dispose()
-			{
-				((IDisposable)_response).Dispose();
-			}
+		/// <inheritdoc/>
+		public void Dispose()
+		{
+			((IDisposable)_response).Dispose();
 		}
+	}
 
-		// a response writer that does not actually use the response in any way, just buffering it internally
-		public class NullOidcHttpResponse(Stream outputStream) : IOidcHttpResponse
+	// a response writer that does not actually use the response in any way, just buffering it internally
+	public class NullOidcHttpResponse(Stream outputStream) : IOidcHttpResponse
+	{
+		/// <summary>
+		/// Dispose method
+		/// </summary>
+		/// <param name="disposing"></param>
+		protected virtual void Dispose(bool disposing)
 		{
-			public void Dispose()
+			if (disposing)
 			{
-				// nothing to dispose
+				// no resources to release
 			}
+		}
 
-			public long ContentLength64 { get; set; }
-			public int StatusCode { get; set; }
-			public Stream OutputStream { get; } = outputStream;
+		/// <inheritdoc/>
+		public void Dispose()
+		{
+			Dispose(true);
+			GC.SuppressFinalize(this);
 		}
 
+		/// <inheritdoc/>
+		public long ContentLength64 { get; set; }
+		/// <inheritdoc/>
+		public int StatusCode { get; set; }
+		/// <inheritdoc/>
+		public Stream OutputStream { get; } = outputStream;
+	}
+
+	/// <summary>
+	/// Wrapper around HttpListener to make it possible to inject a mock and thus test the Oidc flow without starting a http server
+	/// </summary>
+	// ReSharper disable once ClassWithVirtualMembersNeverInherited.Global - overriden in test mocks
+	public class OidcHttpServer : IDisposable
+	{
 		private readonly HttpListener _http = new();
 
 		public virtual void Start()
@@ -310,12 +359,28 @@ namespace EpicGames.OIDC
 			_http.Start();
 		}
 
-		public virtual void Dispose()
+		/// <summary>
+		/// Dispose method
+		/// </summary>
+		/// <param name="disposing"></param>
+		protected virtual void Dispose(bool disposing)
+		{
+			if (disposing)
+			{
+				((IDisposable)_http).Dispose();
+			}
+		}
+
+		/// <inheritdoc/>
+		public void Dispose()
 		{
-			((IDisposable)_http).Dispose();
+			Dispose(true);
+			GC.SuppressFinalize(this);
 		}
 
+#pragma warning disable CA1716 // we want to call this Stop to maintain an interface that is very similar to http.listener
 		public virtual void Stop()
+#pragma warning restore CA1716
 		{
 			_http.Stop();
 		}
@@ -544,7 +609,7 @@ namespace EpicGames.OIDC
 					return new LoginResult("Login cancelled");
 				}
 
-				(OidcHttpServer.OidcHttpRequest request, OidcHttpServer.IOidcHttpResponse response) = await http.ProcessNextRequestAsync();
+				(OidcHttpRequest request, IOidcHttpResponse response) = await http.ProcessNextRequestAsync();
 
 				if (cancellationToken.IsCancellationRequested)
 				{
@@ -637,7 +702,9 @@ namespace EpicGames.OIDC
 			return httpPageFailure;
 		}
 
+#pragma warning disable CA1054
 		public virtual Process? OpenBrowser(string url)
+#pragma warning restore CA1054
 		{
 			_logger?.LogDebug("Opening browser.");
 			if (RuntimeInformation.IsOSPlatform(OSPlatform.Windows))
