UnrealBuildTool: Fix compiler warnings

diff --git a/Engine/Source/Programs/Shared/UnrealBuildTool.Tests/CppCompileWarningsTests.cs b/Engine/Source/Programs/Shared/UnrealBuildTool.Tests/CppCompileWarningsTests.cs
--- a/Engine/Source/Programs/Shared/UnrealBuildTool.Tests/CppCompileWarningsTests.cs
+++ b/Engine/Source/Programs/Shared/UnrealBuildTool.Tests/CppCompileWarningsTests.cs
@@ -251,7 +251,7 @@ namespace UnrealBuildTool.Tests
 			using (SafeTestFile testTargetTestFile = new SafeTestFile(File.ReadAllText(Path.Combine(testStubsDirectory.FullName, "TargetStubs", "TestStub.Target.ubttest")), "TestStub.Target.cs", testSourceDirectory.TemporaryDirectory))
 			using (SafeTestFile testUprojectTestFile = new SafeTestFile(File.ReadAllText(Path.Combine(testStubsDirectory.FullName, "UprojectStubs", "TestStub.uproject.ubttest")), "TestStub.uproject", testRootDirectory.TemporaryDirectory))
 			{
-				CppCompileEnvironment newCompileEnvironment = new CppCompileEnvironment(UnrealTargetPlatform.Win64, CppConfiguration.Development, new UnrealArchitectures(UnrealArch.X64), null);
+				CppCompileEnvironment newCompileEnvironment = new CppCompileEnvironment(UnrealTargetPlatform.Win64, CppConfiguration.Development, new UnrealArchitectures(UnrealArch.X64), null!);
 
 				FileReference testUproject = new FileReference(testUprojectTestFile.TemporaryFile);
 
@@ -260,7 +260,7 @@ namespace UnrealBuildTool.Tests
 				TargetRules testTargetRules = generatedTestRulesAssembly.CreateTargetRules("TestStubEditor", UnrealTargetPlatform.Win64, UnrealTargetConfiguration.Development, null, testUproject, null, new TestLogger());
 				CppCompileWarnings baseWarnings = new CppCompileWarnings(testTargetRules, null);
 
-				IEnumerable<string> args = null;
+				IEnumerable<string> args = null!;
 
 				// No args, base check
 				{
diff --git a/Engine/Source/Programs/Shared/UnrealBuildTool.Tests/TestUtilities/StaticTargetRulesInitializer.cs b/Engine/Source/Programs/Shared/UnrealBuildTool.Tests/TestUtilities/StaticTargetRulesInitializer.cs
--- a/Engine/Source/Programs/Shared/UnrealBuildTool.Tests/TestUtilities/StaticTargetRulesInitializer.cs
+++ b/Engine/Source/Programs/Shared/UnrealBuildTool.Tests/TestUtilities/StaticTargetRulesInitializer.cs
@@ -8,7 +8,7 @@ using Microsoft.VisualStudio.TestTools.UnitTesting;
 namespace UnrealBuildTool.Tests.TestUtilities
 {
 	[TestClass]
-	public class StaticTargetRulesInitializer
+	public static class StaticTargetRulesInitializer
 	{
 		internal const string UBTTestFolderPrefix = "UBT-Test-";
 		internal static readonly string UBTTestFolderRoot = Path.GetTempPath();
diff --git a/Engine/Source/Programs/Shared/UnrealBuildTool.Tests/TestUtilities/TestToolChain.cs b/Engine/Source/Programs/Shared/UnrealBuildTool.Tests/TestUtilities/TestToolChain.cs
--- a/Engine/Source/Programs/Shared/UnrealBuildTool.Tests/TestUtilities/TestToolChain.cs
+++ b/Engine/Source/Programs/Shared/UnrealBuildTool.Tests/TestUtilities/TestToolChain.cs
@@ -20,17 +20,17 @@ namespace UnrealBuildTool.Tests.TestUtilities
 
 		protected override CPPOutput CompileCPPFiles(CppCompileEnvironment CompileEnvironment, IEnumerable<FileItem> InputFiles, DirectoryReference OutputDir, string ModuleName, IActionGraphBuilder Graph)
 		{
-			return null;
+			return null!;
 		}
 
 		protected override CPPOutput CompileISPCFiles(CppCompileEnvironment CompileEnvironment, IEnumerable<FileItem> InputFiles, DirectoryReference OutputDir, IActionGraphBuilder Graph)
 		{
-			return null;
+			return null!;
 		}
 
 		protected override CPPOutput GenerateISPCHeaders(CppCompileEnvironment CompileEnvironment, IEnumerable<FileItem> InputFiles, DirectoryReference OutputDir, IActionGraphBuilder Graph)
 		{
-			return null;
+			return null!;
 		}
 	}
 }
