Improve where optional attributes can appear in enum, class, and struct definitions.

diff --git a/Engine/Source/Programs/Shared/EpicGames.UHT/Parsers/UhtClassParser.cs b/Engine/Source/Programs/Shared/EpicGames.UHT/Parsers/UhtClassParser.cs
--- a/Engine/Source/Programs/Shared/EpicGames.UHT/Parsers/UhtClassParser.cs
+++ b/Engine/Source/Programs/Shared/EpicGames.UHT/Parsers/UhtClassParser.cs
@@ -74,8 +74,8 @@ namespace EpicGames.UHT.Parsers
 					topScope.AddFormattedCommentsAsTooltipMetaData();
 
 					// Consume the class specifier
+					topScope.TokenReader.OptionalAttributes(false);
 					topScope.TokenReader.Require("class");
-
 					topScope.TokenReader.OptionalAttributes(true);
 
 					// Read the class name and possible API macro name
diff --git a/Engine/Source/Programs/Shared/EpicGames.UHT/Parsers/UhtEnumParser.cs b/Engine/Source/Programs/Shared/EpicGames.UHT/Parsers/UhtEnumParser.cs
--- a/Engine/Source/Programs/Shared/EpicGames.UHT/Parsers/UhtEnumParser.cs
+++ b/Engine/Source/Programs/Shared/EpicGames.UHT/Parsers/UhtEnumParser.cs
@@ -45,13 +45,17 @@ namespace EpicGames.UHT.Parsers
 					{
 						enumObject.CppForm = UhtEnumCppForm.Namespaced;
 					}
-					else if (topScope.TokenReader.TryOptional("enum"))
-					{
-						enumObject.CppForm = topScope.TokenReader.TryOptional("class") || topScope.TokenReader.TryOptional("struct") ? UhtEnumCppForm.EnumClass : UhtEnumCppForm.Regular;
-					}
 					else
 					{
-						throw new UhtTokenException(topScope.TokenReader, topScope.TokenReader.PeekToken(), null);
+						topScope.TokenReader.OptionalAttributes(false);
+						if (topScope.TokenReader.TryOptional("enum"))
+						{
+							enumObject.CppForm = topScope.TokenReader.TryOptional("class") || topScope.TokenReader.TryOptional("struct") ? UhtEnumCppForm.EnumClass : UhtEnumCppForm.Regular;
+						}
+						else
+						{
+							throw new UhtTokenException(topScope.TokenReader, topScope.TokenReader.PeekToken(), null);
+						}
 					}
 
 					topScope.TokenReader.OptionalAttributes(false);
@@ -101,6 +105,7 @@ namespace EpicGames.UHT.Parsers
 					{
 						case UhtEnumCppForm.Namespaced:
 							// Now handle the inner true enum portion
+							topScope.TokenReader.OptionalAttributes(false);
 							topScope.TokenReader.Require("enum");
 							topScope.TokenReader.OptionalAttributes(true);
 
diff --git a/Engine/Source/Programs/Shared/EpicGames.UHT/Parsers/UhtScriptStructParser.cs b/Engine/Source/Programs/Shared/EpicGames.UHT/Parsers/UhtScriptStructParser.cs
--- a/Engine/Source/Programs/Shared/EpicGames.UHT/Parsers/UhtScriptStructParser.cs
+++ b/Engine/Source/Programs/Shared/EpicGames.UHT/Parsers/UhtScriptStructParser.cs
@@ -79,8 +79,8 @@ namespace EpicGames.UHT.Parsers
 					specifiers.ParseSpecifiers();
 
 					// Consume the struct specifier
+					topScope.TokenReader.OptionalAttributes(false);
 					topScope.TokenReader.Require("struct");
-
 					topScope.TokenReader.OptionalAttributes(true);
 
 					// Read the struct name and possible API macro name
