Fix UBT issue where UBT plugins were being built concurrently.

diff --git a/Engine/Source/Programs/Shared/EpicGames.Build/System/EnumeratePlugins.cs b/Engine/Source/Programs/Shared/EpicGames.Build/System/EnumeratePlugins.cs
--- a/Engine/Source/Programs/Shared/EpicGames.Build/System/EnumeratePlugins.cs
+++ b/Engine/Source/Programs/Shared/EpicGames.Build/System/EnumeratePlugins.cs
@@ -145,32 +145,36 @@ namespace UnrealBuildBase
 		public static bool BuildUbtPlugins(FileReference? ProjectFile, IEnumerable<FileReference> FoundPlugins,
 			IEnumerable<string>? DefineConstants, ILogger Logger, out (FileReference ProjectFile, FileReference TargetAssembly)[]? Plugins)
 		{
-			bool bBuildSuccess = true;
-			Plugins = null;
+			lock (UbtPluginFileCache)
+			{
 
-			// Collect the list of plugins already build
-			List<FileReference> Built = FoundPlugins.Where(P => UbtPluginFileCache.ContainsKey(P)).ToList();
+				bool bBuildSuccess = true;
+				Plugins = null;
 
-			// Collect the list of plugins not build
-			List<FileReference> NotBuilt = FoundPlugins.Where(P => !UbtPluginFileCache.ContainsKey(P)).ToList();
+				// Collect the list of plugins already build
+				List<FileReference> Built = FoundPlugins.Where(P => UbtPluginFileCache.ContainsKey(P)).ToList();
 
-			// If we have anything left to build, then build it
-			if (NotBuilt.Count > 0)
-			{
-				if (Unreal.IsEngineInstalled())
-				{
-					bBuildSuccess = BuildUbtPluginsInternal(PluginSet.EngineOnly, ProjectFile, Built, NotBuilt, DefineConstants, Logger);
-					bBuildSuccess |= BuildUbtPluginsInternal(PluginSet.ProjectOnly, ProjectFile, Built, NotBuilt, DefineConstants, Logger);
-				}
-				else
+				// Collect the list of plugins not build
+				List<FileReference> NotBuilt = FoundPlugins.Where(P => !UbtPluginFileCache.ContainsKey(P)).ToList();
+
+				// If we have anything left to build, then build it
+				if (NotBuilt.Count > 0)
 				{
-					bBuildSuccess =
-						BuildUbtPluginsInternal(PluginSet.Both, ProjectFile, Built, NotBuilt, DefineConstants, Logger);
+					if (Unreal.IsEngineInstalled())
+					{
+						bBuildSuccess = BuildUbtPluginsInternal(PluginSet.EngineOnly, ProjectFile, Built, NotBuilt, DefineConstants, Logger);
+						bBuildSuccess |= BuildUbtPluginsInternal(PluginSet.ProjectOnly, ProjectFile, Built, NotBuilt, DefineConstants, Logger);
+					}
+					else
+					{
+						bBuildSuccess =
+							BuildUbtPluginsInternal(PluginSet.Both, ProjectFile, Built, NotBuilt, DefineConstants, Logger);
+					}
 				}
-			}
 
-			Plugins = [.. Built.Where(P => UbtPluginFileCache.ContainsKey(P)).Select(P => (P, UbtPluginFileCache[P])).OrderBy(X => X.P.FullName)];
-			return bBuildSuccess;
+				Plugins = [.. Built.Where(P => UbtPluginFileCache.ContainsKey(P)).Select(P => (P, UbtPluginFileCache[P])).OrderBy(X => X.P.FullName)];
+				return bBuildSuccess;
+			}
 		}
 
 		private enum PluginSet
