Fixed incorrect json serialization of cb fields that had fields without names (mostly encountered when serializing an array).

diff --git a/Engine/Source/Programs/Shared/EpicGames.Serialization/CbField.cs b/Engine/Source/Programs/Shared/EpicGames.Serialization/CbField.cs
--- a/Engine/Source/Programs/Shared/EpicGames.Serialization/CbField.cs
+++ b/Engine/Source/Programs/Shared/EpicGames.Serialization/CbField.cs
@@ -2675,7 +2675,16 @@ namespace EpicGames.Serialization
 			}
 			else if (field.IsArray())
 			{
-				writer.WriteStartArray(field.Name.Span);
+				if (field._nameLen != 0)
+				{
+					writer.WriteStartArray(field.Name.Span);
+				}
+				else
+				{
+					writer.WriteStartArray();
+				}
+
+
 				CbArray array = field.AsArray();
 				foreach (CbField objectField in array)
 				{
@@ -2687,36 +2696,98 @@ namespace EpicGames.Serialization
 			{
 				if (field.GetType() == CbFieldType.IntegerNegative)
 				{
-					writer.WriteNumber(field.Name.Span, -field.AsInt64());
+					if (field._nameLen != 0)
+					{
+						writer.WriteNumber(field.Name.Span, -field.AsInt64());
+					}
+					else
+					{
+						
+						writer.WriteNumberValue(-field.AsInt64());
+					}
 				}
 				else
 				{
-					writer.WriteNumber(field.Name.Span, field.AsUInt64());
+					if (field._nameLen != 0)
+					{
+						writer.WriteNumber(field.Name.Span, field.AsUInt64());
+					}
+					else
+					{
+						
+						writer.WriteNumberValue(field.AsUInt64());
+					}
 				}
 			}
 			else if (field.IsBool())
 			{
-				writer.WriteBoolean(field.Name.Span, field.AsBool());
+				if (field._nameLen != 0)
+				{
+					writer.WriteBoolean(field.Name.Span, field.AsBool());
+				}
+				else
+				{
+						
+					writer.WriteBooleanValue(field.AsBool());
+				}
 			}
 			else if (field.IsNull())
 			{
-				writer.WriteNullValue();
+				if (field._nameLen != 0)
+				{
+					writer.WriteNull(field.Name.Span);
+				}
+				else
+				{
+					writer.WriteNullValue();	
+				}
 			}
 			else if (field.IsDateTime())
 			{
-				writer.WriteString(field.Name.Span, field.AsDateTime());
+				if (field._nameLen != 0)
+				{
+					writer.WriteString(field.Name.Span, field.AsDateTime());
+				}
+				else
+				{
+					writer.WriteStringValue(field.AsDateTime());
+				}
 			}
 			else if (field.IsHash())
 			{
-				writer.WriteString(field.Name.Span, field.AsHash().ToUtf8String().Span);
+				ReadOnlySpan<byte> stringValue = field.AsHash().ToUtf8String().Span;
+				if (field._nameLen != 0)
+				{
+					writer.WriteString(field.Name.Span, stringValue);
+				}
+				else
+				{
+					writer.WriteStringValue(stringValue);
+				}
 			}
 			else if (field.IsString())
 			{
-				writer.WriteString(field.Name.Span, field.AsUtf8String().Span);
+				ReadOnlySpan<byte> stringValue = field.AsUtf8String().Span;
+				if (field._nameLen != 0)
+				{
+					writer.WriteString(field.Name.Span, stringValue);
+				}
+				else
+				{
+					writer.WriteStringValue(stringValue);
+				}
 			}
 			else if (field.IsObjectId())
 			{
-				writer.WriteString(field.Name.Span, field.AsObjectId().ToUtf8String());
+				ReadOnlySpan<byte> stringValue = field.AsObjectId().ToUtf8String().Span;
+				if (field._nameLen != 0)
+				{
+					writer.WriteString(field.Name.Span, stringValue);
+				}
+				else
+				{
+					writer.WriteStringValue(stringValue);
+				}
 			}
 			else
 			{
