Remove very old UHT output code that places a space between ">>" in template.

#rnx

diff --git a/Engine/Source/Programs/Shared/EpicGames.UHT/Types/Properties/UhtArrayProperty.cs b/Engine/Source/Programs/Shared/EpicGames.UHT/Types/Properties/UhtArrayProperty.cs
--- a/Engine/Source/Programs/Shared/EpicGames.UHT/Types/Properties/UhtArrayProperty.cs
+++ b/Engine/Source/Programs/Shared/EpicGames.UHT/Types/Properties/UhtArrayProperty.cs
@@ -121,13 +121,7 @@ namespace EpicGames.UHT.Types
 					break;
 
 				default:
-					builder.Append("TArray<").AppendPropertyText(ValueProperty, textType, true);
-					if (builder[^1] == '>')
-					{
-						// if our internal property type is a template class, add a space between the closing brackets b/c VS.NET cannot parse this correctly
-						builder.Append(' ');
-					}
-					builder.Append('>');
+					builder.Append("TArray<").AppendPropertyText(ValueProperty, textType, true).Append('>');
 					break;
 			}
 			return builder;
diff --git a/Engine/Source/Programs/Shared/EpicGames.UHT/Types/Properties/UhtMapProperty.cs b/Engine/Source/Programs/Shared/EpicGames.UHT/Types/Properties/UhtMapProperty.cs
--- a/Engine/Source/Programs/Shared/EpicGames.UHT/Types/Properties/UhtMapProperty.cs
+++ b/Engine/Source/Programs/Shared/EpicGames.UHT/Types/Properties/UhtMapProperty.cs
@@ -150,19 +150,7 @@ namespace EpicGames.UHT.Types
 					break;
 
 				default:
-					builder.Append("TMap<").AppendPropertyText(KeyProperty, textType, true);
-					if (builder[^1] == '>')
-					{
-						// if our internal property type is a template class, add a space between the closing brackets b/c VS.NET cannot parse this correctly
-						builder.Append(' ');
-					}
-					builder.Append(',').AppendPropertyText(ValueProperty, textType, true);
-					if (builder[^1] == '>')
-					{
-						// if our internal property type is a template class, add a space between the closing brackets b/c VS.NET cannot parse this correctly
-						builder.Append(' ');
-					}
-					builder.Append('>');
+					builder.Append("TMap<").AppendPropertyText(KeyProperty, textType, true).Append(',').AppendPropertyText(ValueProperty, textType, true).Append('>');
 					break;
 			}
 			return builder;
diff --git a/Engine/Source/Programs/Shared/EpicGames.UHT/Types/Properties/UhtOptionalProperty.cs b/Engine/Source/Programs/Shared/EpicGames.UHT/Types/Properties/UhtOptionalProperty.cs
--- a/Engine/Source/Programs/Shared/EpicGames.UHT/Types/Properties/UhtOptionalProperty.cs
+++ b/Engine/Source/Programs/Shared/EpicGames.UHT/Types/Properties/UhtOptionalProperty.cs
@@ -121,13 +121,7 @@ namespace EpicGames.UHT.Types
 					break;
 
 				default:
-					builder.Append("TOptional<").AppendPropertyText(ValueProperty, textType, true);
-					if (builder[^1] == '>')
-					{
-						// if our internal property type is a template class, add a space between the closing brackets b/c VS.NET cannot parse this correctly
-						builder.Append(' ');
-					}
-					builder.Append('>');
+					builder.Append("TOptional<").AppendPropertyText(ValueProperty, textType, true).Append('>');
 					break;
 			}
 			return builder;
diff --git a/Engine/Source/Programs/Shared/EpicGames.UHT/Types/Properties/UhtSetProperty.cs b/Engine/Source/Programs/Shared/EpicGames.UHT/Types/Properties/UhtSetProperty.cs
--- a/Engine/Source/Programs/Shared/EpicGames.UHT/Types/Properties/UhtSetProperty.cs
+++ b/Engine/Source/Programs/Shared/EpicGames.UHT/Types/Properties/UhtSetProperty.cs
@@ -114,13 +114,7 @@ namespace EpicGames.UHT.Types
 					break;
 
 				default:
-					builder.Append("TSet<").AppendPropertyText(ValueProperty, textType, true);
-					if (builder[^1] == '>')
-					{
-						// if our internal property type is a template class, add a space between the closing brackets b/c VS.NET cannot parse this correctly
-						builder.Append(' ');
-					}
-					builder.Append('>');
+					builder.Append("TSet<").AppendPropertyText(ValueProperty, textType, true).Append('>');
 					break;
 			}
 			return builder;
