APT: fix Linux runs.

diff --git a/Engine/Source/Programs/Shared/EpicGames.Build/System/Rules.cs b/Engine/Source/Programs/Shared/EpicGames.Build/System/Rules.cs
--- a/Engine/Source/Programs/Shared/EpicGames.Build/System/Rules.cs
+++ b/Engine/Source/Programs/Shared/EpicGames.Build/System/Rules.cs
@@ -93,6 +93,7 @@ namespace UnrealBuildBase
 				{
 					Folders.Add(DirectoryReference.Combine(PluginFile.Directory, "Source"));
 					Folders.Add(DirectoryReference.Combine(PluginFile.Directory, "Tests"));
+					Folders.Add(DirectoryReference.Combine(PluginFile.Directory, "Build"));	// this is to pick up plugin's C# automation projects				
 				}
 			}
 
