APT: Delete code references to APT files in other C# projects and re-add in Rules.

- It seems that preflighting such CLs cannot quite catch the issues because unshelving (i.e. deleting the files) happens at a different point and the warning is still produced.

diff --git a/Engine/Plugins/Performance/AutomatedPerfTesting/Resources/TemplatesForHookup/PROJECT.AutoPerfTest.csproj b/Engine/Plugins/Performance/AutomatedPerfTesting/Resources/TemplatesForHookup/PROJECT.AutoPerfTest.csproj
deleted file mode 100644
diff --git a/Engine/Source/Programs/Shared/EpicGames.Build/System/Rules.cs b/Engine/Source/Programs/Shared/EpicGames.Build/System/Rules.cs
--- a/Engine/Source/Programs/Shared/EpicGames.Build/System/Rules.cs
+++ b/Engine/Source/Programs/Shared/EpicGames.Build/System/Rules.cs
@@ -93,6 +93,7 @@ namespace UnrealBuildBase
 				{
 					Folders.Add(DirectoryReference.Combine(PluginFile.Directory, "Source"));
 					Folders.Add(DirectoryReference.Combine(PluginFile.Directory, "Tests"));
+					Folders.Add(DirectoryReference.Combine(PluginFile.Directory, "Build"));	// this is to pick up plugin's C# automation projects					
 				}
 			}
 
