Do not generate exported template specializations for StaticClass and StaticStruct, just use a normal template implementation

diff --git a/Engine/Source/Programs/Shared/EpicGames.UHT/Exporters/CodeGen/UhtHeaderCodeGeneratorCppFile.cs b/Engine/Source/Programs/Shared/EpicGames.UHT/Exporters/CodeGen/UhtHeaderCodeGeneratorCppFile.cs
--- a/Engine/Source/Programs/Shared/EpicGames.UHT/Exporters/CodeGen/UhtHeaderCodeGeneratorCppFile.cs
+++ b/Engine/Source/Programs/Shared/EpicGames.UHT/Exporters/CodeGen/UhtHeaderCodeGeneratorCppFile.cs
@@ -683,12 +683,6 @@ namespace EpicGames.UHT.Exporters.CodeGen
 				builder.Append("\treturn ").Append(registrationName).Append(".OuterSingleton;\r\n");
 				builder.Append("}\r\n");
 
-				// Generate the StaticStruct specialization
-				builder.Append("template<> ").Append(Module.Api).Append("UScriptStruct* StaticStruct<").Append(scriptStruct.SourceName).Append(">()\r\n");
-				builder.Append("{\r\n");
-				builder.Append("\treturn ").Append(scriptStruct.SourceName).Append("::StaticStruct();\r\n");
-				builder.Append("}\r\n");
-
 				// Inject implementation needed to support auto bindings of fast arrays
 				if (ObjectInfos[scriptStruct.ObjectTypeIndex].FastArrayProperty != null)
 				{
@@ -1924,11 +1918,6 @@ namespace EpicGames.UHT.Exporters.CodeGen
 				}
 			}
 
-			builder.Append("template<> ").Append(Module.Api).Append("UClass* StaticClass<").Append(classObj.SourceName).Append(">()\r\n");
-			builder.Append("{\r\n");
-			builder.Append("\treturn ").Append(classObj.SourceName).Append("::StaticClass();\r\n");
-			builder.Append("}\r\n");
-
 			if (classObj.ClassExportFlags.HasAnyFlags(UhtClassExportFlags.SelfHasReplicatedProperties))
 			{
 				builder.Append("void ").Append(classObj.SourceName).Append("::ValidateGeneratedRepEnums(const TArray<struct FRepRecord>& ClassReps) const\r\n");
diff --git a/Engine/Source/Programs/Shared/EpicGames.UHT/Exporters/CodeGen/UhtHeaderCodeGeneratorHFile.cs b/Engine/Source/Programs/Shared/EpicGames.UHT/Exporters/CodeGen/UhtHeaderCodeGeneratorHFile.cs
--- a/Engine/Source/Programs/Shared/EpicGames.UHT/Exporters/CodeGen/UhtHeaderCodeGeneratorHFile.cs
+++ b/Engine/Source/Programs/Shared/EpicGames.UHT/Exporters/CodeGen/UhtHeaderCodeGeneratorHFile.cs
@@ -594,9 +594,8 @@ namespace EpicGames.UHT.Exporters.CodeGen
 					Session.Tables!.CodeGeneratorInjectorTable.Inject(builder, scriptStruct, UhtCodeGeneratorInjectionLocation.GeneratedMacro);
 				}
 
-				// Forward declare the StaticStruct specialization in the header
-				builder.Append("template<> ").Append(Module.Api).Append("UScriptStruct* StaticStruct<struct ").Append(scriptStruct.SourceName).Append(">();\r\n");
-				builder.Append("\r\n");
+				// For back compat we still need to forward declare the struct here
+				builder.Append("struct ").Append(scriptStruct.SourceName).Append(";\r\n");
 				
 				Session.Tables!.CodeGeneratorInjectorTable.Inject(builder, scriptStruct, UhtCodeGeneratorInjectionLocation.Header);
 
@@ -835,8 +834,8 @@ namespace EpicGames.UHT.Exporters.CodeGen
 					sparseScriptStructs.Any(), sparseProperties, getterSetterProperties, hasCallbacks, usesLegacy ? "GENERATED_UCLASS_BODY" : null);
 			}
 
-			// Forward declare the StaticClass specialization in the header
-			builder.Append("template<> ").Append(Module.Api).Append("UClass* StaticClass<class ").Append(classObj.SourceName).Append(">();\r\n");
+			// For back compat we still need to forward declare the class here
+			builder.Append("class ").Append(classObj.SourceName).Append(";\r\n");
 			
 			Session.Tables!.CodeGeneratorInjectorTable.Inject(builder, classObj, UhtCodeGeneratorInjectionLocation.Header);
 			builder.Append("\r\n");
diff --git a/Engine/Source/Runtime/CoreUObject/Public/UObject/ReflectedTypeAccessors.h b/Engine/Source/Runtime/CoreUObject/Public/UObject/ReflectedTypeAccessors.h
--- a/Engine/Source/Runtime/CoreUObject/Public/UObject/ReflectedTypeAccessors.h
+++ b/Engine/Source/Runtime/CoreUObject/Public/UObject/ReflectedTypeAccessors.h
@@ -10,6 +10,14 @@ class UEnum;
 C++ templated Static(Class/Struct/Enum) retrieval function prototypes.
 -----------------------------------------------------------------------------*/
 
-template<typename ClassType>	UClass*			StaticClass();
-template<typename StructType>	UScriptStruct*	StaticStruct();
-template<typename EnumType>		UEnum*			StaticEnum();
+template<typename ClassType> inline UClass* StaticClass()
+{
+	return ClassType::StaticClass();
+}
+
+template<typename StructType> inline UScriptStruct* StaticStruct()
+{
+	return StructType::StaticStruct();
+}
+
+template<typename EnumType> UEnum* StaticEnum();
