Cloud DDC - Improved error message when a unknown compact binary field type is encountered when converting to json

diff --git a/Engine/Source/Programs/Shared/EpicGames.Serialization/CbField.cs b/Engine/Source/Programs/Shared/EpicGames.Serialization/CbField.cs
--- a/Engine/Source/Programs/Shared/EpicGames.Serialization/CbField.cs
+++ b/Engine/Source/Programs/Shared/EpicGames.Serialization/CbField.cs
@@ -1914,7 +1914,7 @@ namespace EpicGames.Serialization
 					}
 					break;
 				default:
-					throw new NotImplementedException($"Unhandled type in cb-json converter");
+					throw new NotImplementedException($"Unhandled type '{field.GetType().ToString()}' in cb-json converter");
 			}
 		}
 	}
