Horde: Downgrade Slack message failure from error to warning log level

diff --git a/Engine/Source/Programs/Shared/EpicGames.Slack/SlackClient.cs b/Engine/Source/Programs/Shared/EpicGames.Slack/SlackClient.cs
--- a/Engine/Source/Programs/Shared/EpicGames.Slack/SlackClient.cs
+++ b/Engine/Source/Programs/Shared/EpicGames.Slack/SlackClient.cs
@@ -137,7 +137,7 @@ namespace EpicGames.Slack
 				}
 				catch (Exception ex)
 				{
-					_logger.LogError(ex, "Failed to send Slack message to {Url} ({Error}). Request: {Request}. Response: {Response}", request.RequestUri, responseObject.Error, requestJson, Encoding.UTF8.GetString(responseBytes));
+					_logger.LogWarning(ex, "Failed to send Slack message to {Url} ({Error}). Request: {Request}. Response: {Response}", request.RequestUri, responseObject.Error, requestJson, Encoding.UTF8.GetString(responseBytes));
 					throw;
 				}
 			}
