Cloud DDC - Fixed exception when a object id is encountered as we convert it to json

diff --git a/Engine/Source/Programs/Shared/EpicGames.Serialization/CbField.cs b/Engine/Source/Programs/Shared/EpicGames.Serialization/CbField.cs
--- a/Engine/Source/Programs/Shared/EpicGames.Serialization/CbField.cs
+++ b/Engine/Source/Programs/Shared/EpicGames.Serialization/CbField.cs
@@ -1884,6 +1884,7 @@ namespace EpicGames.Serialization
 					}
 					break;
 				case CbFieldType.Uuid:
+				case CbFieldType.ObjectId:
 					if (field.HasName())
 					{
 						writer.WriteString(field.Name.Span, field.AsUuid().ToString());
