RigVM UHT:  Support for Empty Execute Context structure

#rnx

diff --git a/Engine/Plugins/Runtime/RigVM/Source/RigVM/Public/RigVMCore/RigVMStruct.h b/Engine/Plugins/Runtime/RigVM/Source/RigVM/Public/RigVMCore/RigVMStruct.h
--- a/Engine/Plugins/Runtime/RigVM/Source/RigVM/Public/RigVMCore/RigVMStruct.h
+++ b/Engine/Plugins/Runtime/RigVM/Source/RigVM/Public/RigVMCore/RigVMStruct.h
@@ -195,6 +195,9 @@ struct FRigVMStruct
 
 public:
 
+	/** return the execute script struct this unit wants to use */
+	virtual const UScriptStruct* GetExecuteContextStruct() const { return FRigVMExecuteContext::StaticStruct(); }
+
 	/** initialize logic for this struct */
 	virtual void Initialize() {}
 
diff --git a/Engine/Source/Programs/Shared/EpicGames.UHT/Exporters/CodeGen/UhtHeaderCodeGeneratorHFile.cs b/Engine/Source/Programs/Shared/EpicGames.UHT/Exporters/CodeGen/UhtHeaderCodeGeneratorHFile.cs
--- a/Engine/Source/Programs/Shared/EpicGames.UHT/Exporters/CodeGen/UhtHeaderCodeGeneratorHFile.cs
+++ b/Engine/Source/Programs/Shared/EpicGames.UHT/Exporters/CodeGen/UhtHeaderCodeGeneratorHFile.cs
@@ -170,7 +170,7 @@ namespace EpicGames.UHT.Exporters.CodeGen
 					builder.Append("#define ").Append(scriptStruct.SourceName).Append('_').Append(methodInfo.Name).Append("() \\\r\n");
 					builder.Append('\t').Append(methodInfo.ReturnType).Append(' ').Append(scriptStruct.SourceName).Append("::Static").Append(methodInfo.Name).Append("( \\\r\n");
 					builder.Append("\t\t");
-					if (String.IsNullOrEmpty(scriptStruct.RigVMStructInfo.ExecuteContextMember))
+					if (!scriptStruct.RigVMStructInfo.HasAnyExecuteContextMember)
 					{
 						builder.Append("const ");
 					}
@@ -206,10 +206,19 @@ namespace EpicGames.UHT.Exporters.CodeGen
 					if (scriptStruct.RigVMStructInfo != null)
 					{
 						string constPrefix = "";
-						if (String.IsNullOrEmpty(scriptStruct.RigVMStructInfo.ExecuteContextMember))
+						if (!scriptStruct.RigVMStructInfo.HasAnyExecuteContextMember)
 						{
 							constPrefix = "const ";
 						}
+
+						builder.Append('\t');
+						if (!scriptStruct.ScriptStructFlags.HasAnyFlags(EStructFlags.RequiredAPI))
+						{
+							builder.Append(Module.Api);
+						}
+						builder.Append("virtual const UScriptStruct* GetExecuteContextStruct() const override { return ");
+						builder.Append(scriptStruct.RigVMStructInfo.ExecuteContextType);
+						builder.Append("::StaticStruct(); } \\\r\n");
 						
 						foreach (UhtRigVMMethodInfo methodInfo in scriptStruct.RigVMStructInfo.Methods)
 						{
@@ -267,7 +276,7 @@ namespace EpicGames.UHT.Exporters.CodeGen
 							builder.Append('\t').Append(methodInfo.ReturnType).Append(' ').Append(methodInfo.Name).Append('(').Append(constPrefix).Append(scriptStruct.RigVMStructInfo.ExecuteContextType).Append("& InExecuteContext); \\\r\n");
 							builder.Append("\tstatic ").Append(methodInfo.ReturnType).Append(" Static").Append(methodInfo.Name).Append("( \\\r\n");
 							builder.Append("\t\t");
-							if (String.IsNullOrEmpty(scriptStruct.RigVMStructInfo.ExecuteContextMember))
+							if (!scriptStruct.RigVMStructInfo.HasAnyExecuteContextMember)
 							{
 								builder.Append("const ");
 							}
diff --git a/Engine/Source/Programs/Shared/EpicGames.UHT/Types/UhtScriptStruct.cs b/Engine/Source/Programs/Shared/EpicGames.UHT/Types/UhtScriptStruct.cs
--- a/Engine/Source/Programs/Shared/EpicGames.UHT/Types/UhtScriptStruct.cs
+++ b/Engine/Source/Programs/Shared/EpicGames.UHT/Types/UhtScriptStruct.cs
@@ -443,7 +443,7 @@ namespace EpicGames.UHT.Types
 		/// <returns>True if the parameter requires a cast</returns>
 		public bool RequiresCast()
 		{
-			return CastType != null && CastName != null;
+			return !IsEmptyExecuteContext() && CastType != null && CastName != null;
 		}
 
 		/// <summary>
@@ -466,6 +466,27 @@ namespace EpicGames.UHT.Types
 			}
 			return false;
 		}
+
+		/// <summary>
+		/// Return true if the parameter is an empty execute context
+		/// </summary>
+		/// <returns>True if the parameter is an empty execute context</returns>
+		public bool IsEmptyExecuteContext()
+		{
+			if (Property is UhtStructProperty structProperty)
+			{
+				UhtScriptStruct? scriptStruct = structProperty.ScriptStruct;
+				while (scriptStruct != null)
+				{
+					if (scriptStruct.SourceName == "FRigVMEmptyExecuteContext")
+					{
+						return true;
+					}
+					scriptStruct = scriptStruct.SuperScriptStruct;
+				}
+			}
+			return false;
+		}
 	}
 
 	/// <summary>
@@ -539,6 +560,11 @@ namespace EpicGames.UHT.Types
 		/// </summary>
 		public string ExecuteContextMember { get; set; } = String.Empty;
 
+		/// <summary>
+		/// true if any of this structs members is an execute context (also an empty one)
+		/// </summary>
+		public bool HasAnyExecuteContextMember { get; set; } = false;
+
 		/// <summary>
 		/// List of the members
 		/// </summary>
@@ -956,6 +982,11 @@ namespace EpicGames.UHT.Types
 							{
 								this.LogError($"RigVM Struct {SourceName} contains properties of varying execute context type {RigVMStructInfo.ExecuteContextType} vs {parameter.Type}.");
 							}
+							RigVMStructInfo.HasAnyExecuteContextMember = true;
+						}
+						else if (parameter.IsEmptyExecuteContext())
+						{
+							RigVMStructInfo.HasAnyExecuteContextMember = true;
 						}
 						else
 						{
