UHT: Fix up rigvm related code generator code to comply with CA1505

#rnx

diff --git a/Engine/Source/Programs/Shared/EpicGames.UHT/Exporters/CodeGen/UhtHeaderCodeGeneratorCppFile.cs b/Engine/Source/Programs/Shared/EpicGames.UHT/Exporters/CodeGen/UhtHeaderCodeGeneratorCppFile.cs
--- a/Engine/Source/Programs/Shared/EpicGames.UHT/Exporters/CodeGen/UhtHeaderCodeGeneratorCppFile.cs
+++ b/Engine/Source/Programs/Shared/EpicGames.UHT/Exporters/CodeGen/UhtHeaderCodeGeneratorCppFile.cs
@@ -799,7 +799,7 @@ namespace EpicGames.UHT.Exporters.CodeGen
 			if (scriptStruct.RigVMStructInfo != null)
 			{
 				string constPrefix = "";
-				if (String.IsNullOrEmpty(scriptStruct.RigVMStructInfo.ExecuteContextMember))
+				if (!scriptStruct.RigVMStructInfo.HasAnyExecuteContextMember)
 				{
 					constPrefix = "const ";
 				}
diff --git a/Engine/Source/Programs/Shared/EpicGames.UHT/Exporters/CodeGen/UhtHeaderCodeGeneratorHFile.cs b/Engine/Source/Programs/Shared/EpicGames.UHT/Exporters/CodeGen/UhtHeaderCodeGeneratorHFile.cs
--- a/Engine/Source/Programs/Shared/EpicGames.UHT/Exporters/CodeGen/UhtHeaderCodeGeneratorHFile.cs
+++ b/Engine/Source/Programs/Shared/EpicGames.UHT/Exporters/CodeGen/UhtHeaderCodeGeneratorHFile.cs
@@ -155,7 +155,7 @@ namespace EpicGames.UHT.Exporters.CodeGen
 			}
 		}
 
-		private StringBuilder AppendScriptStruct(StringBuilder builder, UhtScriptStruct scriptStruct)
+		private StringBuilder AppendRigVmStructInfoDefine(StringBuilder builder, UhtScriptStruct scriptStruct)
 		{
 			if (scriptStruct.RigVMStructInfo != null)
 			{
@@ -189,237 +189,251 @@ namespace EpicGames.UHT.Exporters.CodeGen
 				builder.Append("\r\n");
 			}
 
-			if (scriptStruct.ScriptStructFlags.HasAnyFlags(EStructFlags.Native))
+			return builder;
+		}
+
+		private StringBuilder AppendRigVmStructInfoMethods(StringBuilder builder, UhtScriptStruct scriptStruct)
+		{
+			// if we have RigVM methods on this struct we need to 
+			// declare the static method as well as the stub method
+			if (scriptStruct.RigVMStructInfo != null)
 			{
-				using (UhtMacroCreator macro = new(builder, this, scriptStruct, GeneratedBodyMacroSuffix))
+				string constPrefix = "";
+				if (!scriptStruct.RigVMStructInfo.HasAnyExecuteContextMember)
 				{
-					builder.Append("\tfriend struct Z_Construct_UScriptStruct_").Append(scriptStruct.SourceName).Append("_Statics; \\\r\n");
-					builder.Append('\t');
-					if (!scriptStruct.ScriptStructFlags.HasAnyFlags(EStructFlags.RequiredAPI))
-					{
-						builder.Append(Module.Api);
-					}
-					builder.Append("static class UScriptStruct* StaticStruct(); \\\r\n");
+					constPrefix = "const ";
+				}
 
-					// if we have RigVM methods on this struct we need to 
-					// declare the static method as well as the stub method
-					if (scriptStruct.RigVMStructInfo != null)
+				builder.Append('\t');
+				if (!scriptStruct.ScriptStructFlags.HasAnyFlags(EStructFlags.RequiredAPI))
+				{
+					builder.Append(Module.Api);
+				}
+				builder.Append("virtual const UScriptStruct* GetExecuteContextStruct() const override { return ");
+				builder.Append(scriptStruct.RigVMStructInfo.ExecuteContextType);
+				builder.Append("::StaticStruct(); } \\\r\n");
+				
+				foreach (UhtRigVMMethodInfo methodInfo in scriptStruct.RigVMStructInfo.Methods)
+				{
+					if (methodInfo.IsPredicate)
 					{
-						string constPrefix = "";
-						if (!scriptStruct.RigVMStructInfo.HasAnyExecuteContextMember)
+						builder.Append('\t').Append("struct ").Append(methodInfo.Name).Append("Struct \\\r\n");
+						builder.Append("\t{ \\\r\n");
+						builder.Append("\t\t").Append(methodInfo.Name).Append("Struct(){  } \\\r\n");
+						builder.Append("\t\t").Append(methodInfo.Name).Append("Struct(FRigVMExtendedExecuteContext& InContext, FRigVMPredicateBranch InBranch){ Context = &InContext; Branch = InBranch; } \\\r\n");
+						builder.Append("\t\t").Append(methodInfo.ReturnType).Append(" Execute(\\\r\n\t\t\t");
+						builder.AppendParameterDecls(methodInfo.Parameters, false, ", \\\r\n\t\t\t", true, false);
+						builder.Append("\t\t)  \\\r\n");
+						builder.Append("\t\t{  \\\r\n");
+						builder.Append("\t\t\tif (Branch.IsValid())  \\\r\n");
+						builder.Append("\t\t\t{  \\\r\n");
+						int parameterIndex = 0;
+						foreach (UhtRigVMParameter parameter in methodInfo.Parameters)
 						{
-							constPrefix = "const ";
+							string baseType = parameter.TypeOriginal().ToString();
+							if (baseType.StartsWith("const", StringComparison.Ordinal))
+							{
+								baseType = baseType[5..].Trim();
+							}
+
+							if (baseType.EndsWith("&", StringComparison.Ordinal))
+							{
+								baseType = baseType[..^1].Trim();
+							}
+							builder.Append("\t\t\t\t*(").Append(baseType).Append("*) Branch.MemoryHandles[")
+								.Append(parameterIndex++).Append("].GetData(false) = ")
+								.Append(parameter.Name).Append(";  \\\r\n");
+							
 						}
+						builder.Append("\t\t\t\tBranch.Execute(*Context);  \\\r\n");
+						builder.Append("\t\t\t\treturn *(").Append(methodInfo.ReturnType).Append("*)Branch.MemoryHandles[").Append(methodInfo.Parameters.Count).Append("].GetData(false);  \\\r\n");
+						builder.Append("\t\t\t}  \\\r\n");
+						builder.Append("\t\t\treturn ").Append(methodInfo.Name).Append('(');
+						builder.AppendParameterNames(methodInfo.Parameters, false, ", ", false, false);
+						builder.Append("); \\\r\n");
+						builder.Append("\t\t}  \\\r\n");
+						builder.Append("\tprivate: \\\r\n");
+						builder.Append("\t\tFRigVMExtendedExecuteContext* Context; \\\r\n");
+						builder.Append("\t\tFRigVMPredicateBranch Branch; \\\r\n");
+						builder.Append("\t};  \\\r\n");
+					}
+				}
 
-						builder.Append('\t');
-						if (!scriptStruct.ScriptStructFlags.HasAnyFlags(EStructFlags.RequiredAPI))
+				foreach (UhtRigVMMethodInfo methodInfo in scriptStruct.RigVMStructInfo.Methods)
+				{
+					if (methodInfo.IsPredicate)
+					{
+						continue;
+					}
+					
+					builder.Append('\t').Append(methodInfo.ReturnType).Append(' ').Append(methodInfo.Name).Append('(').Append(constPrefix).Append(scriptStruct.RigVMStructInfo.ExecuteContextType).Append("& InExecuteContext); \\\r\n");
+					builder.Append("\tstatic ").Append(methodInfo.ReturnType).Append(" Static").Append(methodInfo.Name).Append("( \\\r\n");
+					builder.Append("\t\t");
+					if (!scriptStruct.RigVMStructInfo.HasAnyExecuteContextMember)
+					{
+						builder.Append("const ");
+					}
+					builder.Append(scriptStruct.RigVMStructInfo.ExecuteContextType).Append("& ").Append(RigVMExecuteContextParamName);
+					builder.AppendParameterDecls(scriptStruct.RigVMStructInfo.Members, true, ", \\\r\n\t\t", true, false);
+
+					foreach (UhtRigVMMethodInfo predicateInfo in scriptStruct.RigVMStructInfo.Methods)
+					{
+						if (predicateInfo.IsPredicate)
 						{
-							builder.Append(Module.Api);
+							builder.Append(", \\\r\n\t\t").Append(predicateInfo.Name).Append("Struct& ").Append(predicateInfo.Name);
 						}
-						builder.Append("virtual const UScriptStruct* GetExecuteContextStruct() const override { return ");
-						builder.Append(scriptStruct.RigVMStructInfo.ExecuteContextType);
-						builder.Append("::StaticStruct(); } \\\r\n");
-						
-						foreach (UhtRigVMMethodInfo methodInfo in scriptStruct.RigVMStructInfo.Methods)
-						{
-							if (methodInfo.IsPredicate)
-							{
-								builder.Append('\t').Append("struct ").Append(methodInfo.Name).Append("Struct \\\r\n");
-								builder.Append("\t{ \\\r\n");
-								builder.Append("\t\t").Append(methodInfo.Name).Append("Struct(){  } \\\r\n");
-								builder.Append("\t\t").Append(methodInfo.Name).Append("Struct(FRigVMExtendedExecuteContext& InContext, FRigVMPredicateBranch InBranch){ Context = &InContext; Branch = InBranch; } \\\r\n");
-								builder.Append("\t\t").Append(methodInfo.ReturnType).Append(" Execute(\\\r\n\t\t\t");
-								builder.AppendParameterDecls(methodInfo.Parameters, false, ", \\\r\n\t\t\t", true, false);
-								builder.Append("\t\t)  \\\r\n");
-								builder.Append("\t\t{  \\\r\n");
-								builder.Append("\t\t\tif (Branch.IsValid())  \\\r\n");
-								builder.Append("\t\t\t{  \\\r\n");
-								int parameterIndex = 0;
-								foreach (UhtRigVMParameter parameter in methodInfo.Parameters)
-								{
-									string baseType = parameter.TypeOriginal().ToString();
-									if (baseType.StartsWith("const", StringComparison.Ordinal))
-									{
-										baseType = baseType[5..].Trim();
-									}
+					}
+					
+					builder.Append(" \\\r\n");
+					builder.Append("\t); \\\r\n");
 
-									if (baseType.EndsWith("&", StringComparison.Ordinal))
-									{
-										baseType = baseType[..^1].Trim();
-									}
-									builder.Append("\t\t\t\t*(").Append(baseType).Append("*) Branch.MemoryHandles[")
-										.Append(parameterIndex++).Append("].GetData(false) = ")
-										.Append(parameter.Name).Append(";  \\\r\n");
-									
-								}
-								builder.Append("\t\t\t\tBranch.Execute(*Context);  \\\r\n");
-								builder.Append("\t\t\t\treturn *(").Append(methodInfo.ReturnType).Append("*)Branch.MemoryHandles[").Append(methodInfo.Parameters.Count).Append("].GetData(false);  \\\r\n");
-								builder.Append("\t\t\t}  \\\r\n");
-								builder.Append("\t\t\treturn ").Append(methodInfo.Name).Append('(');
-								builder.AppendParameterNames(methodInfo.Parameters, false, ", ", false, false);
-								builder.Append("); \\\r\n");
-								builder.Append("\t\t}  \\\r\n");
-								builder.Append("\tprivate: \\\r\n");
-								builder.Append("\t\tFRigVMExtendedExecuteContext* Context; \\\r\n");
-								builder.Append("\t\tFRigVMPredicateBranch Branch; \\\r\n");
-								builder.Append("\t};  \\\r\n");
-							}
-						}
+					builder.Append("\tFORCEINLINE_DEBUGGABLE static ").Append(methodInfo.ReturnType).Append(" RigVM").Append(methodInfo.Name).Append("( \\\r\n");
+					builder.Append("\t\t");
+					builder.Append(RigVMExecuteContextDeclaration).Append(", \\\r\n");
+					builder.Append("\t\tFRigVMMemoryHandleArray RigVMMemoryHandles, \\\r\n");
+					builder.Append("\t\tFRigVMPredicateBranchArray RigVMBranches \\\r\n");
+					builder.Append("\t) \\\r\n");
+					builder.Append("\t{ \\\r\n");
 
-						foreach (UhtRigVMMethodInfo methodInfo in scriptStruct.RigVMStructInfo.Methods)
+					if (scriptStruct.RigVMStructInfo.Members.Count > 0)
+					{
+						int operandIndex = 0;
+						foreach (UhtRigVMParameter parameter in scriptStruct.RigVMStructInfo.Members)
 						{
-							if (methodInfo.IsPredicate)
+							string paramTypeOriginal = parameter.TypeOriginal(true);
+							string paramNameOriginal = parameter.NameOriginal(false);
+							string additionalParameters = String.Empty;
+							if (parameter.IsLazy || (!parameter.Input && !parameter.Output && !parameter.Singleton))
 							{
-								continue;
+								additionalParameters = ", RigVMExecuteContext.GetSlice().GetIndex()";
 							}
-							
-							builder.Append('\t').Append(methodInfo.ReturnType).Append(' ').Append(methodInfo.Name).Append('(').Append(constPrefix).Append(scriptStruct.RigVMStructInfo.ExecuteContextType).Append("& InExecuteContext); \\\r\n");
-							builder.Append("\tstatic ").Append(methodInfo.ReturnType).Append(" Static").Append(methodInfo.Name).Append("( \\\r\n");
-							builder.Append("\t\t");
-							if (!scriptStruct.RigVMStructInfo.HasAnyExecuteContextMember)
+
+							string getDataMethod = "GetData";
+							if (parameter.IsLazy)
 							{
-								builder.Append("const ");
+								getDataMethod = $"GetDataLazily<{parameter.TypeOriginal(false, false)}>";
 							}
-							builder.Append(scriptStruct.RigVMStructInfo.ExecuteContextType).Append("& ").Append(RigVMExecuteContextParamName);
-							builder.AppendParameterDecls(scriptStruct.RigVMStructInfo.Members, true, ", \\\r\n\t\t", true, false);
 
-							foreach (UhtRigVMMethodInfo predicateInfo in scriptStruct.RigVMStructInfo.Methods)
+							if (parameter.IsArray)
 							{
-								if (predicateInfo.IsPredicate)
+								string extendedType = parameter.ExtendedType();
+
+								if (!parameter.IsLazy)
 								{
-									builder.Append(", \\\r\n\t\t").Append(predicateInfo.Name).Append("Struct& ").Append(predicateInfo.Name);
+									builder
+										.Append("\t\tTArray")
+										.Append(extendedType)
+										.Append("& ")
+										.Append(paramNameOriginal)
+										.Append(" = ")
+										.Append("*(TArray")
+										.Append(extendedType)
+										.Append("*)");
 								}
-							}
-							
-							builder.Append(" \\\r\n");
-							builder.Append("\t); \\\r\n");
-
-							builder.Append("\tFORCEINLINE_DEBUGGABLE static ").Append(methodInfo.ReturnType).Append(" RigVM").Append(methodInfo.Name).Append("( \\\r\n");
-							builder.Append("\t\t");
-							builder.Append(RigVMExecuteContextDeclaration).Append(", \\\r\n");
-							builder.Append("\t\tFRigVMMemoryHandleArray RigVMMemoryHandles, \\\r\n");
-							builder.Append("\t\tFRigVMPredicateBranchArray RigVMBranches \\\r\n");
-							builder.Append("\t) \\\r\n");
-							builder.Append("\t{ \\\r\n");
-
-							if (scriptStruct.RigVMStructInfo.Members.Count > 0)
-							{
-								int operandIndex = 0;
-								foreach (UhtRigVMParameter parameter in scriptStruct.RigVMStructInfo.Members)
+								else
 								{
-									string paramTypeOriginal = parameter.TypeOriginal(true);
-									string paramNameOriginal = parameter.NameOriginal(false);
-									string additionalParameters = String.Empty;
-									if (parameter.IsLazy || (!parameter.Input && !parameter.Output && !parameter.Singleton))
-									{
-										additionalParameters = ", RigVMExecuteContext.GetSlice().GetIndex()";
-									}
-
-									string getDataMethod = "GetData";
-									if (parameter.IsLazy)
-									{
-										getDataMethod = $"GetDataLazily<{parameter.TypeOriginal(false, false)}>";
-									}
-
-									if (parameter.IsArray)
-									{
-										string extendedType = parameter.ExtendedType();
-
-										if (!parameter.IsLazy)
-										{
-											builder
-												.Append("\t\tTArray")
-												.Append(extendedType)
-												.Append("& ")
-												.Append(paramNameOriginal)
-												.Append(" = ")
-												.Append("*(TArray")
-												.Append(extendedType)
-												.Append("*)");
-										}
-										else
-										{
-											builder
-												.Append("\t\tconst ")
-												.Append(paramTypeOriginal)
-												.Append("& ")
-												.Append(paramNameOriginal)
-												.Append(" = ");
-										}
-
-										builder
-											.Append("RigVMMemoryHandles[")
-											.Append(operandIndex)
-											.Append("].")
-											.Append(getDataMethod)
-											.Append("(false")
-											.Append(additionalParameters)
-											.Append("); \\\r\n");
-										operandIndex++;
-									}
-									else
-									{
-										string variableType = parameter.TypeVariableRef(true);
-										string extractedType = parameter.TypeOriginal();
-
-										if (parameter.IsEnumAsByte)
-										{
-											extractedType = parameter.TypeOriginal(true);
-										}
-
-										builder
-											.Append("\t\t")
-											.Append(variableType)
-											.Append(' ')
-											.Append(paramNameOriginal)
-											.Append(" = ");
-
-										if (!parameter.IsLazy)
-										{
-											builder
-												.Append("*(")
-												.Append(extractedType)
-												.Append("*)");
-										}
-
-										builder
-											.Append("RigVMMemoryHandles[")
-											.Append(operandIndex)
-											.Append("].")
-											.Append(getDataMethod)
-											.Append("(false")
-											.Append(additionalParameters)
-											.Append("); \\\r\n");
-										operandIndex++;
-									}
+									builder
+										.Append("\t\tconst ")
+										.Append(paramTypeOriginal)
+										.Append("& ")
+										.Append(paramNameOriginal)
+										.Append(" = ");
 								}
-								builder.Append("\t\t \\\r\n");
-							}
 
-							int predicateId = 0;
-							foreach (UhtRigVMMethodInfo predicateInfo in scriptStruct.RigVMStructInfo.Methods)
+								builder
+									.Append("RigVMMemoryHandles[")
+									.Append(operandIndex)
+									.Append("].")
+									.Append(getDataMethod)
+									.Append("(false")
+									.Append(additionalParameters)
+									.Append("); \\\r\n");
+								operandIndex++;
+							}
+							else
 							{
-								if (predicateInfo.IsPredicate)
+								string variableType = parameter.TypeVariableRef(true);
+								string extractedType = parameter.TypeOriginal();
+
+								if (parameter.IsEnumAsByte)
 								{
-									builder.Append("\t\t").Append(predicateInfo.Name).Append("Struct ").Append(predicateInfo.Name).Append("Predicate(RigVMExecuteContext, RigVMBranches[").Append(predicateId++).Append("]); \\\r\n");
+									extractedType = parameter.TypeOriginal(true);
 								}
-							}
-							
-							builder.Append("\t\t").Append(methodInfo.ReturnPrefix()).Append("Static").Append(methodInfo.Name).Append("( \\\r\n");
-							builder.Append("\t\t\tRigVMExecuteContext.GetPublicData<").Append(scriptStruct.RigVMStructInfo.ExecuteContextType).Append(">()");
-							builder.AppendParameterNames(scriptStruct.RigVMStructInfo.Members, true, ", \\\r\n\t\t\t", false);
-							foreach (UhtRigVMMethodInfo predicateInfo in scriptStruct.RigVMStructInfo.Methods)
-							{
-								if (predicateInfo.IsPredicate)
+
+								builder
+									.Append("\t\t")
+									.Append(variableType)
+									.Append(' ')
+									.Append(paramNameOriginal)
+									.Append(" = ");
+
+								if (!parameter.IsLazy)
 								{
-									builder.Append(", \\\r\n\t\t\t").Append(predicateInfo.Name).Append("Predicate");
+									builder
+										.Append("*(")
+										.Append(extractedType)
+										.Append("*)");
 								}
+
+								builder
+									.Append("RigVMMemoryHandles[")
+									.Append(operandIndex)
+									.Append("].")
+									.Append(getDataMethod)
+									.Append("(false")
+									.Append(additionalParameters)
+									.Append("); \\\r\n");
+								operandIndex++;
 							}
-							builder.Append(" \\\r\n");
-							builder.Append("\t\t); \\\r\n");
-							builder.Append("\t} \\\r\n");
+						}
+						builder.Append("\t\t \\\r\n");
+					}
+
+					int predicateId = 0;
+					foreach (UhtRigVMMethodInfo predicateInfo in scriptStruct.RigVMStructInfo.Methods)
+					{
+						if (predicateInfo.IsPredicate)
+						{
+							builder.Append("\t\t").Append(predicateInfo.Name).Append("Struct ").Append(predicateInfo.Name).Append("Predicate(RigVMExecuteContext, RigVMBranches[").Append(predicateId++).Append("]); \\\r\n");
 						}
 					}
+					
+					builder.Append("\t\t").Append(methodInfo.ReturnPrefix()).Append("Static").Append(methodInfo.Name).Append("( \\\r\n");
+					builder.Append("\t\t\tRigVMExecuteContext.GetPublicData<").Append(scriptStruct.RigVMStructInfo.ExecuteContextType).Append(">()");
+					builder.AppendParameterNames(scriptStruct.RigVMStructInfo.Members, true, ", \\\r\n\t\t\t", false);
+					foreach (UhtRigVMMethodInfo predicateInfo in scriptStruct.RigVMStructInfo.Methods)
+					{
+						if (predicateInfo.IsPredicate)
+						{
+							builder.Append(", \\\r\n\t\t\t").Append(predicateInfo.Name).Append("Predicate");
+						}
+					}
+					builder.Append(" \\\r\n");
+					builder.Append("\t\t); \\\r\n");
+					builder.Append("\t} \\\r\n");
+				}
+			}
+
+			return builder;
+		}
+
+		private StringBuilder AppendScriptStruct(StringBuilder builder, UhtScriptStruct scriptStruct)
+		{
+			AppendRigVmStructInfoDefine(builder, scriptStruct);
+			
+			if (scriptStruct.ScriptStructFlags.HasAnyFlags(EStructFlags.Native))
+			{
+				using (UhtMacroCreator macro = new(builder, this, scriptStruct, GeneratedBodyMacroSuffix))
+				{
+					builder.Append("\tfriend struct Z_Construct_UScriptStruct_").Append(scriptStruct.SourceName).Append("_Statics; \\\r\n");
+					builder.Append('\t');
+					if (!scriptStruct.ScriptStructFlags.HasAnyFlags(EStructFlags.RequiredAPI))
+					{
+						builder.Append(Module.Api);
+					}
+					builder.Append("static class UScriptStruct* StaticStruct(); \\\r\n");
+
+					AppendRigVmStructInfoMethods(builder, scriptStruct);
 
 					if (scriptStruct.SuperScriptStruct != null)
 					{
