UHT: Fix missing static modifier in header code generator

#rnx

diff --git a/Engine/Source/Programs/Shared/EpicGames.UHT/Exporters/CodeGen/UhtHeaderCodeGeneratorHFile.cs b/Engine/Source/Programs/Shared/EpicGames.UHT/Exporters/CodeGen/UhtHeaderCodeGeneratorHFile.cs
--- a/Engine/Source/Programs/Shared/EpicGames.UHT/Exporters/CodeGen/UhtHeaderCodeGeneratorHFile.cs
+++ b/Engine/Source/Programs/Shared/EpicGames.UHT/Exporters/CodeGen/UhtHeaderCodeGeneratorHFile.cs
@@ -155,7 +155,7 @@ namespace EpicGames.UHT.Exporters.CodeGen
 			}
 		}
 
-		private StringBuilder AppendRigVmStructInfoDefine(StringBuilder builder, UhtScriptStruct scriptStruct)
+		private static StringBuilder AppendRigVmStructInfoDefine(StringBuilder builder, UhtScriptStruct scriptStruct)
 		{
 			if (scriptStruct.RigVMStructInfo != null)
 			{
