AutomationTool: Add missing reference to EpicGames.Perforce

#rnx

diff --git a/Engine/Source/Programs/Shared/EpicGames.Perforce/EpicGames.Perforce.csproj b/Engine/Source/Programs/Shared/EpicGames.Perforce/EpicGames.Perforce.csproj
--- a/Engine/Source/Programs/Shared/EpicGames.Perforce/EpicGames.Perforce.csproj
+++ b/Engine/Source/Programs/Shared/EpicGames.Perforce/EpicGames.Perforce.csproj
@@ -1,4 +1,5 @@
 <Project Sdk="Microsoft.NET.Sdk">
+  <Import Project="..\UnrealEngine.csproj.props" />
 
   <PropertyGroup>
     <TargetFramework>net8.0</TargetFramework>
