Update UAT script compiling
* Load csproj files in parallel
* Use custom msbuild task to scan project file dependencies rather than manually resolving file globs, run this scan async during project load
* Remove any code which was resolving file globs as they are no longer needed
* Skip already valid projects instead of reprocessing
* Generate a .props file that is imported into EpicGames.ScriptBuild, which is built with dotnet build. This ended up being the quickest and safest way to compile all script projects and ensure dependencies are restored properly
* Shell out to dotnet build rather than using Microsoft.Build, to allow for better multithreading and to resolve issue where process working directory changes

#rnx

diff --git a/Engine/Source/Programs/Shared/EpicGames.Build/System/CompileScriptModules.cs b/Engine/Source/Programs/Shared/EpicGames.Build/System/CompileScriptModules.cs
--- a/Engine/Source/Programs/Shared/EpicGames.Build/System/CompileScriptModules.cs
+++ b/Engine/Source/Programs/Shared/EpicGames.Build/System/CompileScriptModules.cs
@@ -1,16 +1,15 @@
 // Copyright Epic Games, Inc. All Rights Reserved.
 
-using EpicGames.Core;
 using System;
+using System.Collections.Concurrent;
 using System.Collections.Generic;
 using System.IO;
 using System.Linq;
-using System.Text;
+using System.Runtime.CompilerServices;
 using System.Text.Json;
+using EpicGames.Core;
 using EpicGames.MsBuild;
 using Microsoft.Extensions.Logging;
-using System.Runtime.CompilerServices;
-using Microsoft.Extensions.FileSystemGlobbing;
 
 namespace UnrealBuildBase
 {
@@ -55,8 +54,9 @@ namespace UnrealBuildBase
 		class Hook : CsProjBuildHook
 		{
 			private ILogger Logger;
-			private Dictionary<string, DateTime> WriteTimes = [];
+			private ConcurrentDictionary<string, DateTime> WriteTimes = [];
 			private string BuildRecordDirectory;
+			private Dictionary<FileReference, CsProjBuildRecordEntry> ValidBuildRecords = new();
 
 			public Hook(ILogger InLogger, Rules.RulesFileType InRulesFileType)
 			{
@@ -82,11 +82,7 @@ namespace UnrealBuildBase
 			public DateTime GetLastWriteTime(string BasePath, string RelativeFilePath)
 			{
 				string NormalizedPath = Path.GetFullPath(RelativeFilePath, BasePath);
-				if (!WriteTimes.TryGetValue(NormalizedPath, out DateTime WriteTime))
-				{
-					WriteTimes.Add(NormalizedPath, WriteTime = File.GetLastWriteTime(NormalizedPath));
-				}
-				return WriteTime;
+				return WriteTimes.GetOrAdd(NormalizedPath, File.GetLastWriteTime);
 			}
 
 			public DirectoryReference GetBuildRecordDirectory(DirectoryReference BasePath)
@@ -99,17 +95,13 @@ namespace UnrealBuildBase
 				CompileScriptModule.ValidateBuildRecordRecursively(BuildRecords, ProjectPath, this, Logger);
 			}
 
-			private static readonly char[] s_wildcardCharacters = ['*', '?'];
-			public bool HasWildcards(string FileSpec)
+			public void SetValidBuildRecords(IReadOnlyDictionary<FileReference, CsProjBuildRecordEntry> Records)
 			{
-				// Perf Note: Doing a [Last]IndexOfAny(...) is much faster than compiling a
-				// regular expression that does the same thing, regardless of whether
-				// filespec contains one of the characters.
-				// Choose LastIndexOfAny instead of IndexOfAny because it seems more likely
-				// that wildcards will tend to be towards the right side.
-
-				return -1 != FileSpec.LastIndexOfAny(s_wildcardCharacters);
+				ValidBuildRecords = new(Records.Where(x => x.Value.Status == CsProjBuildRecordStatus.Valid));
 			}
+
+			IReadOnlyDictionary<FileReference, CsProjBuildRecordEntry> CsProjBuildHook.ValidBuildRecords => ValidBuildRecords;
+
 			DirectoryReference CsProjBuildHook.EngineDirectory => Unreal.EngineDirectory;
 
 			DirectoryReference CsProjBuildHook.DotnetDirectory => Unreal.DotnetDirectory;
@@ -385,6 +377,7 @@ namespace UnrealBuildBase
 					bBuildSuccess = true;
 					return new(ValidBuildRecords.Where(x => FoundProjects.Contains(x.Key)));
 				}
+				Hook.SetValidBuildRecords(ValidBuildRecords);
 			}
 
 			// Fall back to the slower approach: use msbuild to load csproj files & build as necessary
@@ -458,86 +451,6 @@ namespace UnrealBuildBase
 			return LoadedBuildRecords;
 		}
 
-		private static bool ValidateGlobbedFiles(DirectoryReference ProjectDirectory,
-			List<CsProjBuildRecord.Glob> Globs, HashSet<string> GlobbedDependencies, out string ValidationFailureMessage)
-		{
-			// First, evaluate globs
-
-			// Files are grouped by ItemType (e.g. Compile, EmbeddedResource) to ensure that Exclude and
-			// Remove act as expected.
-			Dictionary<string, HashSet<string>> Files = [];
-			foreach (CsProjBuildRecord.Glob Glob in Globs)
-			{
-				HashSet<string>? TypedFiles;
-				if (!Files.TryGetValue(Glob.ItemType!, out TypedFiles))
-				{
-					TypedFiles = [];
-					Files.Add(Glob.ItemType!, TypedFiles);
-				}
-
-				Matcher matcher = new Matcher();
-				if (Glob.Include?.Count > 0)
-				{
-					matcher.AddIncludePatterns(Glob.Include);
-				}
-				if (Glob.Exclude?.Count > 0)
-				{
-					matcher.AddExcludePatterns(Glob.Exclude);
-				}
-				TypedFiles.UnionWith(matcher.GetResultsInFullPath(ProjectDirectory.FullName).Select(x => Path.GetRelativePath(ProjectDirectory.FullName, x)));
-
-				if (Glob.Remove?.Count > 0)
-				{
-					// Matcher.Match() doesn't handle inconsistent path separators correctly - which is why globs
-					// are normalized when they are added to CsProjBuildRecord
-					Matcher removeMatcher = new Matcher();
-					removeMatcher.AddIncludePatterns(Glob.Remove);
-					TypedFiles.RemoveWhere(F => removeMatcher.Match(F).HasMatches);
-				}
-			}
-
-			// Then, validation that our evaluation matches what we're comparing against
-
-			bool bValid = true;
-			StringBuilder ValidationFailureText = new StringBuilder();
-
-			// Look for extra files that were found
-			foreach (HashSet<string> TypedFiles in Files.Values)
-			{
-				foreach (string File in TypedFiles)
-				{
-					if (!GlobbedDependencies.Contains(File))
-					{
-						ValidationFailureText.AppendLine($"Found additional file {File}");
-						bValid = false;
-					}
-				}
-			}
-
-			// Look for files that are missing
-			foreach (string File in GlobbedDependencies)
-			{
-				bool bFound = false;
-				foreach (HashSet<string> TypedFiles in Files.Values)
-				{
-					if (TypedFiles.Contains(File))
-					{
-						bFound = true;
-						break;
-					}
-				}
-
-				if (!bFound)
-				{
-					ValidationFailureText.AppendLine($"Did not find {File}");
-					bValid = false;
-				}
-			}
-
-			ValidationFailureMessage = ValidationFailureText.ToString();
-			return bValid;
-		}
-
 		private static bool ValidateBuildRecord(CsProjBuildRecord BuildRecord, DirectoryReference ProjectDirectory, out string ValidationFailureMessage, CsProjBuildHook Hook)
 		{
 			string TargetRelativePath =
@@ -550,6 +463,13 @@ namespace UnrealBuildBase
 				return false;
 			}
 
+			string NormalizedPath = Path.GetFullPath(BuildRecord.TargetPath!, ProjectDirectory.FullName);
+			if (!File.Exists(NormalizedPath))
+			{
+				ValidationFailureMessage = $"{TargetRelativePath} not found";
+				return false;
+			}
+
 			DateTime TargetWriteTime = Hook.GetLastWriteTime(ProjectDirectory, BuildRecord.TargetPath!);
 
 			if (BuildRecord.TargetBuildTime != TargetWriteTime)
@@ -568,21 +488,7 @@ namespace UnrealBuildBase
 				}
 			}
 
-			if (!ValidateGlobbedFiles(ProjectDirectory, BuildRecord.Globs, BuildRecord.GlobbedDependencies,
-				out ValidationFailureMessage))
-			{
-				return false;
-			}
-
-			foreach (string Dependency in BuildRecord.GlobbedDependencies)
-			{
-				if (Hook.GetLastWriteTime(ProjectDirectory, Dependency) > TargetWriteTime)
-				{
-					ValidationFailureMessage = $"{Dependency} is newer than {TargetRelativePath}";
-					return false;
-				}
-			}
-
+			ValidationFailureMessage = String.Empty;
 			return true;
 		}
 
diff --git a/Engine/Source/Programs/Shared/EpicGames.Core/CsProjBuildHook.cs b/Engine/Source/Programs/Shared/EpicGames.Core/CsProjBuildHook.cs
--- a/Engine/Source/Programs/Shared/EpicGames.Core/CsProjBuildHook.cs
+++ b/Engine/Source/Programs/Shared/EpicGames.Core/CsProjBuildHook.cs
@@ -48,11 +48,15 @@ namespace EpicGames.Core
 		void ValidateRecursively(Dictionary<FileReference, CsProjBuildRecordEntry> buildRecords, FileReference projectPath);
 
 		/// <summary>
-		/// Test to see if the given file spec has any wild cards
+		/// (Optional) Sets of existing build records that are valid
 		/// </summary>
-		/// <param name="fileSpec">File spec to test</param>
-		/// <returns>True if wildcards are present</returns>
-		bool HasWildcards(string fileSpec);
+		/// <param name="records">The records to set</param>
+		void SetValidBuildRecords(IReadOnlyDictionary<FileReference, CsProjBuildRecordEntry> records);
+
+		/// <summary>
+		/// (Optional) Collection of existing build records for projects already been registered as valid, which could be skipped when building
+		/// </summary>
+		IReadOnlyDictionary<FileReference, CsProjBuildRecordEntry> ValidBuildRecords { get; }
 
 		/// <summary>
 		/// Unreal engine directory
diff --git a/Engine/Source/Programs/Shared/EpicGames.Core/CsProjBuildRecord.cs b/Engine/Source/Programs/Shared/EpicGames.Core/CsProjBuildRecord.cs
--- a/Engine/Source/Programs/Shared/EpicGames.Core/CsProjBuildRecord.cs
+++ b/Engine/Source/Programs/Shared/EpicGames.Core/CsProjBuildRecord.cs
@@ -66,45 +66,9 @@ namespace EpicGames.Core
 		public List<CsProjBuildRecordRef> ProjectReferencesAndTimes { get; set; } = [];
 
 		/// <summary>
-		/// file dependencies from non-glob sources
+		/// file dependencies
 		/// </summary>
 		public HashSet<string> Dependencies { get; set; } = [];
-
-		/// <summary>
-		/// file dependencies from globs
-		/// </summary>
-		public HashSet<string> GlobbedDependencies { get; set; } = [];
-
-		/// <summary>
-		/// A glob pattern
-		/// </summary>
-		public class Glob
-		{
-			/// <summary>
-			/// Type of item
-			/// </summary>
-			public string? ItemType { get; set; }
-
-			/// <summary>
-			/// Paths to include
-			/// </summary>
-			public List<string>? Include { get; set; }
-
-			/// <summary>
-			/// Paths to exclude
-			/// </summary>
-			public List<string>? Exclude { get; set; }
-
-			/// <summary>
-			/// 
-			/// </summary>
-			public List<string>? Remove { get; set; }
-		}
-
-		/// <summary>
-		/// List of globs
-		/// </summary>
-		public List<Glob> Globs { get; set; } = [];
 	}
 
 	/// <summary>
diff --git a/Engine/Source/Programs/Shared/EpicGames.MsBuild/CsProjBuilder.cs b/Engine/Source/Programs/Shared/EpicGames.MsBuild/CsProjBuilder.cs
--- a/Engine/Source/Programs/Shared/EpicGames.MsBuild/CsProjBuilder.cs
+++ b/Engine/Source/Programs/Shared/EpicGames.MsBuild/CsProjBuilder.cs
@@ -1,122 +1,37 @@
 // Copyright Epic Games, Inc. All Rights Reserved.
 
-using EpicGames.Core;
-using Microsoft.Build.Evaluation;
-using Microsoft.Build.Execution;
-using Microsoft.Build.Framework;
-using Microsoft.Build.Graph;
-using Microsoft.Build.Locator;
-using Microsoft.Extensions.Logging;
 using System;
+using System.Collections.Concurrent;
 using System.Collections.Generic;
 using System.Diagnostics;
 using System.IO;
 using System.Linq;
+using System.Text;
 using System.Text.Json;
+using System.Threading;
+using System.Threading.Tasks;
+using System.Xml;
+using System.Xml.Linq;
+using EpicGames.Core;
+using Microsoft.Build.Evaluation;
+using Microsoft.Build.Locator;
+using Microsoft.Extensions.Logging;
+using Microsoft.Extensions.Logging.Abstractions;
 
 namespace EpicGames.MsBuild
 {
-	using ILogger = Microsoft.Extensions.Logging.ILogger;
-	using IBuildLogger = Microsoft.Build.Framework.ILogger;
-
 	/// <summary>
 	/// Builds .csproj files
 	/// </summary>
 	public static class CsProjBuilder
 	{
-		class MLogger(ILogger inInner) : IBuildLogger
+		[Flags]
+		enum MsBuildTarget
 		{
-			readonly ILogger _inner = inInner;
-
-			LoggerVerbosity IBuildLogger.Verbosity { get => LoggerVerbosity.Normal; set => throw new NotImplementedException(); }
-			string IBuildLogger.Parameters { get => throw new NotImplementedException(); set { } }
-
-			public bool _bVeryVerboseLog = false;
-
-			bool _bFirstError = true;
-
-			void IBuildLogger.Initialize(IEventSource eventSource)
-			{
-				eventSource.ProjectStarted += new ProjectStartedEventHandler(EventSource_ProjectStarted);
-				eventSource.TaskStarted += new TaskStartedEventHandler(EventSource_TaskStarted);
-				eventSource.MessageRaised += new BuildMessageEventHandler(EventSource_MessageRaised);
-				eventSource.WarningRaised += new BuildWarningEventHandler(EventSource_WarningRaised);
-				eventSource.ErrorRaised += new BuildErrorEventHandler(EventSource_ErrorRaised);
-				eventSource.ProjectFinished += new ProjectFinishedEventHandler(EventSource_ProjectFinished);
-			}
-
-			void EventSource_ErrorRaised(object sender, BuildErrorEventArgs e)
-			{
-				if (_bFirstError)
-				{
-					Trace.WriteLine("");
-					_inner.LogInformation("");
-					_bFirstError = false;
-				}
-				_inner.LogError("{File}({Line},{Column}): error {Code}: {Message} ({ProjectFile})", new FileReference(e.File), new LogValue(LogValueType.LineNumber, e.LineNumber.ToString()), new LogValue(LogValueType.ColumnNumber, e.ColumnNumber.ToString()), new LogValue(LogValueType.ErrorCode, e.Code), e.Message, new FileReference(e.ProjectFile));
-			}
-
-			void EventSource_WarningRaised(object sender, BuildWarningEventArgs e)
-			{
-				{
-					// workaround for warnings that appear after revert of net6.0 upgrade. Delete this block when the net6.0 upgrade is done.
-					// ...\Engine\Binaries\ThirdParty\DotNet\Windows\sdk\3.1.403\Microsoft.Common.CurrentVersion.targets(3036,5): warning MSB3088: Could not read state file "obj\Development\[projectname].csproj.GenerateResource.cache". The input stream is not a valid binary format.
-					// The starting contents (in bytes) are: 06-01-01-00-00-00-01-19-50-72-6F-70-65-72-74-69-65 ... (...\[projectname].csproj)
-					if (String.Equals(e.Code, "MSB3088", StringComparison.Ordinal))
-					{
-						_inner.LogDebug("{File}({Line},{Column}): suppressed warning {Code}: {Message} ({ProjectFile})", new FileReference(e.File), new LogValue(LogValueType.LineNumber, e.LineNumber.ToString()), new LogValue(LogValueType.ColumnNumber, e.ColumnNumber.ToString()), new LogValue(LogValueType.ErrorCode, e.Code), e.Message, new FileReference(e.ProjectFile));
-						return;
-					}
-				}
-
-				if (_bFirstError)
-				{
-					_inner.LogInformation("");
-					_bFirstError = false;
-				}
-
-				_inner.LogWarning("{File}({Line},{Column}): warning {Code}: {Message} ({ProjectFile})", new FileReference(e.File), new LogValue(LogValueType.LineNumber, e.LineNumber.ToString()), new LogValue(LogValueType.ColumnNumber, e.ColumnNumber.ToString()), new LogValue(LogValueType.ErrorCode, e.Code), e.Message, new FileReference(e.ProjectFile));
-			}
-
-			void EventSource_MessageRaised(object sender, BuildMessageEventArgs e)
-			{
-				if (_bVeryVerboseLog)
-				{
-					//if (!String.Equals(e.SenderName, "ResolveAssemblyReference"))
-					//if (e.Message.Contains("atic"))
-					{
-						_inner.LogDebug("{SenderName}: {Message}", e.SenderName, e.Message);
-					}
-				}
-			}
-
-			void EventSource_ProjectStarted(object sender, ProjectStartedEventArgs e)
-			{
-				if (_bVeryVerboseLog)
-				{
-					_inner.LogDebug("{SenderName}: {Message}", e.SenderName, e.Message);
-				}
-			}
-
-			void EventSource_ProjectFinished(object sender, ProjectFinishedEventArgs e)
-			{
-				if (_bVeryVerboseLog)
-				{
-					_inner.LogDebug("{SenderName}: {Message}", e.SenderName, e.Message);
-				}
-			}
-
-			void EventSource_TaskStarted(object sender, TaskStartedEventArgs e)
-			{
-				if (_bVeryVerboseLog)
-				{
-					_inner.LogDebug("{SenderName}: {Message}", e.SenderName, e.Message);
-				}
-			}
-
-			void IBuildLogger.Shutdown()
-			{
-			}
+			Restore = 1,
+			Scan = 2,
+			Build = 4,
+			Clean = 8,
 		}
 
 		static FileReference ConstructBuildRecordPath(CsProjBuildHook hook, FileReference projectPath, IEnumerable<DirectoryReference> baseDirectories)
@@ -181,11 +96,15 @@ namespace EpicGames.MsBuild
 			bool bForceCompile, out bool bBuildSuccess, CsProjBuildHook hook, IEnumerable<DirectoryReference> baseDirectories, IEnumerable<string> defineConstants,
 			Action<int> onBuildingProjects, ILogger logger)
 		{
+			Stopwatch stopwatch = new();
+			FileReference csharpTargetsFile = FileReference.Combine(hook.EngineDirectory, "Source", "Programs", "Shared", "UnrealEngine.CSharp.targets");
 			Dictionary<string, string> globalProperties = new Dictionary<string, string>
 			{
 				{ "EngineDir", hook.EngineDirectory.FullName },
 				{ "EnginePath", hook.EngineDirectory.FullName },
 				{ "EpicGamesMsBuild", "true" },
+				{ "CustomAfterMicrosoftCommonProps", csharpTargetsFile.FullName },
+				{ "NoWarn", "$(NoWarn);MSB3026;NETSDK1206".Replace(";", "%3B", StringComparison.Ordinal) },
 #if DEBUG
 				{ "Configuration", "Debug" },
 #else
@@ -198,83 +117,87 @@ namespace EpicGames.MsBuild
 				globalProperties.Add("DefineConstants", String.Join(';', defineConstants));
 			}
 
-			Dictionary<FileReference, CsProjBuildRecordEntry> buildRecords = [];
+			ConcurrentDictionary<FileReference, CsProjBuildRecordEntry> buildRecords = new(hook.ValidBuildRecords);
+			IEnumerable<FileReference> outOfDateProjects = foundProjects.Except(buildRecords.Keys).Distinct();
 
 			using ProjectCollection projectCollection = new ProjectCollection(globalProperties);
-			Dictionary<string, Project> projects = [];
-			HashSet<string> skippedProjects = [];
+			ConcurrentDictionary<FileReference, Project> projects = new();
 
-			// Microsoft.Build.Evaluation.Project provides access to information stored in the .csproj xml that is 
-			// not available when using Microsoft.Build.Execution.ProjectInstance (used later in this function and
-			// in BuildProjects) - particularly, to access glob information defined in the source file.
+			logger.LogDebug("Async Scan {Count} projects...", outOfDateProjects.Count());
+			List<Task> scanTasks = [Task.Run(async () => await RunMsBuildAsync(MsBuildTarget.Scan, outOfDateProjects, globalProperties, hook, NullLogger.Instance, default))];
 
-			// Load all found projects, and any other referenced projects.
-			foreach (FileReference projectPath in foundProjects)
+			stopwatch.Restart();
+			// Load all found projects
+			IEnumerable<FileReference> toProcess = outOfDateProjects;
+			while (toProcess.Any())
 			{
-				void LoadProjectAndReferences(string projectPath, string referencedBy)
+				ConcurrentBag<FileReference> referencedProjects = new();
+				Parallel.ForEach(toProcess, (foundProject) =>
 				{
-					projectPath = Path.GetFullPath(projectPath);
-					if (!projects.ContainsKey(projectPath) && !skippedProjects.Contains(projectPath))
+					string projectPath = Path.GetFullPath(foundProject.FullName);
+
+					Project project;
+					// Microsoft.Build.Evaluation.Project doesn't give a lot of useful information if this fails,
+					// so make sure to print our own diagnostic info if something goes wrong
+					try
 					{
-						Project project;
+						project = new Project(projectPath, globalProperties, toolsVersion: null, projectCollection: projectCollection);
+					}
+					catch (Microsoft.Build.Exceptions.InvalidProjectFileException iPFEx)
+					{
+						logger.LogWarning("Could not load project file {ProjectPath}: {Message}", projectPath, iPFEx.BaseMessage);
+						return;
+					}
 
-						// Microsoft.Build.Evaluation.Project doesn't give a lot of useful information if this fails,
-						// so make sure to print our own diagnostic info if something goes wrong
-						try
-						{
-							project = new Project(projectPath, globalProperties, toolsVersion: null, projectCollection: projectCollection);
-						}
-						catch (Microsoft.Build.Exceptions.InvalidProjectFileException iPFEx)
+					if (!OperatingSystem.IsWindows())
+					{
+						// check the TargetFramework of the project: we can't build Windows-only projects on 
+						// non-Windows platforms.
+						if (project.GetProperty("TargetFramework").EvaluatedValue.Contains("windows", StringComparison.Ordinal))
 						{
-							logger.LogError("Could not load project file {ProjectPath}", projectPath);
-							logger.LogError("{Message}", iPFEx.BaseMessage);
-
-							if (!String.IsNullOrEmpty(referencedBy))
-							{
-								logger.LogError("Referenced by: {ReferencedBy}", referencedBy);
-							}
-							if (projects.Count > 0)
-							{
-								logger.LogError("See the log file for the list of previously loaded projects.");
-								logger.LogError("Loaded projects (most recently loaded first):");
-								foreach (string path in projects.Keys.Reverse())
-								{
-									logger.LogError("  {Path}", path);
-								}
-							}
-							throw;
+							logger.LogInformation("Skipping windows-only project {ProjectPath}", projectPath);
+							return;
 						}
+					}
 
-						if (!OperatingSystem.IsWindows())
-						{
-							// check the TargetFramework of the project: we can't build Windows-only projects on 
-							// non-Windows platforms.
-							if (project.GetProperty("TargetFramework").EvaluatedValue.Contains("windows", StringComparison.Ordinal))
-							{
-								skippedProjects.Add(projectPath);
-								logger.LogInformation("Skipping windows-only project {ProjectPath}", projectPath);
-								return;
-							}
-						}
+					if (project.GetProperty("TargetFramework").EvaluatedValue.Contains("net6.0", StringComparison.Ordinal))
+					{
+						logger.LogInformation("Skipping net6.0 project {ProjectPath}", projectPath);
+						return;
+					}
 
-						projects.Add(projectPath, project);
-						referencedBy = String.IsNullOrEmpty(referencedBy) ? projectPath : $"{projectPath}{Environment.NewLine}{referencedBy}";
-						foreach (string referencedProject in project.GetItems("ProjectReference").
-							Select(i => i.EvaluatedInclude))
-						{
-							LoadProjectAndReferences(Path.Combine(project.DirectoryPath, referencedProject), referencedBy);
-						}
+					projects.TryAdd(foundProject, project);
+					foreach (ProjectItem reference in project.GetItems("ProjectReference"))
+					{
+						referencedProjects.Add(FileReference.FromString(Path.GetFullPath(reference.EvaluatedInclude, foundProject.Directory.FullName)));
 					}
+				});
+
+				toProcess = referencedProjects
+					.Except(projects.Keys)
+					.Except(buildRecords.Keys)
+					.Distinct();
+
+				if (toProcess.Any())
+				{
+					logger.LogDebug("Async Scan {Count} additional project references...", toProcess.Count());
+					IEnumerable<FileReference> toScan = toProcess;
+					scanTasks.Add(scanTasks.Last().ContinueWith(async (item) => await RunMsBuildAsync(MsBuildTarget.Scan, toScan, globalProperties, hook, NullLogger.Instance, default), TaskScheduler.Default));
 				}
-				LoadProjectAndReferences(projectPath.FullName, null);
 			}
+			logger.LogDebug("Load {Count} projects time: {TimeSeconds:0.00} s", projects.Count, stopwatch.Elapsed.TotalSeconds);
+
+			stopwatch.Restart();
+			Task.WaitAll([.. scanTasks]);
+			logger.LogDebug("Async Scan wait time: {TimeSeconds:0.00} s", stopwatch.Elapsed.TotalSeconds);
 
 			// generate a BuildRecord for each loaded project - the gathered information will be used to determine if the project is
 			// out of date, and if building this project can be skipped. It is also used to populate Intermediate/ScriptModules after the
 			// build completes
-			foreach (Project project in projects.Values)
+			stopwatch.Restart();
+			Parallel.ForEach(projects.Values, (project) =>
 			{
-				string targetPath = Path.GetRelativePath(project.DirectoryPath, project.GetPropertyValue("TargetPath"));
+				string targetPath = Path.GetRelativePath(project.DirectoryPath, project.ExpandString(project.GetPropertyValue("TargetPath")));
 
 				FileReference projectPath = FileReference.FromString(project.FullPath);
 				FileReference buildRecordPath = ConstructBuildRecordPath(hook, projectPath, baseDirectories);
@@ -284,196 +207,72 @@ namespace EpicGames.MsBuild
 					Version = CsProjBuildRecord.CurrentVersion,
 					TargetPath = targetPath,
 					TargetBuildTime = hook.GetLastWriteTime(project.DirectoryPath, targetPath),
-					ProjectPath = Path.GetRelativePath(buildRecordPath.Directory.FullName, project.FullPath)
+					ProjectPath = Path.GetRelativePath(buildRecordPath.Directory.FullName, project.FullPath),
 				};
 
-				// the .csproj
-				buildRecord.Dependencies.Add(Path.GetRelativePath(project.DirectoryPath, project.FullPath));
-
-				// Imports: files included in the xml (typically props, targets, etc)
-				foreach (ResolvedImport import in project.Imports)
-				{
-					string importPath = Path.GetRelativePath(project.DirectoryPath, import.ImportedProject.FullPath);
-
-					// nuget.g.props and nuget.g.targets are generated by Restore, and are frequently re-written;
-					// it should be safe to ignore these files - changes to references from a .csproj file will
-					// show up as that file being out of date.
-					if (importPath.Contains("nuget.g.", StringComparison.Ordinal))
-					{
-						continue;
-					}
-
-					buildRecord.Dependencies.Add(importPath);
-				}
-
-				// References: e.g. Ionic.Zip.Reduced.dll, fastJSON.dll
-				foreach (ProjectItem item in project.GetItems("Reference"))
-				{
-					buildRecord.Dependencies.Add(item.GetMetadataValue("HintPath"));
-				}
-
 				foreach (ProjectItem referencedProjectItem in project.GetItems("ProjectReference"))
 				{
 					buildRecord.ProjectReferencesAndTimes.Add(new CsProjBuildRecordRef { ProjectPath = referencedProjectItem.EvaluatedInclude });
 				}
 
-				foreach (ProjectItem compileItem in project.GetItems("Compile"))
+				FileReference depPath = FileReference.FromString(Path.GetFullPath($"{targetPath}.dep.csv", projectPath.Directory.FullName));
+				if (FileReference.Exists(depPath))
 				{
-					if (hook.HasWildcards(compileItem.UnevaluatedInclude))
-					{
-						buildRecord.GlobbedDependencies.Add(compileItem.EvaluatedInclude);
-					}
-					else
-					{
-						buildRecord.Dependencies.Add(compileItem.EvaluatedInclude);
-					}
-				}
-
-				foreach (ProjectItem contentItem in project.GetItems("Content"))
-				{
-					if (hook.HasWildcards(contentItem.UnevaluatedInclude))
-					{
-						buildRecord.GlobbedDependencies.Add(contentItem.EvaluatedInclude);
-					}
-					else
-					{
-						buildRecord.Dependencies.Add(contentItem.EvaluatedInclude);
-					}
-				}
-
-				foreach (ProjectItem embeddedResourceItem in project.GetItems("EmbeddedResource"))
-				{
-					if (hook.HasWildcards(embeddedResourceItem.UnevaluatedInclude))
-					{
-						buildRecord.GlobbedDependencies.Add(embeddedResourceItem.EvaluatedInclude);
-					}
-					else
-					{
-						buildRecord.Dependencies.Add(embeddedResourceItem.EvaluatedInclude);
-					}
-				}
-
-				// this line right here is slow: ~30-40ms per project (which can be more than a second total)
-				// making it one of the slowest steps in gathering or checking dependency information from
-				// .csproj files (after loading as Microsoft.Build.Evalation.Project)
-				// 
-				// This also returns a lot more information than we care for - MSBuildGlob objects,
-				// which have a range of precomputed values. It may be possible to take source for
-				// GetAllGlobs() and construct a version that does less.
-				List<GlobResult> globs = project.GetAllGlobs();
-
-				// FileMatcher.IsMatch() requires directory separators in glob strings to match the
-				// local flavor. There's probably a better way.
-				string CleanGlobString(string globString)
-				{
-					char sep = Path.DirectorySeparatorChar;
-					char notSep = sep == '/' ? '\\' : '/'; // AltDirectorySeparatorChar isn't always what we need (it's '/' on Mac)
-
-					char[] chars = globString.ToCharArray();
-					int p = 0;
-					for (int i = 0; i < globString.Length; ++i, ++p)
-					{
-						// Flip a non-native separator
-						if (chars[i] == notSep)
-						{
-							chars[p] = sep;
-						}
-						else
-						{
-							chars[p] = chars[i];
-						}
+					string[] lines = FileReference.ReadAllLines(depPath);
+					IEnumerable<string> dependencies = lines
+						.Select(x => FileReference.FromString(x.Split(',').Last()))
+						.Except([csharpTargetsFile])
+						.Select(x => Path.GetRelativePath(project.DirectoryPath, x.FullName));
 
-						// Collapse adjacent separators
-						if (i > 0 && chars[p] == sep && chars[p - 1] == sep)
-						{
-							p -= 1;
-						}
-					}
-
-					return new string(chars, 0, p);
-				}
-
-				foreach (GlobResult glob in globs)
-				{
-					if (String.Equals("None", glob.ItemElement.ItemType, StringComparison.Ordinal))
-					{
-						// don't record the default "None" glob - it's not (?) a trigger for any rebuild
-						continue;
-					}
-
-					List<string> include = new List<string>(glob.IncludeGlobs.Select(f => CleanGlobString(f)).OrderBy(x => x));
-					List<string> exclude = new List<string>(glob.Excludes.Select(f => CleanGlobString(f)).OrderBy(x => x));
-					List<string> remove = new List<string>(glob.Removes.Select(f => CleanGlobString(f)).OrderBy(x => x));
-
-					buildRecord.Globs.Add(new CsProjBuildRecord.Glob()
-					{
-						ItemType = glob.ItemElement.ItemType,
-						Include = include,
-						Exclude = exclude,
-						Remove = remove
-					});
+					buildRecord.Dependencies.UnionWith(dependencies);
 				}
 
 				CsProjBuildRecordEntry entry = new CsProjBuildRecordEntry(projectPath, buildRecordPath, buildRecord);
-				buildRecords.Add(entry.ProjectFile, entry);
-			}
-
-			// Potential optimization: Constructing the ProjectGraph here gives the full graph of dependencies - which is nice,
-			// but not strictly necessary, and slower than doing it some other way.
-			ProjectGraph inputProjectGraph;
-			inputProjectGraph = new ProjectGraph(foundProjects
-				// Build the graph without anything that can't be built on this platform
-				.Where(x => !skippedProjects.Contains(x.FullName))
-				.Select(p => p.FullName), globalProperties, projectCollection);
+				buildRecords.AddOrUpdate(entry.ProjectFile, entry, (key, oldValue) => entry);
+			});
 
-			// A ProjectGraph that will represent the set of projects that we actually want to build
-			ProjectGraph buildProjectGraph = null;
+			logger.LogDebug("Generate build records {Count} projects time: {TimeSeconds:0.00} s", projects.Count, stopwatch.Elapsed.TotalSeconds);
 
 			if (bForceCompile)
 			{
 				logger.LogDebug("Script modules will build: '-Compile' on command line");
-				buildProjectGraph = inputProjectGraph;
+				outOfDateProjects = projects.Keys;
 			}
 			else
 			{
-				foreach (ProjectGraphNode project in inputProjectGraph.ProjectNodesTopologicallySorted)
+				stopwatch.Restart();
+				foreach (FileReference project in projects.Keys)
 				{
-					hook.ValidateRecursively(buildRecords, FileReference.FromString(project.ProjectInstance.FullPath));
+					hook.ValidateRecursively(buildRecords.ToDictionary(), project);
 				}
+				 logger.LogDebug("Validate {Count} projects time: {TimeSeconds:0.00} s", projects.Count, stopwatch.Elapsed.TotalSeconds);
 
 				// Select the projects that have been found to be out of date
-				Dictionary<FileReference, CsProjBuildRecordEntry> invalidBuildRecords = new(buildRecords.Where(x => x.Value.Status == CsProjBuildRecordStatus.Invalid));
-				HashSet<ProjectGraphNode> outOfDateProjects = new HashSet<ProjectGraphNode>(inputProjectGraph.ProjectNodes.Where(x => invalidBuildRecords.ContainsKey(FileReference.FromString(x.ProjectInstance.FullPath))));
-
-				if (outOfDateProjects.Count > 0)
-				{
-					buildProjectGraph = new ProjectGraph(outOfDateProjects.Select(p => p.ProjectInstance.FullPath), globalProperties, projectCollection);
-				}
+				outOfDateProjects = buildRecords.Where(x => x.Value.Status == CsProjBuildRecordStatus.Invalid).Select(x => x.Key);
 			}
 
-			if (buildProjectGraph != null)
+			if (outOfDateProjects.Any())
 			{
-				onBuildingProjects(buildProjectGraph.EntryPointNodes.Count);
-				bBuildSuccess = BuildProjects(buildProjectGraph, globalProperties, logger);
+				onBuildingProjects(outOfDateProjects.Count());
+
+				bBuildSuccess = RunDotnetBuildAsync(outOfDateProjects, globalProperties, hook, logger, default).Result;
 			}
 			else
 			{
 				bBuildSuccess = true;
 			}
 
-			// Update the target times
-			foreach (ProjectGraphNode projectNode in inputProjectGraph.ProjectNodes)
+			Parallel.ForEach(projects.Keys, projectPath =>
 			{
-				FileReference projectPath = FileReference.FromString(projectNode.ProjectInstance.FullPath);
 				CsProjBuildRecordEntry entry = buildRecords[projectPath];
+				// Update the target times
 				FileReference fullPath = FileReference.Combine(projectPath.Directory, entry.BuildRecord.TargetPath);
 				entry.BuildRecord.TargetBuildTime = FileReference.GetLastWriteTime(fullPath);
-			}
+			});
 
 			// Update the project reference target times
-			foreach (ProjectGraphNode projectNode in inputProjectGraph.ProjectNodes)
+			Parallel.ForEach(projects.Keys, projectPath =>
 			{
-				FileReference projectPath = FileReference.FromString(projectNode.ProjectInstance.FullPath);
 				CsProjBuildRecordEntry entry = buildRecords[projectPath];
 				foreach (CsProjBuildRecordRef referencedProject in entry.BuildRecord.ProjectReferencesAndTimes)
 				{
@@ -483,82 +282,223 @@ namespace EpicGames.MsBuild
 						referencedProject.TargetBuildTime = refEntry.BuildRecord.TargetBuildTime;
 					}
 				}
-			}
+			});
 
 			// write all build records
 			JsonSerializerOptions jsonOptions = new JsonSerializerOptions { WriteIndented = true };
-			foreach (ProjectGraphNode projectNode in inputProjectGraph.ProjectNodes)
+			Parallel.ForEach(projects.Keys, async projectPath =>
 			{
-				FileReference projectPath = FileReference.FromString(projectNode.ProjectInstance.FullPath);
 				CsProjBuildRecordEntry entry = buildRecords[projectPath];
-				if (FileReference.WriteAllTextIfDifferent(entry.BuildRecordFile, JsonSerializer.Serialize(entry.BuildRecord, jsonOptions)))
+				// write all build records
+				if (await FileReference.WriteAllTextIfDifferentAsync(entry.BuildRecordFile, JsonSerializer.Serialize(entry.BuildRecord, jsonOptions)))
 				{
 					logger.LogDebug("Wrote script module build record to {BuildRecordPath}", entry.BuildRecordFile);
 				}
-			}
+			});
 
 			// todo: re-verify build records after a build to verify that everything is actually up to date
 
 			// even if only a subset was built, this function returns the full list of target assembly paths
-			Dictionary<FileReference, CsProjBuildRecordEntry> outDict = [];
-			foreach (ProjectGraphNode entryPointNode in inputProjectGraph.EntryPointNodes)
+			return buildRecords.Where(x => foundProjects.Contains(x.Key)).ToDictionary();
+		}
+
+		private static async Task<FileReference> WriteMsBuildProjectAsync(MsBuildTarget target, IEnumerable<FileReference> projects, CancellationToken cancellationToken)
+		{
+			FileReference projectPath = FileReference.Combine(DirectoryReference.GetTempPath(), $"CsProj-{Environment.ProcessId}-{Guid.NewGuid()}.proj");
+
+			XNamespace ns = XNamespace.Get("http://schemas.microsoft.com/developer/msbuild/2003");
+			XDocument document = new XDocument(
+					new XElement(ns + "Project",
+						[
+							new XAttribute("DefaultTargets", "Build"), 
+							new XElement(ns + "PropertyGroup",
+								new XElement(ns + "Projects", String.Join(';', projects))
+							),
+							.. Enum.GetValues<MsBuildTarget>().Where(x => target.HasFlag(x)).Select(target =>
+								new XElement(ns + "Target",
+									new XAttribute("Name", target),
+									new XElement(ns + "MSBuild",
+										new XAttribute("Targets", target),
+										new XAttribute("BuildInParallel", "true"),
+										new XAttribute("Projects", "$(Projects)")
+									)
+								)
+							)
+						]
+					)
+			);
+
+			StringBuilder output = new StringBuilder();
+			output.AppendLine("<?xml version=\"1.0\" encoding=\"utf-8\"?>");
+
+			XmlWriterSettings xmlSettings = new XmlWriterSettings();
+			xmlSettings.Async = true;
+			xmlSettings.Encoding = new UTF8Encoding(false);
+			xmlSettings.Indent = true;
+			xmlSettings.OmitXmlDeclaration = true;
+
+			using (XmlWriter writer = XmlWriter.Create(output, xmlSettings))
+			{
+				await document.SaveAsync(writer, cancellationToken);
+			}
+			await FileReference.WriteAllTextIfDifferentAsync(projectPath, output.ToString(), cancellationToken);
+
+			return projectPath;
+		}
+
+		private static async Task<FileReference> WriteCsPropertiesAsync(IEnumerable<FileReference> projects, CsProjBuildHook hook, CancellationToken cancellationToken)
+		{
+			FileReference projectPath = FileReference.Combine(hook.EngineDirectory, "Intermediate", "Build", "EpicGames.ScriptBuild.props");
+
+			XNamespace ns = XNamespace.Get("http://schemas.microsoft.com/developer/msbuild/2003");
+			XDocument document = new XDocument(
+				new XElement(ns + "Project",
+					new XAttribute("ToolsVersion", "Current"),
+					new XElement(ns + "ItemGroup",
+						projects.Order().Select(project =>
+							new XElement(ns + "ProjectReference",
+								new XAttribute("Include", project)
+							)
+						)
+					)
+				)
+			);
+
+			// Create a stub .csproj in p4 (Shared?)
+			// Generate a property sheet in intermediate that references all projects to build, set them to not copy local
+			// dotnet build that .csproj, letting it restore + build normally and resolving the tree correctly
+			// msbuild can f off
+			StringBuilder output = new StringBuilder();
+			output.AppendLine("<?xml version=\"1.0\" encoding=\"utf-8\"?>");
+
+			XmlWriterSettings xmlSettings = new XmlWriterSettings();
+			xmlSettings.Async = true;
+			xmlSettings.Encoding = new UTF8Encoding(false);
+			xmlSettings.Indent = true;
+			xmlSettings.OmitXmlDeclaration = true;
+
+			using (XmlWriter writer = XmlWriter.Create(output, xmlSettings))
 			{
-				FileReference projectPath = FileReference.FromString(entryPointNode.ProjectInstance.FullPath);
-				outDict.Add(projectPath, buildRecords[projectPath]);
+				await document.SaveAsync(writer, cancellationToken);
 			}
-			return outDict;
+			await FileReference.WriteAllTextIfDifferentAsync(projectPath, output.ToString(), cancellationToken);
+
+			return projectPath;
 		}
 
-		private static bool BuildProjects(ProjectGraph projectGraph, Dictionary<string, string> globalProperties, ILogger logger)
+		private static async Task<bool> RunDotnetAsync(IEnumerable<string> arguments, CsProjBuildHook hook, ILogger logger, CancellationToken cancellationToken)
 		{
-			DateTime startTime = DateTime.UtcNow;
-			MLogger buildLogger = new MLogger(logger);
+			string outputPrefix = $"dotnet {arguments.First()}";
 
-			string[] targetsToBuild = ["Restore", "Build"];
+			ProcessStartInfo startInfo = new ProcessStartInfo
+			{
+				FileName = hook.DotnetPath.FullName,
+				Arguments = String.Join(' ', arguments),
+				WorkingDirectory = DirectoryReference.Combine(hook.EngineDirectory, "Source").FullName,
+				UseShellExecute = false,
+				CreateNoWindow = true,
+				RedirectStandardOutput = true,
+				RedirectStandardError = true,
+			};
+			startInfo.EnvironmentVariables["DOTNET_MULTILEVEL_LOOKUP"] = "0"; // use only the bundled dotnet installation - ignore any other/system dotnet install
+
+			logger.LogDebug("Running: {Application} {Arguments}", startInfo.FileName, startInfo.Arguments);
+			using Process dotnetProcess = new Process();
+			dotnetProcess.StartInfo = startInfo;
+			dotnetProcess.OutputDataReceived += (sender, args) =>
+			{
+				string output = args.Data?.TrimEnd() ?? String.Empty;
+				if (!String.IsNullOrEmpty(output))
+				{
+					logger.LogDebug("[{Prefix}] {Message}", outputPrefix, output);
+				}
+			};
+			dotnetProcess.ErrorDataReceived += (sender, args) =>
+			{
+				string output = args.Data?.TrimEnd() ?? String.Empty;
+				if (!String.IsNullOrEmpty(output))
+				{
+					logger.LogError("[{Prefix}] {Message}", outputPrefix, output);
+				}
+			};
+			dotnetProcess.Start();
+			dotnetProcess.BeginOutputReadLine();
+			dotnetProcess.BeginErrorReadLine();
+			await dotnetProcess.WaitForExitAsync(cancellationToken);
+			return dotnetProcess.ExitCode == 0;
 
-			bool result = true;
+		}
 
-			foreach (string targetToBuild in targetsToBuild)
+		private static async Task<bool> RunMsBuildAsync(MsBuildTarget target, IEnumerable<FileReference> projects, Dictionary<string, string> globalProperties, CsProjBuildHook hook, ILogger logger, CancellationToken cancellationToken)
+		{
+			if (!projects.Any())
 			{
-				GraphBuildRequestData graphRequest = new GraphBuildRequestData(projectGraph, [targetToBuild]);
+				return true;
+			}
+			Stopwatch stopwatch = new();
+			stopwatch.Start();
 
-				BuildManager buildMan = BuildManager.DefaultBuildManager;
+			FileReference projectPath = await WriteMsBuildProjectAsync(target, projects, cancellationToken);
 
-				BuildParameters buildParameters = new BuildParameters();
-				buildParameters.AllowFailureWithoutError = false;
-				buildParameters.DetailedSummary = true;
+			IEnumerable<MsBuildTarget> enumValues = Enum.GetValues<MsBuildTarget>().Where(x => target.HasFlag(x));
 
-				buildParameters.Loggers = [buildLogger];
-				buildParameters.MaxNodeCount = 1; // msbuild bug - more than 1 here and the build stalls. Likely related to [REDACTED_PROFILE_URL]/msbuild/issues/1941
+			IEnumerable<string> arguments = [
+				"msbuild",
+				$"\"{projectPath}\"",
+				"-nologo",
+				"-v:quiet",
+				"-p:RestoreRecursive=false",
+				.. enumValues.Select(x => $"-t:{x}"),
+				.. globalProperties.Select(x => $"\"-p:{x.Key}={x.Value}\"")
+			];
 
-				buildParameters.OnlyLogCriticalEvents = false;
-				buildParameters.ShutdownInProcNodeOnBuildFinish = false;
+			bool result = await RunDotnetAsync(arguments, hook, logger, cancellationToken);
 
-				buildParameters.GlobalProperties = globalProperties;
+			logger.LogInformation("{Target} {Count} projects time: {TimeSeconds:0.00} s", String.Join('+', enumValues), projects.Count(), stopwatch.Elapsed.TotalSeconds);
 
-				logger.LogInformation(" {TargetToBuild}...", targetToBuild);
+			FileReference.Delete(projectPath);
+			return result;
+		}
 
-				GraphBuildResult buildResult = buildMan.Build(buildParameters, graphRequest);
+		private static async Task<bool> RunDotnetBuildAsync(IEnumerable<FileReference> projects, Dictionary<string, string> globalProperties, CsProjBuildHook hook, ILogger logger, CancellationToken cancellationToken)
+		{
+			if (!projects.Any())
+			{
+				return true;
+			}
 
-				if (buildResult.OverallResult == BuildResultCode.Failure)
-				{
-					logger.LogInformation("");
-					foreach (KeyValuePair<ProjectGraphNode, BuildResult> nodeResult in buildResult.ResultsByNode)
-					{
-						if (nodeResult.Value.OverallResult == BuildResultCode.Failure)
-						{
-							logger.LogError("  Failed to build: {ProjectPath}", new FileReference(nodeResult.Key.ProjectInstance.FullPath));
-						}
-					}
-					result = false;
-				}
+			// Acquire a mutex to prevent running builds from different processes concurrently
+			Task<IDisposable> mutexTask = SingleInstanceMutex.AcquireAsync($"Global\\CsProj_{Sha1Hash.Compute(Encoding.Default.GetBytes(hook.EngineDirectory.FullName)).ToString()}", cancellationToken);
+
+			Task delayTask = Task.Delay(TimeSpan.FromSeconds(1.0), cancellationToken);
+			if (Task.WhenAny(mutexTask, delayTask) == delayTask)
+			{
+				logger.LogInformation("dotnet build is already running. Waiting for it to terminate...");
 			}
-			logger.LogInformation("Build projects time: {TimeSeconds:0.00} s", (DateTime.UtcNow - startTime).TotalMilliseconds / 1000);
+
+			using IDisposable mutex = await mutexTask;
+
+			Stopwatch stopwatch = new();
+			stopwatch.Start();
+
+			FileReference propertiesPath = await WriteCsPropertiesAsync(projects, hook, cancellationToken);
+
+			IEnumerable<string> arguments = [
+				"build",
+				$"\"{FileReference.Combine(hook.EngineDirectory, "Source", "Programs", "Shared", "EpicGames.ScriptBuild")}\"",
+				"-nologo",
+				"-v:quiet",
+				.. globalProperties.Select(x => $"\"/p:{x.Key}={x.Value}\""),
+			];
+
+			bool result = await RunDotnetAsync(arguments, hook, logger, cancellationToken);
+
+			logger.LogInformation("Build {Count} projects time: {TimeSeconds:0.00} s", projects.Count(), stopwatch.Elapsed.TotalSeconds);
 
 			return result;
 		}
 
-		static bool s_hasRegiteredMsBuildPath = false;
+		static bool s_hasRegisteredMsBuildPath = false;
 
 		/// <summary>
 		/// Register our bundled dotnet installation to be used by Microsoft.Build
@@ -566,11 +506,11 @@ namespace EpicGames.MsBuild
 		/// </summary>
 		public static void RegisterMsBuildPath(CsProjBuildHook hook)
 		{
-			if (s_hasRegiteredMsBuildPath)
+			if (s_hasRegisteredMsBuildPath)
 			{
 				return;
 			}
-			s_hasRegiteredMsBuildPath = true;
+			s_hasRegisteredMsBuildPath = true;
 
 			// Find our bundled dotnet SDK
 			List<string> listOfSdks = [];
@@ -605,7 +545,7 @@ namespace EpicGames.MsBuild
 			DirectoryReference dotnetSdkDirectory = DirectoryReference.Combine(hook.DotnetDirectory, "sdk", sdkVersion);
 			if (!DirectoryReference.Exists(dotnetSdkDirectory))
 			{
-				throw new Exception("Failed to find .NET SDK directory: " + dotnetSdkDirectory.FullName);
+				throw new Exception($"Failed to find .NET SDK directory: {dotnetSdkDirectory.FullName}");
 			}
 
 			MSBuildLocator.RegisterMSBuildPath(dotnetSdkDirectory.FullName);
diff --git a/Engine/Source/Programs/Shared/EpicGames.ScriptBuild/EpicGames.ScriptBuild.csproj b/Engine/Source/Programs/Shared/EpicGames.ScriptBuild/EpicGames.ScriptBuild.csproj
new file mode 100644
--- /dev/null
+++ b/Engine/Source/Programs/Shared/EpicGames.ScriptBuild/EpicGames.ScriptBuild.csproj
@@ -0,0 +1,46 @@
+﻿<Project Sdk="Microsoft.NET.Sdk">
+  <Import Project="..\UnrealEngine.csproj.props" />
+  
+  <PropertyGroup>
+    <TargetFramework>net8.0</TargetFramework>
+    <Configuration Condition=" '$(Configuration)' == '' ">Development</Configuration>
+    <OutputType>Library</OutputType>
+    <Deterministic>true</Deterministic>
+    <RootNamespace>EpicGames.ScriptBuild</RootNamespace>
+    <GenerateAssemblyInfo>false</GenerateAssemblyInfo>
+    <GenerateTargetFrameworkAttribute>false</GenerateTargetFrameworkAttribute>
+    <Configurations>Debug;Release;Development;Analyze</Configurations>
+    <Nullable>enable</Nullable>
+    <DebugType>full</DebugType>
+    <DebugType>pdbonly</DebugType>
+    <DebugType Condition="'$([System.Runtime.InteropServices.RuntimeInformation]::IsOSPlatform($([System.Runtime.InteropServices.OSPlatform]::Windows)))' == 'true' And '$([System.Runtime.InteropServices.RuntimeInformation]::OSArchitecture)' == 'Arm64'">portable</DebugType>
+  </PropertyGroup>
+  <PropertyGroup Condition=" '$(Configuration)|$(Platform)' == 'Debug|AnyCPU' ">
+    <DocumentationFile></DocumentationFile>
+  </PropertyGroup>
+  <PropertyGroup Condition="'$(Configuration)|$(Platform)' == 'Development|AnyCPU'">
+    <DocumentationFile></DocumentationFile>
+    <Optimize>true</Optimize>
+  </PropertyGroup>
+  <PropertyGroup Condition="'$(Configuration)|$(Platform)' == 'Release|AnyCPU'">
+    <DocumentationFile></DocumentationFile>
+    <Optimize>true</Optimize>
+  </PropertyGroup>
+  <PropertyGroup Condition=" '$(Configuration)|$(Platform)' == 'Analyze|AnyCPU' ">
+    <RunAnalyzersDuringBuild>True</RunAnalyzersDuringBuild>
+    <TreatWarningsAsErrors>False</TreatWarningsAsErrors>
+    <EnforceCodeStyleInBuild>True</EnforceCodeStyleInBuild>
+    <GenerateDocumentationFile>True</GenerateDocumentationFile>
+    <DocumentationFile></DocumentationFile>
+  </PropertyGroup>
+  <ItemGroup>
+    <Compile Include="..\..\Shared\MetaData.cs">
+      <Link>Properties\MetaData.cs</Link>
+    </Compile>
+  </ItemGroup>
+
+  <!-- Reference file is generated by AutomationTool. It contains a list of project references that will be compiled with 'dotnet build' -->
+  <Import
+    Project="$(EngineDir)\Intermediate\Build\EpicGames.ScriptBuild.props"
+    Condition="Exists('$(EngineDir)\Intermediate\Build\EpicGames.ScriptBuild.props') And '$(EpicGamesMsBuild)' == 'true'" />
+</Project>
\ No newline at end of file
