Added verification that interfaces have matching native interface base classes

diff --git a/Engine/Source/Programs/Shared/EpicGames.UHT/Exporters/CodeGen/UhtHeaderCodeGeneratorCppFile.cs b/Engine/Source/Programs/Shared/EpicGames.UHT/Exporters/CodeGen/UhtHeaderCodeGeneratorCppFile.cs
--- a/Engine/Source/Programs/Shared/EpicGames.UHT/Exporters/CodeGen/UhtHeaderCodeGeneratorCppFile.cs
+++ b/Engine/Source/Programs/Shared/EpicGames.UHT/Exporters/CodeGen/UhtHeaderCodeGeneratorCppFile.cs
@@ -1354,9 +1354,9 @@ namespace EpicGames.UHT.Exporters.CodeGen
 				builder
 					.Append("\telse if (auto I = (")
 					.Append(function.FunctionFlags.HasAnyFlags(EFunctionFlags.Const) ? "const I" : "I")
-					.Append(classObj.EngineName)
-					.Append("*)(O->GetNativeInterfaceAddress(U")
-					.Append(classObj.EngineName)
+					.Append(classObj.SourceName[1..])
+					.Append("*)(O->GetNativeInterfaceAddress(")
+					.Append(classObj.SourceName)
 					.Append("::StaticClass())))\r\n");
 				builder.Append("\t{\r\n");
 				builder.Append("\t\t");
diff --git a/Engine/Source/Programs/Shared/EpicGames.UHT/Exporters/CodeGen/UhtHeaderCodeGeneratorHFile.cs b/Engine/Source/Programs/Shared/EpicGames.UHT/Exporters/CodeGen/UhtHeaderCodeGeneratorHFile.cs
--- a/Engine/Source/Programs/Shared/EpicGames.UHT/Exporters/CodeGen/UhtHeaderCodeGeneratorHFile.cs
+++ b/Engine/Source/Programs/Shared/EpicGames.UHT/Exporters/CodeGen/UhtHeaderCodeGeneratorHFile.cs
@@ -1516,7 +1516,7 @@ namespace EpicGames.UHT.Exporters.CodeGen
 
 		private StringBuilder AppendInClassIInterface(StringBuilder builder, UhtClass classObj, List<UhtFunction> callbackFunctions, string api)
 		{
-			string interfaceSourceName = "I" + classObj.EngineName;
+			string interfaceSourceName = classObj.GetNativeInterfaceName();
 
 			builder.Append("protected: \\\r\n");
 			builder.Append("\tvirtual ~").Append(interfaceSourceName).Append("() {} \\\r\n");
diff --git a/Engine/Source/Programs/Shared/EpicGames.UHT/Types/UhtClass.cs b/Engine/Source/Programs/Shared/EpicGames.UHT/Types/UhtClass.cs
--- a/Engine/Source/Programs/Shared/EpicGames.UHT/Types/UhtClass.cs
+++ b/Engine/Source/Programs/Shared/EpicGames.UHT/Types/UhtClass.cs
@@ -753,6 +753,34 @@ namespace EpicGames.UHT.Types
 			}
 		}
 
+		/// <summary>
+		/// Return the native interface name from an interface
+		/// </summary>
+		/// <returns>Native interface name</returns>
+		/// <exception cref="UhtIceException"></exception>
+		public string GetNativeInterfaceName()
+		{
+			if (ClassType != UhtClassType.Interface)
+			{
+				throw new UhtIceException("Can only get the native interface name from interfaces");
+			}
+			return "I" + SourceName[1..];
+		}
+
+		/// <summary>
+		/// Return the interface name from a native interface
+		/// </summary>
+		/// <returns>Interface name</returns>
+		/// <exception cref="UhtIceException"></exception>
+		public string GetInterfaceName()
+		{
+			if (ClassType != UhtClassType.NativeInterface)
+			{
+				throw new UhtIceException("Can only get the interface name from native interfaces");
+			}
+			return "U" + SourceName[1..];
+		}
+
 		/// <inheritdoc/>
 		protected override bool ResolveSelf(UhtResolvePhase phase)
 		{
@@ -768,7 +796,7 @@ namespace EpicGames.UHT.Types
 
 						case UhtClassType.Interface:
 							{
-								string nativeInterfaceName = "I" + EngineName;
+								string nativeInterfaceName = GetNativeInterfaceName();
 								UhtType? nativeInterface = Session.FindType(null, UhtFindOptions.SourceName | UhtFindOptions.Class, nativeInterfaceName);
 								if (nativeInterface == null)
 								{
@@ -784,7 +812,7 @@ namespace EpicGames.UHT.Types
 
 						case UhtClassType.NativeInterface:
 							{
-								string interfaceName = "U" + EngineName;
+								string interfaceName = GetInterfaceName();
 								UhtClass? interfaceObj = (UhtClass?)Session.FindType(null, UhtFindOptions.SourceName | UhtFindOptions.Class, interfaceName);
 								if (interfaceObj == null)
 								{
@@ -1498,7 +1526,7 @@ namespace EpicGames.UHT.Types
 			{
 				if (ClassType == UhtClassType.Interface)
 				{
-					this.LogError($"Interface '{SourceName}' must derive from an interface");
+					this.LogError($"Interface '{SourceName}' must derive from an interface or UInterface");
 				}
 			}
 
@@ -1639,6 +1667,26 @@ namespace EpicGames.UHT.Types
 
 				case UhtClassType.Interface:
 					{
+
+						// Verify that our supers match
+						if (NativeInterface != null && SuperClass != null)
+						{
+							if (SuperClass == Session.UInterface)
+							{
+								if (NativeInterface.SuperClass != null && NativeInterface.SuperClass != Session.IInterface)
+								{
+									this.LogError($"Interface {SourceName} derives from {SuperClass.SourceName}, therefore {NativeInterface.SourceName} can not derive from an another native interface");
+								}
+							}
+							else if (SuperClass.NativeInterface != null)
+							{
+								if (SuperClass.NativeInterface != NativeInterface.SuperClass)
+								{
+									this.LogError($"Interface {SourceName} derives from {SuperClass.SourceName}, therefore {NativeInterface.SourceName} must derive from {SuperClass.NativeInterface.SourceName}");
+								}
+							}
+						}
+
 						// Interface with blueprint data should declare explicitly Blueprintable or NotBlueprintable to be clear
 						// In the backward compatible case where they declare neither, both of these bools are false
 						bool canImplementInBlueprints = CanImplementInBlueprints();
