Removed TargetFramework override in "Tests" .csproj files
PR #12794

diff --git a/Engine/Source/Programs/Shared/EpicGames.BuildGraph.Tests/EpicGames.BuildGraph.Tests.csproj b/Engine/Source/Programs/Shared/EpicGames.BuildGraph.Tests/EpicGames.BuildGraph.Tests.csproj
--- a/Engine/Source/Programs/Shared/EpicGames.BuildGraph.Tests/EpicGames.BuildGraph.Tests.csproj
+++ b/Engine/Source/Programs/Shared/EpicGames.BuildGraph.Tests/EpicGames.BuildGraph.Tests.csproj
@@ -2,7 +2,6 @@
 
   <PropertyGroup>
     <TargetFramework>net8.0</TargetFramework>
-	<TargetFramework Condition="'$(UE_DOTNET_VERSION)' != ''">$(UE_DOTNET_VERSION)</TargetFramework>
     <Configuration Condition=" '$(Configuration)' == '' ">Development</Configuration>
     <OutputType>Library</OutputType>
     <Deterministic>true</Deterministic>
diff --git a/Engine/Source/Programs/Shared/EpicGames.Core.Tests/EpicGames.Core.Tests.csproj b/Engine/Source/Programs/Shared/EpicGames.Core.Tests/EpicGames.Core.Tests.csproj
--- a/Engine/Source/Programs/Shared/EpicGames.Core.Tests/EpicGames.Core.Tests.csproj
+++ b/Engine/Source/Programs/Shared/EpicGames.Core.Tests/EpicGames.Core.Tests.csproj
@@ -3,7 +3,6 @@
 
   <PropertyGroup>
     <TargetFramework>net8.0</TargetFramework>
-	<TargetFramework Condition="'$(UE_DOTNET_VERSION)' != ''">$(UE_DOTNET_VERSION)</TargetFramework>
     <Configuration Condition=" '$(Configuration)' == '' ">Development</Configuration>
     <OutputType>Library</OutputType>
     <Deterministic>true</Deterministic>
diff --git a/Engine/Source/Programs/Shared/EpicGames.Horde.Tests/EpicGames.Horde.Tests.csproj b/Engine/Source/Programs/Shared/EpicGames.Horde.Tests/EpicGames.Horde.Tests.csproj
--- a/Engine/Source/Programs/Shared/EpicGames.Horde.Tests/EpicGames.Horde.Tests.csproj
+++ b/Engine/Source/Programs/Shared/EpicGames.Horde.Tests/EpicGames.Horde.Tests.csproj
@@ -3,7 +3,6 @@
 
   <PropertyGroup>
     <TargetFramework>net8.0</TargetFramework>
-	<TargetFramework Condition="'$(UE_DOTNET_VERSION)' != ''">$(UE_DOTNET_VERSION)</TargetFramework>
     <Configurations>Debug;Release;Development;Analyze</Configurations>
     <Nullable>enable</Nullable>
     <IsPackable>false</IsPackable>
diff --git a/Engine/Source/Programs/Shared/EpicGames.IoHash.Tests/EpicGames.IoHash.Tests.csproj b/Engine/Source/Programs/Shared/EpicGames.IoHash.Tests/EpicGames.IoHash.Tests.csproj
--- a/Engine/Source/Programs/Shared/EpicGames.IoHash.Tests/EpicGames.IoHash.Tests.csproj
+++ b/Engine/Source/Programs/Shared/EpicGames.IoHash.Tests/EpicGames.IoHash.Tests.csproj
@@ -3,7 +3,6 @@
 
   <PropertyGroup>
     <TargetFramework>net8.0</TargetFramework>
-	<TargetFramework Condition="'$(UE_DOTNET_VERSION)' != ''">$(UE_DOTNET_VERSION)</TargetFramework>
     <Configurations>Debug;Release;Development;Analyze</Configurations>
     <Nullable>enable</Nullable>
     <IsPackable>false</IsPackable>
diff --git a/Engine/Source/Programs/Shared/EpicGames.Redis.Tests/EpicGames.Redis.Tests.csproj b/Engine/Source/Programs/Shared/EpicGames.Redis.Tests/EpicGames.Redis.Tests.csproj
--- a/Engine/Source/Programs/Shared/EpicGames.Redis.Tests/EpicGames.Redis.Tests.csproj
+++ b/Engine/Source/Programs/Shared/EpicGames.Redis.Tests/EpicGames.Redis.Tests.csproj
@@ -2,7 +2,6 @@
 
   <PropertyGroup>
     <TargetFramework>net8.0</TargetFramework>
-	<TargetFramework Condition="'$(UE_DOTNET_VERSION)' != ''">$(UE_DOTNET_VERSION)</TargetFramework>
     <Configuration Condition=" '$(Configuration)' == '' ">Development</Configuration>
     <OutputType>Library</OutputType>
     <Deterministic>true</Deterministic>
diff --git a/Engine/Source/Programs/Shared/EpicGames.Serialization.Tests/EpicGames.Serialization.Tests.csproj b/Engine/Source/Programs/Shared/EpicGames.Serialization.Tests/EpicGames.Serialization.Tests.csproj
--- a/Engine/Source/Programs/Shared/EpicGames.Serialization.Tests/EpicGames.Serialization.Tests.csproj
+++ b/Engine/Source/Programs/Shared/EpicGames.Serialization.Tests/EpicGames.Serialization.Tests.csproj
@@ -3,7 +3,6 @@
 
   <PropertyGroup>
     <TargetFramework>net8.0</TargetFramework>
-	<TargetFramework Condition="'$(UE_DOTNET_VERSION)' != ''">$(UE_DOTNET_VERSION)</TargetFramework>
     <Configuration Condition=" '$(Configuration)' == '' ">Development</Configuration>
     <OutputType>Library</OutputType>
     <Deterministic>true</Deterministic>
diff --git a/Engine/Source/Programs/Shared/UnrealBuildTool.Tests/UnrealBuildTool.Tests.csproj b/Engine/Source/Programs/Shared/UnrealBuildTool.Tests/UnrealBuildTool.Tests.csproj
--- a/Engine/Source/Programs/Shared/UnrealBuildTool.Tests/UnrealBuildTool.Tests.csproj
+++ b/Engine/Source/Programs/Shared/UnrealBuildTool.Tests/UnrealBuildTool.Tests.csproj
@@ -3,7 +3,6 @@
 
   <PropertyGroup>
     <TargetFramework>net8.0</TargetFramework>
-    <TargetFramework Condition="'$(UE_DOTNET_VERSION)' != ''">$(UE_DOTNET_VERSION)</TargetFramework>
     <Configurations>Debug;Release;Development;Analyze</Configurations>
     <Nullable>enable</Nullable>
     <IsPackable>false</IsPackable>
