Horde: Extend compute socket test timeouts to eliminate flakiness

Tests were intermittently failing when CI resources were under load.

diff --git a/Engine/Source/Programs/Shared/EpicGames.Horde.Tests/Compute/ComputeSocketTests.cs b/Engine/Source/Programs/Shared/EpicGames.Horde.Tests/Compute/ComputeSocketTests.cs
--- a/Engine/Source/Programs/Shared/EpicGames.Horde.Tests/Compute/ComputeSocketTests.cs
+++ b/Engine/Source/Programs/Shared/EpicGames.Horde.Tests/Compute/ComputeSocketTests.cs
@@ -25,7 +25,7 @@ namespace EpicGames.Horde.Tests.Compute
 	[DoNotParallelize]
 	public class ComputeSocketTests
 	{
-		public const int TestTimoutMs = 20000;
+		public const int TestTimoutMs = 40000;
 		class TestComputeSocket : ComputeSocket, IDisposable
 		{
 			public Dictionary<int, ComputeBufferWriter> RecvBufferWriters { get; } = new Dictionary<int, ComputeBufferWriter>();
@@ -132,8 +132,8 @@ namespace EpicGames.Horde.Tests.Compute
 			await using AesTransport localAes = new(localTcp, key);
 			await using AesTransport agentAes = new(agentTcp, key);
 			
-			await using IdleTimeoutTransport localIdleTimeout = new(localAes, TimeSpan.FromSeconds(5));
-			await using IdleTimeoutTransport agentIdleTimeout = new(agentAes, TimeSpan.FromSeconds(5));
+			await using IdleTimeoutTransport localIdleTimeout = new(localAes, TimeSpan.FromSeconds(15));
+			await using IdleTimeoutTransport agentIdleTimeout = new(agentAes, TimeSpan.FromSeconds(15));
 			await using RemoteComputeSocket localComputeSocket = new(localIdleTimeout, ComputeProtocol.Latest, new TestLogger("local"));
 			await using RemoteComputeSocket agentComputeSocket = new(agentIdleTimeout, ComputeProtocol.Latest, new TestLogger("agent"));
 
