EpicGames.Core: Fix thread safety issues with default logger

#rnx

diff --git a/Engine/Source/Programs/Shared/EpicGames.Core/Log.cs b/Engine/Source/Programs/Shared/EpicGames.Core/Log.cs
--- a/Engine/Source/Programs/Shared/EpicGames.Core/Log.cs
+++ b/Engine/Source/Programs/Shared/EpicGames.Core/Log.cs
@@ -264,7 +264,6 @@ namespace EpicGames.Core
 		/// </summary>
 		/// <param name="name">Identifier for the writer</param>
 		/// <param name="outputFile">The file to write to</param>
-		/// <returns>The created trace listener</returns>
 		public static void AddFileWriter(string name, FileReference outputFile)
 		{
 			Logger.LogInformation("Log file: {OutputFile}", outputFile);
@@ -280,19 +279,9 @@ namespace EpicGames.Core
 		/// </summary>
 		/// <param name="name">Identifier for the writer</param>
 		/// <param name="outputFile">The file to write to</param>
-		/// <returns>The created trace listener</returns>
 		public static void AddFileWriterWithoutBackup(string name, FileReference outputFile)
 		{
-			TextWriterTraceListener firstTextWriter = s_defaultLogger.AddFileWriter(name, outputFile);
-			
-			// find the StartupTraceListener in the listeners that was added early on
-			IEnumerable<StartupTraceListener> startupListeners = Trace.Listeners.OfType<StartupTraceListener>();
-			if (startupListeners.Any())
-			{
-				StartupTraceListener startupListener = startupListeners.First();
-				startupListener.CopyTo(firstTextWriter);
-				Trace.Listeners.Remove(startupListener);
-			}
+			s_defaultLogger.AddFileWriterWithoutBackup(name, outputFile);
 		}
 
 		/// <summary>
@@ -313,13 +302,21 @@ namespace EpicGames.Core
 			s_defaultLogger.RemoveTraceListener(traceListener);
 		}
 
+		/// <summary>
+		/// Removes a <see cref="StartupTraceListener"/> from the collection in a safe manner.
+		/// </summary>
+		public static void RemoveStartupTraceListener()
+		{
+			s_defaultLogger.RemoveStartupTraceListener();
+		}
+
 		/// <summary>
 		/// Determines if a TextWriterTraceListener has been added to the list of trace listeners
 		/// </summary>
 		/// <returns>True if a TextWriterTraceListener has been added</returns>
 		public static bool HasFileWriter()
 		{
-			return DefaultLogger.HasFileWriter();
+			return s_defaultLogger.HasFileWriter();
 		}
 
 		/// <summary>
@@ -1218,19 +1215,54 @@ namespace EpicGames.Core
 		}
 
 		/// <summary>
-		/// Determines if a TextWriterTraceListener has been added to the list of trace listeners
+		/// Removes <see cref="StartupTraceListener"/> from the collection in a safe manner.
 		/// </summary>
-		/// <returns>True if a TextWriterTraceListener has been added</returns>
-		public static bool HasFileWriter()
+		public void RemoveStartupTraceListener()
 		{
-			foreach (TraceListener? listener in Trace.Listeners)
+			lock (_syncObject)
 			{
-				if (listener is TextWriterTraceListener)
+				IEnumerable<StartupTraceListener> startupListeners = Trace.Listeners.OfType<StartupTraceListener>();
+				if (startupListeners.Any())
 				{
-					return true;
+					Trace.Listeners.Remove(startupListeners.First());
 				}
 			}
-			return false;
+		}
+
+		/// <summary>
+		/// Adds a trace listener that writes to a log file.
+		/// If a StartupTraceListener was in use, this function will copy its captured data to the log file(s)
+		/// and remove the startup listener from the list of registered listeners.
+		/// </summary>
+		/// <param name="name">Identifier for the writer</param>
+		/// <param name="outputFile">The file to write to</param>
+		public void AddFileWriterWithoutBackup(string name, FileReference outputFile)
+		{
+			TextWriterTraceListener firstTextWriter = AddFileWriter(name, outputFile);
+
+			lock (_syncObject)
+			{
+				// find the StartupTraceListener in the listeners that was added early on
+				IEnumerable<StartupTraceListener> startupListeners = Trace.Listeners.OfType<StartupTraceListener>();
+				if (startupListeners.Any())
+				{
+					StartupTraceListener startupListener = startupListeners.First();
+					startupListener.CopyTo(firstTextWriter);
+					Trace.Listeners.Remove(startupListeners.First());
+				}
+			}
+		}
+
+		/// <summary>
+		/// Determines if a TextWriterTraceListener has been added to the list of trace listeners
+		/// </summary>
+		/// <returns>True if a TextWriterTraceListener has been added</returns>
+		public bool HasFileWriter()
+		{
+			lock (_syncObject)
+			{
+				return Trace.Listeners.OfType<TextWriterTraceListener>().Any();
+			}
 		}
 
 		/// <inheritdoc/>
diff --git a/Engine/Source/Programs/UnrealBuildTool/Modes/BuildMode.cs b/Engine/Source/Programs/UnrealBuildTool/Modes/BuildMode.cs
--- a/Engine/Source/Programs/UnrealBuildTool/Modes/BuildMode.cs
+++ b/Engine/Source/Programs/UnrealBuildTool/Modes/BuildMode.cs
@@ -142,11 +142,7 @@ namespace UnrealBuildTool
 			}
 			else
 			{
-				IEnumerable<StartupTraceListener> StartupListeners = Trace.Listeners.OfType<StartupTraceListener>();
-				if (StartupListeners.Any())
-				{
-					Trace.Listeners.Remove(StartupListeners.First());
-				}
+				Log.RemoveStartupTraceListener();
 			}
 
 			// Create the build configuration object, and read the settings
diff --git a/Engine/Source/Programs/UnrealBuildTool/Modes/IWYUMode.cs b/Engine/Source/Programs/UnrealBuildTool/Modes/IWYUMode.cs
--- a/Engine/Source/Programs/UnrealBuildTool/Modes/IWYUMode.cs
+++ b/Engine/Source/Programs/UnrealBuildTool/Modes/IWYUMode.cs
@@ -366,11 +366,7 @@ namespace UnrealBuildTool
 			}
 			else
 			{
-				IEnumerable<StartupTraceListener> StartupListeners = Trace.Listeners.OfType<StartupTraceListener>();
-				if (StartupListeners.Any())
-				{
-					Trace.Listeners.Remove(StartupListeners.First());
-				}
+				Log.RemoveStartupTraceListener();
 			}
 
 			// Create the build configuration object, and read the settings
