[UBA]
* Removed cache update in visualizer if no cache connected

diff --git a/Engine/Source/Programs/Shared/EpicGames.UBA/Impl/CacheClientImpl.cs b/Engine/Source/Programs/Shared/EpicGames.UBA/Impl/CacheClientImpl.cs
--- a/Engine/Source/Programs/Shared/EpicGames.UBA/Impl/CacheClientImpl.cs
+++ b/Engine/Source/Programs/Shared/EpicGames.UBA/Impl/CacheClientImpl.cs
@@ -24,10 +24,10 @@ namespace EpicGames.UBA.Impl
 		static extern nint CacheClient_Create(nint session, bool reportMissReason, string crypto, string hint);
 
 		[DllImport("UbaHost", CharSet = CharSet.Auto)]
-		static extern bool CacheClient_Connect(nint cacheClient, string host, int port);
+		static extern byte CacheClient_Connect(nint cacheClient, string host, int port);
 
 		[DllImport("UbaHost", CharSet = CharSet.Auto)]
-		static extern bool CacheClient_WriteToCache2(nint cacheClient, uint bucket, nint info, byte[] inputs, uint inputsSize, byte[] outputs, uint outputsSize);
+		static extern byte CacheClient_WriteToCache2(nint cacheClient, uint bucket, nint info, byte[] inputs, uint inputsSize, byte[] outputs, uint outputsSize);
 
 		[DllImport("UbaHost", CharSet = CharSet.Auto)]
 		static extern nint CacheClient_FetchFromCache3(nint cacheClient, ulong rootsHandle, uint bucket, nint info);
@@ -79,12 +79,12 @@ namespace EpicGames.UBA.Impl
 
 		public bool Connect(string host, int port)
 		{
-			return CacheClient_Connect(_handle, host, port);
+			return CacheClient_Connect(_handle, host, port) != 0;
 		}
 
 		public bool WriteToCache(uint bucket, IProcess process, byte[] inputs, uint inputsSize, byte[] outputs, uint outputsSize)
 		{
-			return CacheClient_WriteToCache2(_handle, bucket, process.GetHandle(), inputs, inputsSize, outputs, outputsSize);
+			return CacheClient_WriteToCache2(_handle, bucket, process.GetHandle(), inputs, inputsSize, outputs, outputsSize) != 0;
 		}
 
 		public FetchFromCacheResult FetchFromCache(ulong rootsHandle, uint bucket, ProcessStartInfo info)
diff --git a/Engine/Source/Programs/UnrealBuildTool/Executors/UnrealBuildAccelerator/UBAExecutor.cs b/Engine/Source/Programs/UnrealBuildTool/Executors/UnrealBuildAccelerator/UBAExecutor.cs
--- a/Engine/Source/Programs/UnrealBuildTool/Executors/UnrealBuildAccelerator/UBAExecutor.cs
+++ b/Engine/Source/Programs/UnrealBuildTool/Executors/UnrealBuildAccelerator/UBAExecutor.cs
@@ -439,7 +439,7 @@ namespace UnrealBuildTool
 								{
 									_cacheClient.Dispose();
 									_cacheClient = null;
-									_session.UpdateStatus(1, 6, "Not connected", LogEntryType.Info);
+									_session.UpdateStatus(1, 6, "Failed to connect", LogEntryType.Info);
 								}
 							}
 
@@ -552,8 +552,11 @@ namespace UnrealBuildTool
 			ImmediateActionQueueRunner crossArchitectureRemoteRunner = queue.CreateManualRunner(action => RunActionRemote(queue, action, true));
 			queue.CancellationToken.Register(onCancel);
 
-			queue.OnArtifactsRead += (action) => { UpdateCacheProgress(queue); UpdateProgress(queue); };
-			queue.OnArtifactsMiss += (action) => UpdateCacheProgress(queue);
+			if (actionArtifactCache != null)
+			{
+				queue.OnArtifactsRead += (action) => { UpdateCacheProgress(queue); UpdateProgress(queue); };
+				queue.OnArtifactsMiss += (action) => UpdateCacheProgress(queue);
+			}
 
 			UpdateProgress(queue);
 
