Horde: Add retry mechanism for agent socket connection in relay mode

diff --git a/Engine/Source/Programs/Shared/EpicGames.Horde/Compute/Clients/ServerComputeClient.cs b/Engine/Source/Programs/Shared/EpicGames.Horde/Compute/Clients/ServerComputeClient.cs
--- a/Engine/Source/Programs/Shared/EpicGames.Horde/Compute/Clients/ServerComputeClient.cs
+++ b/Engine/Source/Programs/Shared/EpicGames.Horde/Compute/Clients/ServerComputeClient.cs
@@ -344,30 +344,34 @@ namespace EpicGames.Horde.Compute.Clients
 				}
 			}
 
-			string agentAddress = $"{response.Ip}:{response.Port}";
+			string agentHost = response.Ip;
+			int agentPort = response.Port;
 
 			// Connect to the remote machine
-			using Socket socket = new Socket(SocketType.Stream, ProtocolType.Tcp);
+			using Socket socket = new (SocketType.Stream, ProtocolType.Tcp);
 
-			workerLogger.LogDebug("Connecting to {AgentId} at {AgentAddress} ({ConnectionType} via {ConnectionAddress}) with nonce {Nonce} and encryption {Encryption}...",
-				response.AgentId, agentAddress, response.ConnectionMode, response.ConnectionAddress ?? "None", response.Nonce, response.Encryption);
+			workerLogger.LogDebug("Connecting to {AgentId} at {AgentHost}{AgentPort} ({ConnectionType} via {ConnectionAddress}) with nonce {Nonce} and encryption {Encryption}. Lease ID {LeaseId}...",
+				response.AgentId, agentHost, agentPort, response.ConnectionMode, response.ConnectionAddress ?? "None", response.Nonce, response.Encryption, response.LeaseId);
+			
 			try
 			{
 				switch (response.ConnectionMode)
 				{
 					case ConnectionMode.Direct:
-						await socket.ConnectAsync(IPAddress.Parse(response.Ip), response.Port, cancellationToken);
+						await socket.ConnectAsync(agentHost, agentPort, cancellationToken);
 						break;
 
 					case ConnectionMode.Tunnel when !String.IsNullOrEmpty(response.ConnectionAddress):
-						(string host, int port) = ParseHostPort(response.ConnectionAddress);
-						await socket.ConnectAsync(host, port, cancellationToken);
+						(agentHost, agentPort) = ParseHostPort(response.ConnectionAddress);
+						await socket.ConnectAsync(agentHost, agentPort, cancellationToken);
 						await TunnelHandshakeAsync(socket, response, cancellationToken);
 						break;
 
 					case ConnectionMode.Relay when !String.IsNullOrEmpty(response.ConnectionAddress):
 						response.Ip = response.ConnectionAddress;
-						await socket.ConnectAsync(IPAddress.Parse(response.ConnectionAddress), response.Ports[ConnectionMetadataPort.ComputeId].Port, cancellationToken);
+						agentHost = response.ConnectionAddress;
+						agentPort = response.Ports[ConnectionMetadataPort.ComputeId].Port;
+						await ConnectWithRetryAsync(socket, agentHost, agentPort, TimeSpan.FromSeconds(5), 3, cancellationToken);
 						break;
 
 					default:
@@ -376,7 +380,7 @@ namespace EpicGames.Horde.Compute.Clients
 			}
 			catch (SocketException se)
 			{
-				throw new ServerComputeClientException($"Unable to connect to {agentAddress}", se);
+				throw new ServerComputeClientException($"Unable to connect to {agentHost}:{agentPort} with mode {response.ConnectionMode}", se);
 			}
 
 			// Send the nonce
@@ -389,6 +393,39 @@ namespace EpicGames.Horde.Compute.Clients
 			await using RemoteComputeSocket computeSocket = new(transport, (ComputeProtocol)response.Protocol, workerLogger);
 			yield return new LeaseInfo(response.ClusterId, response.Properties, response.AssignedResources, computeSocket, response.Ip, response.ConnectionMode, response.Ports);
 		}
+		
+		private async Task ConnectWithRetryAsync(Socket socket, string host, int port, TimeSpan timeout, int maxRetries, CancellationToken cancellationToken)
+		{
+			TimeSpan retryDelay = TimeSpan.FromSeconds(1);
+			for (int attempt = 1; attempt <= maxRetries; attempt++)
+			{
+				try
+				{
+					using CancellationTokenSource timeoutCts = new (timeout);
+					using CancellationTokenSource linkedCts = CancellationTokenSource.CreateLinkedTokenSource(timeoutCts.Token, cancellationToken);
+					await socket.ConnectAsync(host, port, linkedCts.Token);
+					return;
+				}
+				catch (OperationCanceledException)
+				{
+					if (attempt == maxRetries)
+					{
+						throw new TimeoutException($"Failed to connect {host}:{port} after {maxRetries} attempts");
+					}
+				}
+				catch (SocketException se)
+				{
+					if (attempt == maxRetries)
+					{
+						throw;
+					}
+					_logger.LogInformation("Unable to connect to {Host}:{Port} within {Timeout} ms. Error code: {Error}. Waiting {RetryDelay} ms before retrying...",
+						host, port, (int)timeout.TotalMilliseconds, se.SocketErrorCode, (int)retryDelay.TotalMilliseconds);
+				}
+				
+				await Task.Delay(retryDelay, cancellationToken);
+			}
+		}
 
 		private static async Task<ComputeTransport> CreateTransportAsync(Socket socket, AssignComputeResponse response, CancellationToken cancellationToken)
 		{
