Horde: Improve log message for agent connect

diff --git a/Engine/Source/Programs/Shared/EpicGames.Horde/Compute/Clients/ServerComputeClient.cs b/Engine/Source/Programs/Shared/EpicGames.Horde/Compute/Clients/ServerComputeClient.cs
--- a/Engine/Source/Programs/Shared/EpicGames.Horde/Compute/Clients/ServerComputeClient.cs
+++ b/Engine/Source/Programs/Shared/EpicGames.Horde/Compute/Clients/ServerComputeClient.cs
@@ -343,35 +343,52 @@ namespace EpicGames.Horde.Compute.Clients
 					throw new InvalidOperationException();
 				}
 			}
-
-			string agentHost = response.Ip;
-			int agentPort = response.Port;
-
+			
+			string nonce = response.Nonce;
+			nonce = nonce.Length <= 8 ? nonce : $"{nonce[..4]}...{nonce[^4..]}"; // Trim large nonce strings
+			
+			(string host, int port) agentAddress = (response.Ip, response.Port); // Canonical address of agent not accounting for relays
+			(string host, int port) connectionAddress = agentAddress; // De facto address of agent, accounting for relays
+			if (response.ConnectionMode == ConnectionMode.Relay && !String.IsNullOrEmpty(response.ConnectionAddress))
+			{
+				agentAddress.port = response.Ports[ConnectionMetadataPort.ComputeId].AgentPort;
+				connectionAddress = (response.ConnectionAddress, response.Ports[ConnectionMetadataPort.ComputeId].Port);
+			}
+			
 			// Connect to the remote machine
 			using Socket socket = new (SocketType.Stream, ProtocolType.Tcp);
-
-			workerLogger.LogDebug("Connecting to {AgentId} at {AgentHost}{AgentPort} ({ConnectionType} via {ConnectionAddress}) with nonce {Nonce} and encryption {Encryption}. Lease ID {LeaseId}...",
-				response.AgentId, agentHost, agentPort, response.ConnectionMode, response.ConnectionAddress ?? "None", response.Nonce, response.Encryption, response.LeaseId);
+			
+			workerLogger.LogDebug(
+				"Connecting to {AgentId}. Agent={AgentHost}:{AgentPort} Connection={ConnectionMode}/{ConnectionHost}:{ConnectionPort} Encryption={Encryption} LeaseId={LeaseId} Requirements={Requirements} PublicIp={PublicIp} Nonce={Nonce}",
+				response.AgentId,
+				agentAddress.host,
+				agentAddress.port,
+				response.ConnectionMode,
+				connectionAddress.host,
+				connectionAddress.port,
+				response.Encryption,
+				response.LeaseId,
+				request.Requirements,
+				request.Connection?.ClientPublicIp,
+				nonce);
 			
 			try
 			{
 				switch (response.ConnectionMode)
 				{
 					case ConnectionMode.Direct:
-						await socket.ConnectAsync(agentHost, agentPort, cancellationToken);
+						await socket.ConnectAsync(connectionAddress.host, connectionAddress.port, cancellationToken);
 						break;
 
 					case ConnectionMode.Tunnel when !String.IsNullOrEmpty(response.ConnectionAddress):
-						(agentHost, agentPort) = ParseHostPort(response.ConnectionAddress);
-						await socket.ConnectAsync(agentHost, agentPort, cancellationToken);
+						(connectionAddress.host, connectionAddress.port) = ParseHostPort(response.ConnectionAddress);
+						await socket.ConnectAsync(connectionAddress.host, connectionAddress.port, cancellationToken);
 						await TunnelHandshakeAsync(socket, response, cancellationToken);
 						break;
 
 					case ConnectionMode.Relay when !String.IsNullOrEmpty(response.ConnectionAddress):
 						response.Ip = response.ConnectionAddress;
-						agentHost = response.ConnectionAddress;
-						agentPort = response.Ports[ConnectionMetadataPort.ComputeId].Port;
-						await ConnectWithRetryAsync(socket, agentHost, agentPort, TimeSpan.FromSeconds(5), 3, cancellationToken);
+						await ConnectWithRetryAsync(socket, connectionAddress.host, connectionAddress.port, TimeSpan.FromSeconds(5), 3, cancellationToken);
 						break;
 
 					default:
@@ -380,12 +397,12 @@ namespace EpicGames.Horde.Compute.Clients
 			}
 			catch (SocketException se)
 			{
-				throw new ServerComputeClientException($"Unable to connect to {agentHost}:{agentPort} with mode {response.ConnectionMode}", se);
+				throw new ServerComputeClientException($"Unable to connect to {connectionAddress.host}:{connectionAddress.port} with mode {response.ConnectionMode}", se);
 			}
 
 			// Send the nonce
-			byte[] nonce = StringUtils.ParseHexString(response.Nonce);
-			await socket.SendMessageAsync(nonce, SocketFlags.None, cancellationToken);
+			byte[] nonceData = StringUtils.ParseHexString(response.Nonce);
+			await socket.SendMessageAsync(nonceData, SocketFlags.None, cancellationToken);
 			workerLogger.LogInformation("Connected to {AgentId} ({Ip}) under lease {LeaseId} (agent version: {AgentVersion})", response.AgentId, response.Ip, response.LeaseId, response.AgentVersion ?? "unknown");
 			
 			response.Properties = [..response.Properties, $"{KnownPropertyNames.LeaseId}={response.LeaseId}"];
