Fix Horde C# client api GetUgsMetadataAsync streamId conversion to string

#rnx

diff --git a/Engine/Source/Programs/Shared/EpicGames.Horde/HordeHttpClient.cs b/Engine/Source/Programs/Shared/EpicGames.Horde/HordeHttpClient.cs
--- a/Engine/Source/Programs/Shared/EpicGames.Horde/HordeHttpClient.cs
+++ b/Engine/Source/Programs/Shared/EpicGames.Horde/HordeHttpClient.cs
@@ -9,6 +9,7 @@ using System.Net.Http;
 using System.Net.Http.Headers;
 using System.Net.Http.Json;
 using System.Text.Json;
+using System.Text.RegularExpressions;
 using System.Threading;
 using System.Threading.Tasks;
 using EpicGames.Core;
@@ -555,8 +556,17 @@ namespace EpicGames.Horde
 		/// <returns></returns>
 		public Task<GetUgsMetadataListResponse> GetUgsMetadataAsync(StreamId streamId, CommitId commitId, ProjectId projectId, CancellationToken cancellationToken = default)
 		{
-			string perforceStream = $"//{streamId.ToString().Replace('-', '/')}";
-			return GetAsync<GetUgsMetadataListResponse>(_httpClient, $"/ugs/api/metadata?stream={perforceStream}&change={commitId.GetPerforceChange()}&project={projectId}", cancellationToken);
+			Regex streamRe = new Regex(@"(.*?)\-(.*)");
+			Match streamMatch = streamRe.Match(streamId.ToString());
+			if (streamMatch.Success)
+			{
+				string streamName = $"//{streamMatch.Groups[1]}/{streamMatch.Groups[2]}";
+				return GetAsync<GetUgsMetadataListResponse>(_httpClient, $"/ugs/api/metadata?stream={streamName}&change={commitId.GetPerforceChange()}&project={projectId}", cancellationToken);
+			}
+			else
+			{
+				throw new ArgumentException($"StreamId '{streamId.ToString()} is invalid'");
+			}
 		}
 		#endregion
 	}
