EpicGames.ProjectStore: Remove unnecessary references to UnrealBuildTool project

#rnx

diff --git a/Engine/Source/Programs/Shared/EpicGames.ProjectStore/EpicGames.ProjectStore.csproj b/Engine/Source/Programs/Shared/EpicGames.ProjectStore/EpicGames.ProjectStore.csproj
--- a/Engine/Source/Programs/Shared/EpicGames.ProjectStore/EpicGames.ProjectStore.csproj
+++ b/Engine/Source/Programs/Shared/EpicGames.ProjectStore/EpicGames.ProjectStore.csproj
@@ -37,7 +37,7 @@
     </Compile>
   </ItemGroup>
   <ItemGroup>
-    <ProjectReference Include="..\..\UnrealBuildTool\UnrealBuildTool.csproj" />
+    <ProjectReference Include="..\EpicGames.Build\EpicGames.Build.csproj" />
     <ProjectReference Include="..\EpicGames.Core\EpicGames.Core.csproj" />
     <ProjectReference Include="..\EpicGames.Horde\EpicGames.Horde.csproj" />
   </ItemGroup>
diff --git a/Engine/Source/Programs/Shared/EpicGames.ProjectStore/ProjectStoreTypes.cs b/Engine/Source/Programs/Shared/EpicGames.ProjectStore/ProjectStoreTypes.cs
--- a/Engine/Source/Programs/Shared/EpicGames.ProjectStore/ProjectStoreTypes.cs
+++ b/Engine/Source/Programs/Shared/EpicGames.ProjectStore/ProjectStoreTypes.cs
@@ -2,14 +2,12 @@
 
 using System;
 using System.Collections.Generic;
-using System.Diagnostics;
 using System.Globalization;
 using System.IO;
 using System.Linq;
 using System.Net.Http.Headers;
 using System.Net.Http;
 using System.Net;
-using System.Runtime.InteropServices;
 using System.Text.Json.Serialization;
 using System.Text.Json;
 using System.Text;
