EpicGames.Core: No need to compile out IsDryRun except in debug

#rnx

diff --git a/Engine/Source/Programs/Shared/EpicGames.Core/Telemetry/DataRouterTelemetryService.cs b/Engine/Source/Programs/Shared/EpicGames.Core/Telemetry/DataRouterTelemetryService.cs
--- a/Engine/Source/Programs/Shared/EpicGames.Core/Telemetry/DataRouterTelemetryService.cs
+++ b/Engine/Source/Programs/Shared/EpicGames.Core/Telemetry/DataRouterTelemetryService.cs
@@ -291,7 +291,6 @@ namespace EpicGames.Core.Telemetry
 	/// </summary>
 	public class DataRouterTelemetryService : ITelemetryService<DataRouterEvent>
 	{
-#if DEBUG
 		/// <summary>
 		/// Denotes if this is a dry run
 		/// </summary>
@@ -300,7 +299,6 @@ namespace EpicGames.Core.Telemetry
 			get;
 			set;
 		} = false;
-#endif
 
 		/// <summary>
 		/// Logging interface
@@ -535,7 +533,6 @@ namespace EpicGames.Core.Telemetry
 
 			EventsQueue.Clear();
 
-#if DEBUG
 			if (IsDryRun)
 			{
 				foreach (KeyValuePair<string, DataRouterPost> postData in posts)
@@ -545,7 +542,7 @@ namespace EpicGames.Core.Telemetry
 
 				return;
 			}
-#endif
+
 			JsonSerializerOptions serializeOptions = new JsonSerializerOptions
 			{
 				WriteIndented = false,
diff --git a/Engine/Source/Programs/UnrealBuildTool/System/Telemetry.cs b/Engine/Source/Programs/UnrealBuildTool/System/Telemetry.cs
--- a/Engine/Source/Programs/UnrealBuildTool/System/Telemetry.cs
+++ b/Engine/Source/Programs/UnrealBuildTool/System/Telemetry.cs
@@ -203,9 +203,7 @@ namespace UnrealBuildTool
 			_service = new(() =>
 			{
 				DataRouterTelemetryService service = new(_httpClient.Value, _baseAddress);
-#if DEBUG
-				service.IsDryRun = false;
-#endif
+				//service.IsDryRun = false;
 				return service;
 			});
 		}
