EpicGames.MsBuild: Properly escape semicolon in property argument

#rnx

diff --git a/Engine/Source/Programs/Shared/EpicGames.MsBuild/CsProjBuilder.cs b/Engine/Source/Programs/Shared/EpicGames.MsBuild/CsProjBuilder.cs
--- a/Engine/Source/Programs/Shared/EpicGames.MsBuild/CsProjBuilder.cs
+++ b/Engine/Source/Programs/Shared/EpicGames.MsBuild/CsProjBuilder.cs
@@ -105,7 +105,7 @@ namespace EpicGames.MsBuild
 
 			if (defineConstants.Any())
 			{
-				globalProperties.Add("DefineConstants", String.Join(';', defineConstants));
+				globalProperties.Add("DefineConstants", String.Join("%3B", defineConstants));
 			}
 
 			ConcurrentDictionary<FileReference, CsProjBuildRecordEntry> buildRecords = new(hook.ValidBuildRecords);
