EpicGames.MsBuild: Handle platform specific TargetFramework projects that need to be build in a separate call to dotnet build

#rnx

diff --git a/Engine/Source/Programs/Shared/EpicGames.MsBuild/CsProjBuilder.cs b/Engine/Source/Programs/Shared/EpicGames.MsBuild/CsProjBuilder.cs
--- a/Engine/Source/Programs/Shared/EpicGames.MsBuild/CsProjBuilder.cs
+++ b/Engine/Source/Programs/Shared/EpicGames.MsBuild/CsProjBuilder.cs
@@ -113,6 +113,7 @@ namespace EpicGames.MsBuild
 
 			using ProjectCollection projectCollection = new ProjectCollection(globalProperties);
 			ConcurrentDictionary<FileReference, Project> projects = new();
+			ConcurrentBag<FileReference> platformSpecificProjects = new();
 
 			stopwatch.Restart();
 			// Load all found projects
@@ -148,9 +149,22 @@ namespace EpicGames.MsBuild
 						}
 					}
 
-					if (project.GetProperty("TargetFramework").EvaluatedValue.Contains("net6.0", StringComparison.Ordinal))
+					string targetFramework = project.GetProperty("TargetFramework").EvaluatedValue;
+					bool skipFramework = false;
+					// check the TargetFramework of the project: we can't build Windows-only projects on non-Windows platforms etc.
+					if (targetFramework.Contains('-', StringComparison.Ordinal))
 					{
-						logger.LogInformation("Skipping net6.0 project {ProjectPath}", projectPath);
+						skipFramework |= (OperatingSystem.IsWindows() && !targetFramework.Contains("windows", StringComparison.Ordinal))
+							|| (OperatingSystem.IsMacOS() && !targetFramework.Contains("macos", StringComparison.Ordinal))
+							|| (OperatingSystem.IsLinux());
+						platformSpecificProjects.Add(foundProject);
+					}
+
+					skipFramework |= targetFramework.Contains("net6.0", StringComparison.Ordinal);
+
+					if (skipFramework)
+					{
+						logger.LogInformation("Skipping {Framework} project {ProjectPath}", targetFramework, projectPath);
 						return;
 					}
 
@@ -285,7 +299,14 @@ namespace EpicGames.MsBuild
 			{
 				onBuildingProjects(outOfDateProjects.Count());
 
-				bBuildSuccess = RunDotnetBuildAsync(outOfDateProjects, globalProperties, hook, logger, default).Result;
+				IEnumerable<FileReference> singleBuildProjects = outOfDateProjects.Intersect(platformSpecificProjects);
+				bBuildSuccess = RunDotnetBuildAsync(outOfDateProjects.Except(singleBuildProjects), globalProperties, hook, logger, default).Result;
+
+				// Projects that have a platform specific TargetFramework can't be built using EpicGames.ScriptBuild due to the mismatch
+				foreach (FileReference project in singleBuildProjects)
+				{
+					bBuildSuccess = RunDotnetBuildAsync([project], globalProperties, hook, logger, default).Result && bBuildSuccess;
+				}
 			}
 			else
 			{
@@ -462,11 +483,20 @@ namespace EpicGames.MsBuild
 			Stopwatch stopwatch = new();
 			stopwatch.Start();
 
-			FileReference propertiesPath = await WriteCsPropertiesAsync(projects, hook, cancellationToken);
+			FileReference buildProject;
+			if (projects.Count() > 1)
+			{
+				buildProject = FileReference.Combine(hook.EngineDirectory, "Source", "Programs", "Shared", "EpicGames.ScriptBuild");
+				await WriteCsPropertiesAsync(projects, hook, cancellationToken);
+			}
+			else
+			{
+				buildProject = projects.First();
+			}
 
 			IEnumerable<string> arguments = [
 				"build",
-				$"\"{FileReference.Combine(hook.EngineDirectory, "Source", "Programs", "Shared", "EpicGames.ScriptBuild")}\"",
+				$"\"{buildProject}\"",
 				"-nologo",
 				"-v:quiet",
 				.. globalProperties.Select(x => $"\"/p:{x.Key}={x.Value}\""),
