Slack: Catch exception and log when requesting user by email

diff --git a/Engine/Source/Programs/Shared/EpicGames.Slack/SlackClient.cs b/Engine/Source/Programs/Shared/EpicGames.Slack/SlackClient.cs
--- a/Engine/Source/Programs/Shared/EpicGames.Slack/SlackClient.cs
+++ b/Engine/Source/Programs/Shared/EpicGames.Slack/SlackClient.cs
@@ -492,13 +492,21 @@ namespace EpicGames.Slack
 		/// <returns>User profile</returns>
 		public async Task<SlackUser?> FindUserByEmailAsync(string email, CancellationToken cancellationToken = default)
 		{
-			using (HttpRequestMessage request = new HttpRequestMessage(HttpMethod.Post, $"{UsersLookupByEmailUrl}?email={email}"))
+			try
 			{
-				static bool ShouldLogError(UserResponse response) => !response.Ok && !String.Equals(response.Error, "users_not_found", StringComparison.Ordinal);
+				using (HttpRequestMessage request = new HttpRequestMessage(HttpMethod.Post, $"{UsersLookupByEmailUrl}?email={email}"))
+				{
+					static bool ShouldLogError(UserResponse response) => !response.Ok && !String.Equals(response.Error, "users_not_found", StringComparison.Ordinal);
 
-				UserResponse response = await SendRequestAsync<UserResponse>(request, "", ShouldLogError, cancellationToken);
-				return response.User;
+					UserResponse response = await SendRequestAsync<UserResponse>(request, "", ShouldLogError, cancellationToken);
+					return response.User;
+				}
 			}
+			catch (Exception ex)
+			{
+				_logger.LogWarning(ex, "Failed to get Slack user email {Url} ({Error})", $"{UsersLookupByEmailUrl}?email={email}", ex.Message);
+				return null;
+			} 
 		}
 
 		#endregion
