Fixed issues with object id generation from C#
Time part was mixing up units.
Serial part was not initializing from random value.

diff --git a/Engine/Source/Programs/Shared/EpicGames.Serialization/CbField.cs b/Engine/Source/Programs/Shared/EpicGames.Serialization/CbField.cs
--- a/Engine/Source/Programs/Shared/EpicGames.Serialization/CbField.cs
+++ b/Engine/Source/Programs/Shared/EpicGames.Serialization/CbField.cs
@@ -2811,7 +2811,8 @@ namespace EpicGames.Serialization
 		readonly uint _b = b;
 		readonly uint _c = c;
 
-		private static uint s_objectIdSerial = 0u;
+		// initialize the serial to a random integer
+		private static uint s_objectIdSerial = CalculateRunId();
 
 		/// <summary>
 		/// The zero value of an object id
@@ -2866,7 +2867,8 @@ namespace EpicGames.Serialization
 		public static CbObjectId NewObjectId()
 		{
 			const long Offset = 1_609_459_200; // Seconds from 1970 -> 2021
-			uint time = (uint)(DateTimeOffset.Now.ToUnixTimeMilliseconds() - Offset);
+			long longTime = DateTimeOffset.Now.ToUnixTimeSeconds() - Offset;
+			uint time = (uint)longTime;
 			uint serial = Interlocked.Increment(ref s_objectIdSerial);
 			uint runId = RunId;
 			return new CbObjectId(time, serial, runId);
