[UBA]
* Changed pinvoke declaration to return byte instead of bool since that seems to work better in other cases when importing functions that returns normal bool

diff --git a/Engine/Source/Programs/Shared/EpicGames.UBA/Impl/SessionServerImpl.cs b/Engine/Source/Programs/Shared/EpicGames.UBA/Impl/SessionServerImpl.cs
--- a/Engine/Source/Programs/Shared/EpicGames.UBA/Impl/SessionServerImpl.cs
+++ b/Engine/Source/Programs/Shared/EpicGames.UBA/Impl/SessionServerImpl.cs
@@ -89,7 +89,7 @@ namespace EpicGames.UBA.Impl
 		static extern void SessionServer_RegisterNewFile(nint server, string filename);
 
 		[DllImport("UbaHost", CharSet = CharSet.Auto)]
-		static extern bool SessionServer_RegisterVirtualFile(nint server, string filename, string sourceFile, ulong sourceOffset, ulong sourceSize);
+		static extern byte SessionServer_RegisterVirtualFile(nint server, string filename, string sourceFile, ulong sourceOffset, ulong sourceSize);
 
 		[DllImport("UbaHost", CharSet = CharSet.Auto)]
 		static extern uint SessionServer_BeginExternalProcess(nint server, string description);
@@ -221,7 +221,7 @@ namespace EpicGames.UBA.Impl
 		public void RefreshDirectories(params string[] directories) => Array.ForEach(directories, (directory) => SessionServer_RefreshDirectory(_handle, directory));
 
 		public void RegisterNewFiles(params string[] files) => Array.ForEach(files, (file) => SessionServer_RegisterNewFile(_handle, file));
-		public bool RegisterVirtualFile(string name, string sourceFile, ulong sourceOffset, ulong sourceSize) => SessionServer_RegisterVirtualFile(_handle, name, sourceFile, sourceOffset, sourceSize);
+		public bool RegisterVirtualFile(string name, string sourceFile, ulong sourceOffset, ulong sourceSize) => SessionServer_RegisterVirtualFile(_handle, name, sourceFile, sourceOffset, sourceSize) != 0;
 
 		public uint BeginExternalProcess(string description) => SessionServer_BeginExternalProcess(_handle, description);
 
