Minor updates to the UHT generated output

#rnx

diff --git a/Engine/Source/Programs/Shared/EpicGames.UHT/Exporters/CodeGen/UhtCodeBlockComment.cs b/Engine/Source/Programs/Shared/EpicGames.UHT/Exporters/CodeGen/UhtCodeBlockComment.cs
new file mode 100644
--- /dev/null
+++ b/Engine/Source/Programs/Shared/EpicGames.UHT/Exporters/CodeGen/UhtCodeBlockComment.cs
@@ -0,0 +1,80 @@
+// Copyright Epic Games, Inc. All Rights Reserved.
+
+using System;
+using System.Text;
+using EpicGames.UHT.Types;
+
+namespace EpicGames.UHT.Exporters.CodeGen
+{
+	internal readonly struct UhtCodeBlockComment : IDisposable
+	{
+		private const string HighlightPrefix = "**********";
+		private const char HighlightChar = '*';
+		private const int HighlightLength = 100;
+		private readonly StringBuilder _builder;
+		private readonly UhtType? _primaryType = null;
+		private readonly UhtType? _secondaryType = null;
+		private readonly string? _text = null;
+
+		public UhtCodeBlockComment(StringBuilder builder, string text)
+		{
+			_builder = builder;
+			_text = text;
+			Begin();
+		}
+
+		public UhtCodeBlockComment(StringBuilder builder, UhtType primaryType, string? text = null)
+		{
+			_builder = builder;
+			_primaryType = primaryType;
+			_text = text;
+			Begin();
+		}
+
+		public UhtCodeBlockComment(StringBuilder builder, UhtType primaryType, UhtType secondaryType, string? text = null)
+		{
+			_builder = builder;
+			_primaryType = primaryType;
+			_secondaryType = secondaryType;
+			_text = text;
+			Begin();
+		}
+
+		public void Dispose()
+		{
+			AppendText("End");
+		}
+
+		private void Begin()
+		{
+			_builder.Append("\r\n");
+			AppendText("Begin");
+		}
+
+		private void AppendText(string startOrEnd)
+		{
+			int startLength = _builder.Length;
+			_builder.Append("// ").Append(HighlightPrefix).Append(' ').Append(startOrEnd);
+			if (_primaryType != null)
+			{
+				_builder.Append(' ').Append(_primaryType.EngineType).Append(' ').Append(_primaryType.SourceName);
+			}
+			if (_secondaryType != null)
+			{
+				_builder.Append(' ').Append(_secondaryType.EngineType).Append(' ').Append(_secondaryType.SourceName);
+			}
+			if (!String.IsNullOrEmpty(_text))
+			{
+				_builder.Append(' ').Append(_text);
+			}
+			_builder.Append(' ');
+			int endLength = _builder.Length;
+			int lineLength = endLength - startLength;
+			if (lineLength < HighlightLength)
+			{
+				_builder.Append(HighlightChar, HighlightLength - lineLength);
+			}
+			_builder.Append("\r\n");
+		}
+	}
+}
diff --git a/Engine/Source/Programs/Shared/EpicGames.UHT/Exporters/CodeGen/UhtHeaderCodeGeneratorCppFile.cs b/Engine/Source/Programs/Shared/EpicGames.UHT/Exporters/CodeGen/UhtHeaderCodeGeneratorCppFile.cs
--- a/Engine/Source/Programs/Shared/EpicGames.UHT/Exporters/CodeGen/UhtHeaderCodeGeneratorCppFile.cs
+++ b/Engine/Source/Programs/Shared/EpicGames.UHT/Exporters/CodeGen/UhtHeaderCodeGeneratorCppFile.cs
@@ -10,56 +10,6 @@ using EpicGames.UHT.Utils;
 
 namespace EpicGames.UHT.Exporters.CodeGen
 {
-	internal readonly struct UhtCodeBlockComment : IDisposable
-	{
-		private readonly StringBuilder _builder;
-		private readonly UhtType? _primaryType = null;
-		private readonly UhtType? _secondaryType = null;
-		private readonly string? _text = null;
-
-		public UhtCodeBlockComment(StringBuilder builder, string text)
-		{
-			_builder = builder;
-			_text = text;
-			builder.Append("\r\n// Begin");
-			AppendText();
-			_builder.Append("\r\n");
-		}
-
-		public UhtCodeBlockComment(StringBuilder builder, UhtType primaryType, UhtType? secondaryType = null)
-		{
-			_builder = builder;
-			_primaryType = primaryType;
-			_secondaryType = secondaryType;
-			builder.Append("\r\n// Begin");
-			AppendText();
-			_builder.Append("\r\n");
-		}
-
-		public void Dispose()
-		{
-			_builder.Append("// End");
-			AppendText();
-			_builder.Append("\r\n");
-		}
-
-		private void AppendText()
-		{
-			if (!String.IsNullOrEmpty(_text))
-			{
-				_builder.Append(' ').Append(_text);
-			}
-			if (_primaryType != null)
-			{
-				_builder.Append(' ').Append(_primaryType.EngineType).Append(' ').Append(_primaryType.SourceName);
-			}
-			if (_secondaryType != null)
-			{
-				_builder.Append(' ').Append(_secondaryType.EngineType).Append(' ').Append(_secondaryType.SourceName);
-			}
-		}
-	}
-
 	/// <summary>
 	/// Collection of all registrations for a specific package
 	/// </summary>
@@ -210,7 +160,10 @@ namespace EpicGames.UHT.Exporters.CodeGen
 					builder.Append("#include \"").Append(include).Append("\"\r\n");
 				}
 
+				builder.Append("\r\n");
 				builder.Append(DisableDeprecationWarnings).Append("\r\n");
+
+				builder.Append("\r\n");
 				string cleanFileName = HeaderFile.FileNameWithoutExtension.Replace('.', '_');
 				builder.Append("void EmptyLinkFunctionForGeneratedCode").Append(cleanFileName).Append("() {}\r\n");
 
@@ -374,6 +327,7 @@ namespace EpicGames.UHT.Exporters.CodeGen
 
 				int generatedBodyEnd = builder.Length;
 
+				builder.Append("\r\n");
 				builder.Append(EnableDeprecationWarnings).Append("\r\n");
 
 				{
@@ -926,7 +880,7 @@ namespace EpicGames.UHT.Exporters.CodeGen
 			// Everything from this point on will be part of the definition hash
 			int hashCodeBlockStart = builder.Length;
 
-			// Note: This is fairly agressive 
+			// Note: This is fairly aggressive 
 			using (UhtMacroBlockEmitter blockEmitter = new(builder, UhtDefineScopeNames.WithEditor, function.DefineScope))
 			{
 
diff --git a/Engine/Source/Programs/Shared/EpicGames.UHT/Exporters/CodeGen/UhtHeaderCodeGeneratorHFile.cs b/Engine/Source/Programs/Shared/EpicGames.UHT/Exporters/CodeGen/UhtHeaderCodeGeneratorHFile.cs
--- a/Engine/Source/Programs/Shared/EpicGames.UHT/Exporters/CodeGen/UhtHeaderCodeGeneratorHFile.cs
+++ b/Engine/Source/Programs/Shared/EpicGames.UHT/Exporters/CodeGen/UhtHeaderCodeGeneratorHFile.cs
@@ -42,20 +42,28 @@ namespace EpicGames.UHT.Exporters.CodeGen
 				builder.Append(HeaderCopyright);
 				builder.Append("// IWYU pragma: private, include \"").Append(HeaderFile.IncludeFilePath).Append("\"\r\n");
 
+				string strippedName = Path.GetFileNameWithoutExtension(HeaderFile.FilePath);
+				string defineName = $"{Module.ShortName.ToUpper()}_{strippedName}_generated_h".Replace('.', '_');
+
+				builder.Append("\r\n");
+				builder.Append("#ifdef ").Append(defineName).Append("\r\n");
+				builder.Append("#error \"").Append(strippedName).Append(".generated.h already included, missing '#pragma once' in ").Append(strippedName).Append(".h\"\r\n");
+				builder.Append("#endif\r\n");
+				builder.Append("#define ").Append(defineName).Append("\r\n");
+
 				// Attempt to limit the headers included. This is needed in the lower level engine code
 				// to get around circular header include issues.
+				builder.Append("\r\n");
+				if (HeaderFile.References.ExportTypes.Count > 0 && HeaderFile.References.ExportTypes.Find(x => x is not UhtEnum) == null)
 				{
-					if (HeaderFile.References.ExportTypes.Count > 0 && HeaderFile.References.ExportTypes.Find(x => x is not UhtEnum) == null)
-					{
-						builder.Append("#include \"Templates/IsUEnumClass.h\"\r\n");
-						builder.Append("#include \"UObject/ObjectMacros.h\"\r\n");
-						builder.Append("#include \"UObject/ReflectedTypeAccessors.h\"\r\n");
-					}
-					else
-					{
-						builder.Append("#include \"UObject/ObjectMacros.h\"\r\n");
-						builder.Append("#include \"UObject/ScriptMacros.h\"\r\n");
-					}
+					builder.Append("#include \"Templates/IsUEnumClass.h\"\r\n");
+					builder.Append("#include \"UObject/ObjectMacros.h\"\r\n");
+					builder.Append("#include \"UObject/ReflectedTypeAccessors.h\"\r\n");
+				}
+				else
+				{
+					builder.Append("#include \"UObject/ObjectMacros.h\"\r\n");
+					builder.Append("#include \"UObject/ScriptMacros.h\"\r\n");
 				}
 				
 				if (headerInfo.NeedsPushModelHeaders)
@@ -82,12 +90,10 @@ namespace EpicGames.UHT.Exporters.CodeGen
 				{
 					builder.Append("#include \"VersePointer.h\"\r\n");
 				}
+
 				builder.Append("\r\n");
 				builder.Append(DisableDeprecationWarnings).Append("\r\n");
 
-				string strippedName = Path.GetFileNameWithoutExtension(HeaderFile.FilePath);
-				string defineName = $"{Module.ShortName.ToUpper()}_{strippedName}_generated_h".Replace('.', '_');
-
 				if (HeaderFile.References.ForwardDeclarations.Count > 0)
 				{
 					string[] sorted = new string[HeaderFile.References.ForwardDeclarations.Count];
@@ -97,18 +103,13 @@ namespace EpicGames.UHT.Exporters.CodeGen
 						sorted[index++] = forwardDeclaration;
 					}
 					Array.Sort(sorted, StringComparerUE.OrdinalIgnoreCase);
+					builder.Append("\r\n");
 					foreach (string forwardDeclaration in sorted)
 					{
 						builder.Append(forwardDeclaration).Append("\r\n");
 					}
 				}
 
-				builder.Append("#ifdef ").Append(defineName).Append("\r\n");
-				builder.Append("#error \"").Append(strippedName).Append(".generated.h already included, missing '#pragma once' in ").Append(strippedName).Append(".h\"\r\n");
-				builder.Append("#endif\r\n");
-				builder.Append("#define ").Append(defineName).Append("\r\n");
-				builder.Append("\r\n");
-
 				foreach (UhtField field in HeaderFile.References.ExportTypes)
 				{
 					if (field is UhtEnum enumObj)
@@ -118,33 +119,39 @@ namespace EpicGames.UHT.Exporters.CodeGen
 					{
 						if (scriptStruct.ScriptStructFlags.HasAnyFlags(EStructFlags.Native))
 						{
+							using UhtCodeBlockComment blockComment = new(builder, field);
 							AppendScriptStruct(builder, scriptStruct);
 						}
 					}
 					else if (field is UhtFunction function)
 					{
+						using UhtCodeBlockComment blockComment = new(builder, field);
 						AppendDelegate(builder, function);
 					}
 					else if (field is UhtClass classObj)
 					{
 						if (!classObj.ClassFlags.HasAnyFlags(EClassFlags.Intrinsic))
 						{
+							using UhtCodeBlockComment blockComment = new(builder, field);
 							AppendClass(builder, classObj);
 						}
 					}
 				}
 
+				builder.Append("\r\n");
 				builder.Append("#undef CURRENT_FILE_ID\r\n");
-				builder.Append("#define CURRENT_FILE_ID ").Append(headerInfo.FileId).Append("\r\n\r\n\r\n");
+				builder.Append("#define CURRENT_FILE_ID ").Append(headerInfo.FileId).Append("\r\n");
 
 				foreach (UhtField field in HeaderFile.References.ExportTypes)
 				{
 					if (field is UhtEnum enumObject)
 					{
+						using UhtCodeBlockComment blockComment = new(builder, field);
 						AppendEnum(builder, enumObject);
 					}
 				}
 
+				builder.Append("\r\n");
 				builder.Append(EnableDeprecationWarnings).Append("\r\n");
 
 				if (SaveExportedHeaders)
@@ -619,19 +626,18 @@ namespace EpicGames.UHT.Exporters.CodeGen
 
 				// For back compat we still need to forward declare the struct here
 				builder.Append("struct ").Append(scriptStruct.SourceName).Append(";\r\n");
-				
-				Session.Tables!.CodeGeneratorInjectorTable.Inject(builder, scriptStruct, UhtCodeGeneratorInjectionLocation.Header);
 
 				// Verse declarations
 				if (scriptStruct.IsVerseField)
 				{
 					if (!scriptStruct.FieldExportFlags.HasAnyFlags(UhtFieldExportFlags.NoVerseAlias))
 					{
-						builder.Append("namespace verse { using ").Append(scriptStruct.EngineName).Append(" = ::").Append(scriptStruct.SourceName).Append("; }\r\n");
+						builder.Append("namespace verse { using ").Append(scriptStruct.EngineName).Append(" = ").Append(scriptStruct.SourceName).Append("; }\r\n");
 					}
 					builder.Append("V_DECLARE_STRUCT_RUNTIME_TRAITS(, ").Append(scriptStruct.SourceName).Append(", \"").AppendVerseUEPackageName(scriptStruct).Append("\", \"").AppendVerseName(scriptStruct).Append("\")\r\n");
-					builder.Append("\r\n");
 				}
+
+				Session.Tables!.CodeGeneratorInjectorTable.Inject(builder, scriptStruct, UhtCodeGeneratorInjectionLocation.Header);
 			}
 			return builder;
 		}
@@ -684,14 +690,16 @@ namespace EpicGames.UHT.Exporters.CodeGen
 			{	
 				// Forward declare the StaticEnum<> specialization for enum classes
 				builder.Append("template<> ").Append(Module.Api).Append("UEnum* StaticEnum<").Append(enumObj.CppType).Append(">();\r\n");
-				builder.Append("\r\n");
 			}
 
 			// Verse declarations
 			if (enumObj.IsVerseField)
 			{
+				if (!enumObj.FieldExportFlags.HasAnyFlags(UhtFieldExportFlags.NoVerseAlias))
+				{
+					builder.Append("namespace verse { using ").Append(enumObj.EngineName).Append(" = ::").Append(enumObj.SourceName).Append("; }\r\n");
+				}
 				builder.Append("V_DECLARE_IMPORTED_ENUM(").Append(Module.Api.TrimEnd()).Append(", ").Append(enumObj.SourceName).Append(");\r\n");
-				builder.Append("\r\n");
 			}
 			return builder;
 		}
@@ -860,9 +868,6 @@ namespace EpicGames.UHT.Exporters.CodeGen
 			// For back compat we still need to forward declare the class here
 			builder.Append("class ").Append(classObj.SourceName).Append(";\r\n");
 			
-			Session.Tables!.CodeGeneratorInjectorTable.Inject(builder, classObj, UhtCodeGeneratorInjectionLocation.Header);
-			builder.Append("\r\n");
-
 			// Verse declarations
 			if (classObj.IsVerseField)
 			{
@@ -870,8 +875,10 @@ namespace EpicGames.UHT.Exporters.CodeGen
 				{
 					builder.Append("namespace verse { using ").Append(classObj.EngineName).Append(" = ::").Append(classObj.SourceName).Append("; }\r\n");
 				}
-				builder.Append("\r\n");
 			}
+
+			Session.Tables!.CodeGeneratorInjectorTable.Inject(builder, classObj, UhtCodeGeneratorInjectionLocation.Header);
+			builder.Append("\r\n");
 			return builder;
 		}
 
diff --git a/Engine/Source/Programs/Shared/EpicGames.UHT/Utils/UhtStringBuilderExtensions.cs b/Engine/Source/Programs/Shared/EpicGames.UHT/Utils/UhtStringBuilderExtensions.cs
--- a/Engine/Source/Programs/Shared/EpicGames.UHT/Utils/UhtStringBuilderExtensions.cs
+++ b/Engine/Source/Programs/Shared/EpicGames.UHT/Utils/UhtStringBuilderExtensions.cs
@@ -327,7 +327,6 @@ namespace EpicGames.UHT.Utils
 							if (c < 31)
 							{
 								trailingHex = true;
-								//Builder.Append($"\\x{(uint)c:x2}");
 								builder.Append("\\x").AppendFormat(CultureInfo.InvariantCulture, "{0:x2}", (uint)c);
 							}
 							else
