Added new IUhtExportFactory methods, GetModuleShortestIncludePath and GetPluginShortestIncludePath.

The IncludeBase member in UHTModule has been replaced with IncludePaths.  Use the methods to return the best shortest relative path that can include the given file path.

diff --git a/Engine/Plugins/ScriptPlugin/Source/ScriptGeneratorUbtPlugin/GenericScriptCodeGenerator.cs b/Engine/Plugins/ScriptPlugin/Source/ScriptGeneratorUbtPlugin/GenericScriptCodeGenerator.cs
--- a/Engine/Plugins/ScriptPlugin/Source/ScriptGeneratorUbtPlugin/GenericScriptCodeGenerator.cs
+++ b/Engine/Plugins/ScriptPlugin/Source/ScriptGeneratorUbtPlugin/GenericScriptCodeGenerator.cs
@@ -184,7 +184,7 @@ namespace ScriptGeneratorUbtPlugin
 			sortedHeaders.Sort(StringComparerUE.OrdinalIgnoreCase);
 			foreach (string filePath in sortedHeaders)
 			{
-				string relativePath = Path.GetRelativePath(Factory.PluginModule!.IncludeBase, filePath).Replace("\\", "/");
+				string relativePath = Factory.GetPluginShortestIncludePath(filePath);
 				builder.Append("#include \"").Append(relativePath).Append("\"\r\n");
 			}
 
diff --git a/Engine/Plugins/ScriptPlugin/Source/ScriptGeneratorUbtPlugin/LuaScriptCodeGenerator.cs b/Engine/Plugins/ScriptPlugin/Source/ScriptGeneratorUbtPlugin/LuaScriptCodeGenerator.cs
--- a/Engine/Plugins/ScriptPlugin/Source/ScriptGeneratorUbtPlugin/LuaScriptCodeGenerator.cs
+++ b/Engine/Plugins/ScriptPlugin/Source/ScriptGeneratorUbtPlugin/LuaScriptCodeGenerator.cs
@@ -382,7 +382,7 @@ namespace ScriptGeneratorUbtPlugin
 			sortedHeaders.Sort(StringComparerUE.OrdinalIgnoreCase);
 			foreach (string filePath in sortedHeaders)
 			{
-				string relativePath = Path.GetRelativePath(Factory.PluginModule!.IncludeBase, filePath).Replace("\\", "/");
+				string relativePath = Factory.GetPluginShortestIncludePath(filePath);
 				builder.Append("#include \"").Append(relativePath).Append("\"\r\n");
 			}
 
diff --git a/Engine/Source/Programs/Shared/EpicGames.UHT/Exporters/CodeGen/UhtCodeGenerator.cs b/Engine/Source/Programs/Shared/EpicGames.UHT/Exporters/CodeGen/UhtCodeGenerator.cs
--- a/Engine/Source/Programs/Shared/EpicGames.UHT/Exporters/CodeGen/UhtCodeGenerator.cs
+++ b/Engine/Source/Programs/Shared/EpicGames.UHT/Exporters/CodeGen/UhtCodeGenerator.cs
@@ -245,31 +245,7 @@ namespace EpicGames.UHT.Exporters.CodeGen
 			ref HeaderInfo headerInfo = ref HeaderInfos[headerFile.HeaderFileTypeIndex];
 
 			// Find the shortest matching relative path
-			string? relativePath = null;
-			foreach (string includePath in headerFile.Module.Module.IncludePaths)
-			{
-				if (headerFile.FilePath.StartsWith(includePath, StringComparison.Ordinal))
-				{
-					if (relativePath == null || (includePath.Length - headerFile.FilePath.Length - 1) < relativePath.Length)
-					{
-						relativePath = headerFile.FilePath.Substring(includePath.Length + 1);
-					}
-				}
-			}
-
-			// This will create a "../" path which is not great
-			if (relativePath == null)
-			{
-				string? includePath = headerFile.Module.Module.IncludePaths.Find(includePath => includePath.StartsWith(headerFile.Module.Module.BaseDirectory, StringComparison.Ordinal));
-				if (includePath == null)
-				{
-					Session.Logger.LogWarning("There are no include paths in module {ModuleName} for {FileName}. It will be included by a \"../\" path external to the module.", headerFile.Module.ShortName, headerFile.FilePath);
-					includePath = headerFile.Module.Module.IncludePaths.FirstOrDefault(".");
-				}
-				relativePath = Path.GetRelativePath(includePath!, headerFile.FilePath);
-			}
-
-			headerInfo.IncludePath = relativePath.Replace('\\', '/');
+			headerInfo.IncludePath = Factory.GetModuleShortestIncludePath(headerFile.Module, headerFile.FilePath);
 
 			// Convert the file path to a C identifier
 			string filePath = headerFile.FilePath;
diff --git a/Engine/Source/Programs/Shared/EpicGames.UHT/Utils/UhtExport.cs b/Engine/Source/Programs/Shared/EpicGames.UHT/Utils/UhtExport.cs
--- a/Engine/Source/Programs/Shared/EpicGames.UHT/Utils/UhtExport.cs
+++ b/Engine/Source/Programs/Shared/EpicGames.UHT/Utils/UhtExport.cs
@@ -91,5 +91,20 @@ namespace EpicGames.UHT.Utils
 		/// </summary>
 		/// <param name="filePath">External dependency to add</param>
 		public void AddExternalDependency(string filePath);
+
+		/// <summary>
+		/// Return the shortest relative include path
+		/// </summary>
+		/// <param name="moduleObj">Module in question</param>
+		/// <param name="filePath">File path</param>
+		/// <returns>Shortest include path relative to the given file path</returns>
+		public string GetModuleShortestIncludePath(UhtModule moduleObj, string filePath);
+
+		/// <summary>
+		/// Return the shortest relative include path
+		/// </summary>
+		/// <param name="filePath">File path</param>
+		/// <returns>Shortest include path relative to the given file path</returns>
+		public string GetPluginShortestIncludePath(string filePath);
 	}
 }
diff --git a/Engine/Source/Programs/Shared/EpicGames.UHT/Utils/UhtSession.cs b/Engine/Source/Programs/Shared/EpicGames.UHT/Utils/UhtSession.cs
--- a/Engine/Source/Programs/Shared/EpicGames.UHT/Utils/UhtSession.cs
+++ b/Engine/Source/Programs/Shared/EpicGames.UHT/Utils/UhtSession.cs
@@ -242,6 +242,50 @@ namespace EpicGames.UHT.Utils
 			ExternalDependencies.Add(filePath);
 		}
 
+		/// <inheritdoc/>
+		public string GetModuleShortestIncludePath(UhtModule moduleObj, string filePath)
+		{
+			return GetShortestIncludePath(moduleObj.Module, filePath);
+		}
+
+		/// <inheritdoc/>
+		public string GetPluginShortestIncludePath(string filePath)
+		{
+			if (PluginModule == null)
+			{
+				throw new UhtIceException("Made a request for the plugin's shortest include path but there is no plugin");
+			}
+			return GetShortestIncludePath(PluginModule, filePath);
+		}
+
+		public string GetShortestIncludePath(UHTManifest.Module moduleObj, string filePath)
+		{
+			string? relativePath = null;
+			foreach (string includePath in moduleObj.IncludePaths)
+			{
+				if (filePath.StartsWith(includePath, StringComparison.Ordinal))
+				{
+					if (relativePath == null || (includePath.Length - filePath.Length - 1) < relativePath.Length)
+					{
+						relativePath = filePath.Substring(includePath.Length + 1);
+					}
+				}
+			}
+
+			// This will create a "../" path which is not great
+			if (relativePath == null)
+			{
+				string? includePath = moduleObj.IncludePaths.Find(includePath => includePath.StartsWith(moduleObj.BaseDirectory, StringComparison.Ordinal));
+				if (includePath == null)
+				{
+					Session.Logger.LogWarning("There are no include paths in module {ModuleName} for {FileName}. It will be included by a \"../\" path external to the module.", moduleObj.Name, filePath);
+					includePath = moduleObj.IncludePaths.FirstOrDefault(".");
+				}
+				relativePath = Path.GetRelativePath(includePath!, filePath);
+			}
+			return relativePath!.Replace("\\", "/");
+		}
+
 		private string MakePath(UHTManifest.Module module, string fileName, string suffix)
 		{
 			if (PluginModule != null)
