Improving UHT support for generating interface source names from the source name and not the engine name.

#rnx

diff --git a/Engine/Source/Programs/Shared/EpicGames.UHT/Exporters/CodeGen/UhtHeaderCodeGeneratorCppFile.cs b/Engine/Source/Programs/Shared/EpicGames.UHT/Exporters/CodeGen/UhtHeaderCodeGeneratorCppFile.cs
--- a/Engine/Source/Programs/Shared/EpicGames.UHT/Exporters/CodeGen/UhtHeaderCodeGeneratorCppFile.cs
+++ b/Engine/Source/Programs/Shared/EpicGames.UHT/Exporters/CodeGen/UhtHeaderCodeGeneratorCppFile.cs
@@ -1711,14 +1711,7 @@ namespace EpicGames.UHT.Exporters.CodeGen
 				}
 
 				builder.Append("\tstatic constexpr FCppClassTypeInfoStatic StaticCppClassTypeInfo = {\r\n");
-				if (classObj.ClassFlags.HasAnyFlags(EClassFlags.Interface))
-				{
-					builder.Append("\t\tTCppClassTypeTraits<I").Append(classObj.EngineName).Append(">::IsAbstract,\r\n");
-				}
-				else
-				{
-					builder.Append("\t\tTCppClassTypeTraits<").Append(classObj.SourceName).Append(">::IsAbstract,\r\n");
-				}
+				builder.Append("\t\tTCppClassTypeTraits<").AppendClassSourceNameOrInterfaceName(classObj).Append(">::IsAbstract,\r\n");
 				builder.Append("\t};\r\n");
 
 				builder.Append("\tstatic const UECodeGen_Private::FClassParams ClassParams;\r\n");
diff --git a/Engine/Source/Programs/Shared/EpicGames.UHT/Utils/UhtStringBuilderExtensions.cs b/Engine/Source/Programs/Shared/EpicGames.UHT/Utils/UhtStringBuilderExtensions.cs
--- a/Engine/Source/Programs/Shared/EpicGames.UHT/Utils/UhtStringBuilderExtensions.cs
+++ b/Engine/Source/Programs/Shared/EpicGames.UHT/Utils/UhtStringBuilderExtensions.cs
@@ -423,7 +423,7 @@ namespace EpicGames.UHT.Utils
 				case UhtClassType.NativeInterface:
 					return builder.Append(classObj.SourceName);
 				case UhtClassType.Interface:
-					return builder.Append('I').Append(classObj.EngineName);
+					return builder.Append('I').Append(classObj.SourceName[1..]);
 				default:
 					throw new NotImplementedException();
 			}
