Improvements to how interfaces are parsed

#rnx

diff --git a/Engine/Source/Programs/Shared/EpicGames.UHT/Tokenizer/UhtTokenReaderFlowControlExtensions.cs b/Engine/Source/Programs/Shared/EpicGames.UHT/Tokenizer/UhtTokenReaderFlowControlExtensions.cs
--- a/Engine/Source/Programs/Shared/EpicGames.UHT/Tokenizer/UhtTokenReaderFlowControlExtensions.cs
+++ b/Engine/Source/Programs/Shared/EpicGames.UHT/Tokenizer/UhtTokenReaderFlowControlExtensions.cs
@@ -363,5 +363,23 @@ namespace EpicGames.UHT.Tokenizer
 			}
 			return tokenReader;
 		}
+
+		/// <summary>
+		/// Throw an exception if the given symbol is found in the token stream
+		/// </summary>
+		/// <param name="tokenReader">Token reader</param>
+		/// <param name="candidate">Symbol or identifier to match</param>
+		/// <param name="message">Message thrown if found</param>
+		/// <returns>Token reader</returns>
+		/// <exception cref="UhtException"></exception>
+		public static IUhtTokenReader ExceptionIf(this IUhtTokenReader tokenReader, string candidate, string message)
+		{
+			UhtToken token = tokenReader.PeekToken();
+			if ((token.IsIdentifier() || token.IsSymbol()) && token.IsValue(candidate))
+			{
+				throw new UhtException(tokenReader, message);
+			}
+			return tokenReader;
+		}
 	}
 }
diff --git a/Engine/Source/Programs/Shared/EpicGames.UHT/Types/Properties/UhtInterfaceProperty.cs b/Engine/Source/Programs/Shared/EpicGames.UHT/Types/Properties/UhtInterfaceProperty.cs
--- a/Engine/Source/Programs/Shared/EpicGames.UHT/Types/Properties/UhtInterfaceProperty.cs
+++ b/Engine/Source/Programs/Shared/EpicGames.UHT/Types/Properties/UhtInterfaceProperty.cs
@@ -188,7 +188,7 @@ namespace EpicGames.UHT.Types
 		[SuppressMessage("Style", "IDE0060:Remove unused parameter", Justification = "Attribute accessed method")]
 		private static UhtProperty? TScriptInterfaceProperty(UhtPropertyResolveArgs args)
 		{
-			UhtClass? propertyClass = args.ParseTemplateObject(false);
+			UhtClass? propertyClass = args.ParseTemplateObject(UhtTemplateObjectMode.PreserveNativeInterface);
 			if (propertyClass == null)
 			{
 				return null;
diff --git a/Engine/Source/Programs/Shared/EpicGames.UHT/Types/Properties/UhtLazyObjectPtrProperty.cs b/Engine/Source/Programs/Shared/EpicGames.UHT/Types/Properties/UhtLazyObjectPtrProperty.cs
--- a/Engine/Source/Programs/Shared/EpicGames.UHT/Types/Properties/UhtLazyObjectPtrProperty.cs
+++ b/Engine/Source/Programs/Shared/EpicGames.UHT/Types/Properties/UhtLazyObjectPtrProperty.cs
@@ -107,7 +107,7 @@ namespace EpicGames.UHT.Types
 		[SuppressMessage("Style", "IDE0060:Remove unused parameter", Justification = "Attribute accessed method")]
 		private static UhtProperty? LazyObjectPtrProperty(UhtPropertyResolveArgs args)
 		{
-			UhtClass? propertyClass = args.ParseTemplateObject(true);
+			UhtClass? propertyClass = args.ParseTemplateObject(UhtTemplateObjectMode.Normal);
 			if (propertyClass == null)
 			{
 				return null;
diff --git a/Engine/Source/Programs/Shared/EpicGames.UHT/Types/Properties/UhtObjectPtrProperty.cs b/Engine/Source/Programs/Shared/EpicGames.UHT/Types/Properties/UhtObjectPtrProperty.cs
--- a/Engine/Source/Programs/Shared/EpicGames.UHT/Types/Properties/UhtObjectPtrProperty.cs
+++ b/Engine/Source/Programs/Shared/EpicGames.UHT/Types/Properties/UhtObjectPtrProperty.cs
@@ -126,7 +126,7 @@ namespace EpicGames.UHT.Types
 		{
 			int typeStartPos = args.TokenReader.PeekToken().InputStartPos;
 
-			UhtClass? propertyClass = args.ParseTemplateObject(true);
+			UhtClass? propertyClass = args.ParseTemplateObject(UhtTemplateObjectMode.Normal);
 			if (propertyClass == null)
 			{
 				return null;
diff --git a/Engine/Source/Programs/Shared/EpicGames.UHT/Types/Properties/UhtSoftObjectProperty.cs b/Engine/Source/Programs/Shared/EpicGames.UHT/Types/Properties/UhtSoftObjectProperty.cs
--- a/Engine/Source/Programs/Shared/EpicGames.UHT/Types/Properties/UhtSoftObjectProperty.cs
+++ b/Engine/Source/Programs/Shared/EpicGames.UHT/Types/Properties/UhtSoftObjectProperty.cs
@@ -92,7 +92,7 @@ namespace EpicGames.UHT.Types
 		[SuppressMessage("Style", "IDE0060:Remove unused parameter", Justification = "Attribute accessed method")]
 		private static UhtProperty? SoftObjectPtrProperty(UhtPropertyResolveArgs args)
 		{
-			UhtClass? propertyClass = args.ParseTemplateObject(true);
+			UhtClass? propertyClass = args.ParseTemplateObject(UhtTemplateObjectMode.Normal);
 			if (propertyClass == null)
 			{
 				return null;
diff --git a/Engine/Source/Programs/Shared/EpicGames.UHT/Types/Properties/UhtWeakObjectPtrProperty.cs b/Engine/Source/Programs/Shared/EpicGames.UHT/Types/Properties/UhtWeakObjectPtrProperty.cs
--- a/Engine/Source/Programs/Shared/EpicGames.UHT/Types/Properties/UhtWeakObjectPtrProperty.cs
+++ b/Engine/Source/Programs/Shared/EpicGames.UHT/Types/Properties/UhtWeakObjectPtrProperty.cs
@@ -114,7 +114,7 @@ namespace EpicGames.UHT.Types
 		[SuppressMessage("Style", "IDE0060:Remove unused parameter", Justification = "Attribute accessed method")]
 		private static UhtProperty? WeakObjectProperty(UhtPropertyResolveArgs args)
 		{
-			UhtClass? propertyClass = args.ParseTemplateObject(true);
+			UhtClass? propertyClass = args.ParseTemplateObject(UhtTemplateObjectMode.Normal);
 			if (propertyClass == null)
 			{
 				return null;
@@ -128,7 +128,7 @@ namespace EpicGames.UHT.Types
 		[SuppressMessage("Style", "IDE0060:Remove unused parameter", Justification = "Attribute accessed method")]
 		private static UhtProperty? AutoWeakObjectProperty(UhtPropertyResolveArgs args)
 		{
-			UhtClass? propertyClass = args.ParseTemplateObject(true);
+			UhtClass? propertyClass = args.ParseTemplateObject(UhtTemplateObjectMode.Normal);
 			if (propertyClass == null)
 			{
 				return null;
diff --git a/Engine/Source/Programs/Shared/EpicGames.UHT/Utils/UhtPropertyResolveArgs.cs b/Engine/Source/Programs/Shared/EpicGames.UHT/Utils/UhtPropertyResolveArgs.cs
--- a/Engine/Source/Programs/Shared/EpicGames.UHT/Utils/UhtPropertyResolveArgs.cs
+++ b/Engine/Source/Programs/Shared/EpicGames.UHT/Utils/UhtPropertyResolveArgs.cs
@@ -27,6 +27,22 @@ namespace EpicGames.UHT.Utils
 		Resolving,
 	}
 
+	/// <summary>
+	/// Controls what is returned from the parsing of a template object
+	/// </summary>
+	public enum UhtTemplateObjectMode
+	{
+		/// <summary>
+		/// If the specified type.  If the type is a native interface, return the UInterface
+		/// </summary>
+		Normal,
+
+		/// <summary>
+		/// If a native interface is specified, return it
+		/// </summary>
+		PreserveNativeInterface,
+	}
+
 	/// <summary>
 	/// Structure defining the arguments used in the delegate to resolve property types
 	/// </summary>
@@ -162,9 +178,9 @@ namespace EpicGames.UHT.Utils
 		/// <summary>
 		/// Parse a template type
 		/// </summary>
-		/// <param name="returnUInterface">If true, return the UInterface instead of the type listed</param>
+		/// <param name="mode">Adjust what is returned and/or what is searched</param>
 		/// <returns>Referenced class</returns>
-		public readonly UhtClass? ParseTemplateObject(bool returnUInterface)
+		public readonly UhtClass? ParseTemplateObject(UhtTemplateObjectMode mode)
 		{
 			if (TokenReader.TryOptional("const"))
 			{
@@ -183,6 +199,7 @@ namespace EpicGames.UHT.Utils
 				.Optional("const", () => { isNativeConstTemplateArg = true; })
 				.Optional("class")
 				.RequireIdentifier((ref UhtToken token) => { identifier = token; })
+				.ExceptionIf("::", "Template arguments to object types do not currently support namespaces")
 				.Optional("const", () => { isNativeConstTemplateArg = true; })
 				.Require('>');
 
@@ -193,7 +210,7 @@ namespace EpicGames.UHT.Utils
 
 			Config.RedirectTypeIdentifier(ref identifier);
 			UhtClass? returnClass = PropertySettings.Outer.FindType(UhtFindOptions.SourceName | UhtFindOptions.Class, ref identifier, TokenReader) as UhtClass;
-			if (returnClass != null && returnClass.AlternateObject != null && returnUInterface)
+			if (returnClass != null && returnClass.AlternateObject != null && mode == UhtTemplateObjectMode.Normal)
 			{
 				returnClass = returnClass.AlternateObject as UhtClass;
 			}
@@ -222,6 +239,7 @@ namespace EpicGames.UHT.Utils
 				.Require('<')
 				.Optional("class")
 				.RequireIdentifier((ref UhtToken token) => { identifier = token; })
+				.ExceptionIf("::", "Template arguments to object types do not currently support namespaces")
 				.Require('>');
 
 			Config.RedirectTypeIdentifier(ref identifier);
@@ -255,6 +273,7 @@ namespace EpicGames.UHT.Utils
 				.Require('<')
 				.Optional("struct")
 				.RequireIdentifier((ref UhtToken token) => { identifier = token; })
+				.ExceptionIf("::", "Template arguments to object types do not currently support namespaces")
 				.Require('>');
 
 			Config.RedirectTypeIdentifier(ref identifier);
diff --git a/Engine/Source/Programs/Shared/EpicGames.UHT/Utils/UhtSession.cs b/Engine/Source/Programs/Shared/EpicGames.UHT/Utils/UhtSession.cs
--- a/Engine/Source/Programs/Shared/EpicGames.UHT/Utils/UhtSession.cs
+++ b/Engine/Source/Programs/Shared/EpicGames.UHT/Utils/UhtSession.cs
@@ -283,7 +283,7 @@ namespace EpicGames.UHT.Utils
 				}
 				relativePath = Path.GetRelativePath(includePath!, filePath);
 			}
-			return relativePath!.Replace("\\", "/");
+			return relativePath!.Replace("\\", "/", StringComparison.Ordinal);
 		}
 
 		private string MakePath(UHTManifest.Module module, string fileName, string suffix)
