Renamed a few helper functions used by defined scopes.
Refactored how inheritance is parsed

#rnx

diff --git a/Engine/Source/Programs/Shared/EpicGames.UHT/Exporters/CodeGen/UhtHeaderCodeGeneratorHFile.cs b/Engine/Source/Programs/Shared/EpicGames.UHT/Exporters/CodeGen/UhtHeaderCodeGeneratorHFile.cs
--- a/Engine/Source/Programs/Shared/EpicGames.UHT/Exporters/CodeGen/UhtHeaderCodeGeneratorHFile.cs
+++ b/Engine/Source/Programs/Shared/EpicGames.UHT/Exporters/CodeGen/UhtHeaderCodeGeneratorHFile.cs
@@ -939,7 +939,7 @@ namespace EpicGames.UHT.Exporters.CodeGen
 				return builder;
 			}
 
-			return builder.AppendMultiMacros(autoGetterSetterProperties, UhtDefineScopeNames.Standard, this, classObj, AutoGettersSettersMacroSuffix,
+			return builder.AppendMultiScopedMacros(autoGetterSetterProperties, UhtDefineScopeNames.Standard, this, classObj, AutoGettersSettersMacroSuffix,
 				(builder, properties) =>
 				{
 					foreach (UhtProperty property in properties)
@@ -981,7 +981,7 @@ namespace EpicGames.UHT.Exporters.CodeGen
 
 			AppendSparseStructDeclarations(builder, classObj, sparseScriptStructs);
 
-			builder.AppendMultiMacros(sparseProperties, UhtDefineScopeNames.Standard, this, classObj, SparseDataPropertyAccessorsMacroSuffix, 
+			builder.AppendMultiScopedMacros(sparseProperties, UhtDefineScopeNames.Standard, this, classObj, SparseDataPropertyAccessorsMacroSuffix, 
 				(builder, properties) =>
 				{
 					foreach (UhtProperty property in properties)
@@ -1031,7 +1031,7 @@ namespace EpicGames.UHT.Exporters.CodeGen
 
 		private StringBuilder AppendRpcFunctions(StringBuilder builder, UhtClass classObj, bool usesLegacy, UhtUsedDefineScopes<UhtFunction> functions)
 		{
-			builder.AppendMultiMacros(functions, UhtDefineScopeNames.WithEditor, this, classObj, usesLegacy ? RpcWrappersMacroSuffix : RpcWrappersNoPureDeclsMacroSuffix, 
+			builder.AppendMultiScopedMacros(functions, UhtDefineScopeNames.WithEditor, this, classObj, usesLegacy ? RpcWrappersMacroSuffix : RpcWrappersNoPureDeclsMacroSuffix,
 				(builder, functions) =>
 				{
 					if (usesLegacy)
diff --git a/Engine/Source/Programs/Shared/EpicGames.UHT/Exporters/CodeGen/UhtUsedDefineScopes.cs b/Engine/Source/Programs/Shared/EpicGames.UHT/Exporters/CodeGen/UhtUsedDefineScopes.cs
--- a/Engine/Source/Programs/Shared/EpicGames.UHT/Exporters/CodeGen/UhtUsedDefineScopes.cs
+++ b/Engine/Source/Programs/Shared/EpicGames.UHT/Exporters/CodeGen/UhtUsedDefineScopes.cs
@@ -215,6 +215,24 @@ namespace EpicGames.UHT.Exporters.CodeGen
 			return builder;
 		}
 
+		/// <summary>
+		/// Append multiple  for the given collection of scopes
+		/// </summary>
+		/// <param name="builder">Destination builder</param>
+		/// <param name="instances">Defined scope instances</param>
+		/// <param name="defineScopeNames">Which set of scope names will be used</param>
+		/// <param name="appendAction">Action to invoke to append an instance</param>
+		/// <returns>String builder</returns>
+		public static StringBuilder AppendMultiScoped<T>(this StringBuilder builder, UhtUsedDefineScopes<T> instances, UhtDefineScopeNames defineScopeNames,
+			Action<StringBuilder, IEnumerable<T>> appendAction) where T : UhtType
+		{
+			foreach (UhtDefineScope defineScope in instances.EnumerateDefinedScopes())
+			{
+				AppendScoped(builder, defineScopeNames, defineScope, builder => appendAction(builder, instances.Instances.Where(x => x.DefineScope == defineScope)));
+			}
+			return builder;
+		}
+
 		/// <summary>
 		/// Append a macro scoped macro
 		/// </summary>
@@ -273,7 +291,7 @@ namespace EpicGames.UHT.Exporters.CodeGen
 		/// <param name="macroSuffix">Macro being created</param>
 		/// <param name="appendAction">Action to invoke to append an instance</param>
 		/// <returns>String builder</returns>
-		public static StringBuilder AppendMultiMacros<T>(this StringBuilder builder, UhtUsedDefineScopes<T> instances, UhtDefineScopeNames defineScopeNames,
+		public static StringBuilder AppendMultiScopedMacros<T>(this StringBuilder builder, UhtUsedDefineScopes<T> instances, UhtDefineScopeNames defineScopeNames,
 			UhtHeaderCodeGenerator generator, UhtType outerType, string macroSuffix, Action<StringBuilder, IEnumerable<T>> appendAction) where T : UhtType
 		{
 			foreach (UhtDefineScope defineScope in instances.EnumerateDefinedScopes())
diff --git a/Engine/Source/Programs/Shared/EpicGames.UHT/Parsers/UhtClassParser.cs b/Engine/Source/Programs/Shared/EpicGames.UHT/Parsers/UhtClassParser.cs
--- a/Engine/Source/Programs/Shared/EpicGames.UHT/Parsers/UhtClassParser.cs
+++ b/Engine/Source/Programs/Shared/EpicGames.UHT/Parsers/UhtClassParser.cs
@@ -99,9 +99,9 @@ namespace EpicGames.UHT.Parsers
 					topScope.TokenReader.Optional("final");
 
 					// Parse the inheritance
-					UhtParserHelpers.ParseInheritance(parentScope.HeaderParser, topScope.Session.Config!, out UhtToken superIdentifier, out List<UhtToken[]>? baseIdentifiers);
-					classObj.SuperIdentifier = superIdentifier;
-					classObj.BaseIdentifiers = baseIdentifiers;
+					UhtParserHelpers.ParseInheritance(parentScope.HeaderParser, topScope.Session.Config!, out UhtInheritance inheritance);
+					classObj.SuperIdentifier = inheritance.SuperIdentifier;
+					classObj.BaseIdentifiers = inheritance.BaseIdentifiers;
 
 					if (apiMacroToken)
 					{
diff --git a/Engine/Source/Programs/Shared/EpicGames.UHT/Parsers/UhtInterfaceClassParser.cs b/Engine/Source/Programs/Shared/EpicGames.UHT/Parsers/UhtInterfaceClassParser.cs
--- a/Engine/Source/Programs/Shared/EpicGames.UHT/Parsers/UhtInterfaceClassParser.cs
+++ b/Engine/Source/Programs/Shared/EpicGames.UHT/Parsers/UhtInterfaceClassParser.cs
@@ -78,10 +78,8 @@ namespace EpicGames.UHT.Parsers
 					classObj.SourceName = topScope.TokenReader.GetIdentifier().Value.ToString();
 
 					// Parse the inheritance
-					topScope.TokenReader.OptionalInheritance((ref UhtToken identifier) =>
-					{
-						classObj.SuperIdentifier = identifier;
-					});
+					UhtParserHelpers.ParseInheritance(parentScope.HeaderParser, topScope.Session.Config!, out UhtInheritance inheritance);
+					classObj.SuperIdentifier = inheritance.SuperIdentifier;
 
 					// Split the source name into the different parts
 					UhtEngineNameParts nameParts = UhtUtilities.GetEngineNameParts(classObj.SourceName);
diff --git a/Engine/Source/Programs/Shared/EpicGames.UHT/Parsers/UhtParserHelpers.cs b/Engine/Source/Programs/Shared/EpicGames.UHT/Parsers/UhtParserHelpers.cs
--- a/Engine/Source/Programs/Shared/EpicGames.UHT/Parsers/UhtParserHelpers.cs
+++ b/Engine/Source/Programs/Shared/EpicGames.UHT/Parsers/UhtParserHelpers.cs
@@ -9,6 +9,31 @@ using EpicGames.UHT.Utils;
 
 namespace EpicGames.UHT.Parsers
 {
+
+	/// <summary>
+	/// Collection of parsed inheritance data
+	/// </summary>
+	public struct UhtInheritance
+	{
+
+		/// <summary>
+		/// Super class identifier
+		/// </summary>
+		public UhtToken SuperIdentifier { get; set; } = new();
+
+		/// <summary>
+		/// Collection of other base identifiers
+		/// </summary>
+		public List<UhtToken[]> BaseIdentifiers { get; } = new();
+
+		/// <summary>
+		/// Constructor
+		/// </summary>
+		public UhtInheritance()
+		{
+		}
+	}
+		
 	/// <summary>
 	/// Collection of helper methods
 	/// </summary>
@@ -20,32 +45,37 @@ namespace EpicGames.UHT.Parsers
 		/// </summary>
 		/// <param name="headerFileParser">Header file being parsed</param>
 		/// <param name="config">Configuration</param>
-		/// <param name="superIdentifier">Output super identifier</param>
-		/// <param name="baseIdentifiers">Output base identifiers</param>
-		public static void ParseInheritance(UhtHeaderFileParser headerFileParser, IUhtConfig config, out UhtToken superIdentifier, out List<UhtToken[]>? baseIdentifiers)
+		/// <param name="inheritance">Collection of inheritance information</param>
+		public static void ParseInheritance(UhtHeaderFileParser headerFileParser, IUhtConfig config, out UhtInheritance inheritance)
 		{
-
 			// TODO: C++ UHT doesn't allow preprocessor statements inside of inheritance lists
 			string? restrictedPreprocessorContext = headerFileParser.RestrictedPreprocessorContext;
 			headerFileParser.RestrictedPreprocessorContext = "parsing inheritance list";
 
 			try
 			{
-				UhtToken superIdentifierTemp = new();
-				List<UhtToken[]>? baseIdentifiersTemp = null;
-				headerFileParser.TokenReader.OptionalInheritance(
-					(ref UhtToken identifier) =>
-					{
-						config.RedirectTypeIdentifier(ref identifier);
-						superIdentifierTemp = identifier;
-					},
-					(UhtTokenList identifier) =>
-					{
-						baseIdentifiersTemp ??= new List<UhtToken[]>();
-						baseIdentifiersTemp.Add(identifier.ToArray());
-					});
-				superIdentifier = superIdentifierTemp;
-				baseIdentifiers = baseIdentifiersTemp;
+				UhtInheritance scratch = new();
+				IUhtTokenReader tokenReader = headerFileParser.TokenReader;
+				tokenReader.Optional(':', () =>
+				{
+					tokenReader
+						.Require("public", "public access modifier")
+						.RequireIdentifier((ref UhtToken identifier) =>
+						{
+							config.RedirectTypeIdentifier(ref identifier);
+							scratch.SuperIdentifier = identifier;
+						})
+						.While(',', () =>
+						{
+							tokenReader
+								.Require("public", "public interface access specifier")
+								.RequireCppIdentifier(UhtCppIdentifierOptions.AllowTemplates, (UhtTokenList identifier) =>
+								{
+									scratch.BaseIdentifiers.Add(identifier.ToArray());
+								});
+						});
+				});
+				inheritance = scratch;
 			}
 			finally
 			{
diff --git a/Engine/Source/Programs/Shared/EpicGames.UHT/Parsers/UhtScriptStructParser.cs b/Engine/Source/Programs/Shared/EpicGames.UHT/Parsers/UhtScriptStructParser.cs
--- a/Engine/Source/Programs/Shared/EpicGames.UHT/Parsers/UhtScriptStructParser.cs
+++ b/Engine/Source/Programs/Shared/EpicGames.UHT/Parsers/UhtScriptStructParser.cs
@@ -110,9 +110,9 @@ namespace EpicGames.UHT.Parsers
 					topScope.TokenReader.Optional("final");
 
 					// Parse the inheritance
-					UhtParserHelpers.ParseInheritance(parentScope.HeaderParser, topScope.Session.Config!, out UhtToken superIdentifier, out List<UhtToken[]>? baseIdentifiers);
-					scriptStruct.SuperIdentifier = superIdentifier;
-					scriptStruct.BaseIdentifiers = baseIdentifiers;
+					UhtParserHelpers.ParseInheritance(parentScope.HeaderParser, topScope.Session.Config!, out UhtInheritance inheritance);
+					scriptStruct.SuperIdentifier = inheritance.SuperIdentifier;
+					scriptStruct.BaseIdentifiers = inheritance.BaseIdentifiers;
 
 					// Add the comments here for compatibility with old UHT
 					//COMPATIBILITY-TODO - Move this back to where the AddModuleRelativePathToMetaData is called.
diff --git a/Engine/Source/Programs/Shared/EpicGames.UHT/Tokenizer/UhtTokenReaderUtilityExtensions.cs b/Engine/Source/Programs/Shared/EpicGames.UHT/Tokenizer/UhtTokenReaderUtilityExtensions.cs
--- a/Engine/Source/Programs/Shared/EpicGames.UHT/Tokenizer/UhtTokenReaderUtilityExtensions.cs
+++ b/Engine/Source/Programs/Shared/EpicGames.UHT/Tokenizer/UhtTokenReaderUtilityExtensions.cs
@@ -7,7 +7,6 @@ using EpicGames.UHT.Utils;
 
 namespace EpicGames.UHT.Tokenizer
 {
-
 	/// <summary>
 	/// Collection of assorted utility token reader extensions
 	/// </summary>
@@ -50,47 +49,6 @@ namespace EpicGames.UHT.Tokenizer
 			return false;
 		}
 
-		/// <summary>
-		/// Parse an optional single inheritance
-		/// </summary>
-		/// <param name="tokenReader">Token reader</param>
-		/// <param name="superClassDelegate">Invoked with the inherited type name</param>
-		/// <returns>Token reader</returns>
-		public static IUhtTokenReader OptionalInheritance(this IUhtTokenReader tokenReader, UhtTokenDelegate superClassDelegate)
-		{
-			tokenReader.Optional(':', () =>
-			{
-				tokenReader
-					.Require("public", "public access modifier")
-					.RequireIdentifier(superClassDelegate);
-			});
-			return tokenReader;
-		}
-
-		/// <summary>
-		/// Parse an optional inheritance
-		/// </summary>
-		/// <param name="tokenReader">Token reader</param>
-		/// <param name="superClassDelegate">Invoked with the inherited type name</param>
-		/// <param name="baseClassDelegate">Invoked when other base classes are parsed</param>
-		/// <returns>Token reader</returns>
-		public static IUhtTokenReader OptionalInheritance(this IUhtTokenReader tokenReader, UhtTokenDelegate superClassDelegate, UhtTokenListDelegate baseClassDelegate)
-		{
-			tokenReader.Optional(':', () =>
-			{
-				tokenReader
-					.Require("public", "public access modifier")
-					.RequireIdentifier(superClassDelegate)
-					.While(',', () =>
-					{
-						tokenReader
-							.Require("public", "public interface access specifier")
-							.RequireCppIdentifier(UhtCppIdentifierOptions.AllowTemplates, baseClassDelegate);
-					});
-			});
-			return tokenReader;
-		}
-
 		/// <summary>
 		/// Given a declaration/statement that starts with the given token, skip the declaration in the header.
 		/// </summary>
