CsCompileTask: Ensure EngineDir is set correctly

#rnx

diff --git a/Engine/Source/Programs/AutomationTool/BuildGraph/Tasks/CsCompileTask.cs b/Engine/Source/Programs/AutomationTool/BuildGraph/Tasks/CsCompileTask.cs
--- a/Engine/Source/Programs/AutomationTool/BuildGraph/Tasks/CsCompileTask.cs
+++ b/Engine/Source/Programs/AutomationTool/BuildGraph/Tasks/CsCompileTask.cs
@@ -189,6 +189,7 @@ namespace AutomationTool.Tasks
 			// Try to figure out the output files
 			HashSet<FileReference> projectBuildProducts;
 			HashSet<FileReference> projectReferences;
+			properties["EngineDirectory"] = Unreal.EngineDirectory.FullName;
 			FindBuildProductsAndReferences(projectFiles, properties, out projectBuildProducts, out projectReferences);
 
 			// Apply the optional tag to the produced archive
diff --git a/Engine/Source/Programs/Shared/EpicGames.Core/CsProjectInfo.cs b/Engine/Source/Programs/Shared/EpicGames.Core/CsProjectInfo.cs
--- a/Engine/Source/Programs/Shared/EpicGames.Core/CsProjectInfo.cs
+++ b/Engine/Source/Programs/Shared/EpicGames.Core/CsProjectInfo.cs
@@ -494,8 +494,20 @@ namespace EpicGames.Core
 					case "IsLinux": projectInfo.Properties[element.Name] = OperatingSystem.IsLinux().ToString(); continue;
 					case "IsMacOS": projectInfo.Properties[element.Name] = OperatingSystem.IsMacOS().ToString(); continue;
 					case "IsWindows": projectInfo.Properties[element.Name] = OperatingSystem.IsWindows().ToString(); continue;
+					case "IsX64": projectInfo.Properties[element.Name] = (RuntimeInformation.OSArchitecture == Architecture.X64).ToString(); continue;
+					case "IsArm64": projectInfo.Properties[element.Name] = (RuntimeInformation.OSArchitecture == Architecture.Arm64).ToString(); continue;
 					case "OSArchitecture": projectInfo.Properties[element.Name] = RuntimeInformation.OSArchitecture.ToString(); continue;
 					case "EngineDirectory": projectInfo.Properties.TryAdd(element.Name, String.Empty); continue;
+					case "EngineDir":
+					case "EnginePath":
+						{
+							projectInfo.Properties[element.Name] = String.Empty;
+							if (projectInfo.Properties.TryGetValue("EngineDirectory", out string? engineDirectory) && !String.IsNullOrEmpty(engineDirectory))
+							{
+								projectInfo.Properties[element.Name] = engineDirectory;
+							}
+							continue;
+						}
 					case "IsEngineProject":
 					{
 						projectInfo.Properties[element.Name] = "False";
