Added ability to specify backends in uartifact files and added a prototype of how it could be used in the create cloud artifact task to generate a uartifact file that UGS could sync.

diff --git a/Engine/Source/Programs/AutomationTool/BuildGraph/Tasks/CreateCloudArtifactTask.cs b/Engine/Source/Programs/AutomationTool/BuildGraph/Tasks/CreateCloudArtifactTask.cs
--- a/Engine/Source/Programs/AutomationTool/BuildGraph/Tasks/CreateCloudArtifactTask.cs
+++ b/Engine/Source/Programs/AutomationTool/BuildGraph/Tasks/CreateCloudArtifactTask.cs
@@ -16,6 +16,7 @@ using EpicGames.Compression;
 using EpicGames.Core;
 using EpicGames.Horde.Artifacts;
 using EpicGames.Horde.Commits;
+using EpicGames.Horde.Storage;
 using EpicGames.Horde.Streams;
 using EpicGames.Serialization;
 using Microsoft.Extensions.Logging;
@@ -375,6 +376,16 @@ namespace AutomationTool.Tasks
 
 			// set part name to default by default - should be exposed as a option to allow for uploading separate parts
 			string partName = "default";
+
+			Uri buildUri = new Uri($"{host}/api/v2/builds/{ns}/{bucketId}/{buildId}");
+
+			// we create the artifact descriptor that Horde can then pickup
+			// TODO: upload this to horde
+			ArtifactDescriptor _ = new ArtifactDescriptor
+			{
+				Backend = new ArtifactBackend {Type = ArtifactBackend.BackendTypeEnum.Zen}, Name = name.ToString(), BaseUrl = buildUri
+			};
+
 			if (_parameters.UseZenCli)
 			{
 				BuildUploader.UploadBuild(host, baseDir, ns, bucketId, buildId, partName, partId, metadata, accessToken, new BuildUploader.BuildInfo
diff --git a/Engine/Source/Programs/Shared/EpicGames.Horde/Artifacts/ArtifactDescriptor.cs b/Engine/Source/Programs/Shared/EpicGames.Horde/Artifacts/ArtifactDescriptor.cs
--- a/Engine/Source/Programs/Shared/EpicGames.Horde/Artifacts/ArtifactDescriptor.cs
+++ b/Engine/Source/Programs/Shared/EpicGames.Horde/Artifacts/ArtifactDescriptor.cs
@@ -11,6 +11,32 @@ using EpicGames.Core;
 
 namespace EpicGames.Horde.Storage
 {
+	/// <summary>
+	/// Class that defines which backend to use to fetch an artifact
+	/// </summary>
+	public class ArtifactBackend
+	{
+		/// <summary>
+		/// The different types of supported backends
+		/// </summary>
+		public enum BackendTypeEnum
+		{
+			/// <summary>
+			/// The horde storage backend
+			/// </summary>
+			Horde,
+			/// <summary>
+			/// A zen based backend
+			/// </summary>
+			Zen
+		}
+
+		/// <summary>
+		/// Enum indicating which type of backend it is and thus how these options should be interpreted
+		/// </summary>
+		public BackendTypeEnum Type { get; set; } = BackendTypeEnum.Horde;
+	}
+
 	/// <summary>
 	/// Class which describes an artifact on Horde which can be serialized to JSON.
 	/// </summary>
@@ -61,6 +87,11 @@ namespace EpicGames.Horde.Storage
 		/// </summary>
 		public List<string>? Filter { get; set; }
 
+		/// <summary>
+		/// Information about which backend system to use. Defaults to Horde.
+		/// </summary>
+		public ArtifactBackend Backend { get; set; } = new ArtifactBackend();
+
 		/// <summary>
 		/// Default constructor
 		/// </summary>
