UE: Add temporary logging about the currently installed SDK to help debug an autosdk issue

diff --git a/Engine/Source/Programs/Shared/EpicGames.Build/UEBuildPlatformSDK.cs b/Engine/Source/Programs/Shared/EpicGames.Build/UEBuildPlatformSDK.cs
--- a/Engine/Source/Programs/Shared/EpicGames.Build/UEBuildPlatformSDK.cs
+++ b/Engine/Source/Programs/Shared/EpicGames.Build/UEBuildPlatformSDK.cs
@@ -2105,17 +2105,22 @@ namespace EpicGames.Core
 						// check to make sure OutputEnvVars doesn't need regenerating
 						string EnvVarFile = Path.Combine(AutoSDKRoot, SDKEnvironmentVarsFile);
 						bool bEnvVarFileExists = File.Exists(EnvVarFile);
-
-						string CurrentSDKString;
-						string CurrentSDKLevel;
-						if (bEnvVarFileExists && GetCurrentlyInstalledSDKString(AutoSDKRoot, out CurrentSDKString, out CurrentSDKLevel))
+						if (bEnvVarFileExists)
 						{
-							// match version
-							if (CurrentSDKString == GetAutoSDKDirectoryForMainVersion())
+							string CurrentSDKString;
+							string CurrentSDKLevel;
+							if (GetCurrentlyInstalledSDKString(AutoSDKRoot, out CurrentSDKString, out CurrentSDKLevel))
 							{
-								if (IsSDKLevelAtLeast(CurrentSDKLevel, DesiredSDKLevel))
+								string MainVersion = GetAutoSDKDirectoryForMainVersion();
+								Logger.LogInformation("Currently installed SDK for '{PlatformName}' is '{InstalledVersion}'. Main version is {MainVersion}. CurrentSDKLevel {CurrentSDKLevel}. DesiredSDKLevel {DesiredSDKLevel}",
+									PlatformName, CurrentSDKString, MainVersion, CurrentSDKLevel, DesiredSDKLevel);
+								// match version
+								if (CurrentSDKString == MainVersion)
 								{
-									return SDKStatus.Valid;
+									if (IsSDKLevelAtLeast(CurrentSDKLevel, DesiredSDKLevel))
+									{
+										return SDKStatus.Valid;
+									}
 								}
 							}
 						}
