[Backout] - CL40490782
#rnx
Original CL Desc
-----------------------------------------------------------------
UE: Add temporary logging about the currently installed SDK to help debug an autosdk issue

diff --git a/Engine/Source/Programs/Shared/EpicGames.Build/UEBuildPlatformSDK.cs b/Engine/Source/Programs/Shared/EpicGames.Build/UEBuildPlatformSDK.cs
--- a/Engine/Source/Programs/Shared/EpicGames.Build/UEBuildPlatformSDK.cs
+++ b/Engine/Source/Programs/Shared/EpicGames.Build/UEBuildPlatformSDK.cs
@@ -2105,22 +2105,17 @@ namespace EpicGames.Core
 						// check to make sure OutputEnvVars doesn't need regenerating
 						string EnvVarFile = Path.Combine(AutoSDKRoot, SDKEnvironmentVarsFile);
 						bool bEnvVarFileExists = File.Exists(EnvVarFile);
-						if (bEnvVarFileExists)
+
+						string CurrentSDKString;
+						string CurrentSDKLevel;
+						if (bEnvVarFileExists && GetCurrentlyInstalledSDKString(AutoSDKRoot, out CurrentSDKString, out CurrentSDKLevel))
 						{
-							string CurrentSDKString;
-							string CurrentSDKLevel;
-							if (GetCurrentlyInstalledSDKString(AutoSDKRoot, out CurrentSDKString, out CurrentSDKLevel))
+							// match version
+							if (CurrentSDKString == GetAutoSDKDirectoryForMainVersion())
 							{
-								string MainVersion = GetAutoSDKDirectoryForMainVersion();
-								Logger.LogInformation("Currently installed SDK for '{PlatformName}' is '{InstalledVersion}'. Main version is {MainVersion}. CurrentSDKLevel {CurrentSDKLevel}. DesiredSDKLevel {DesiredSDKLevel}",
-									PlatformName, CurrentSDKString, MainVersion, CurrentSDKLevel, DesiredSDKLevel);
-								// match version
-								if (CurrentSDKString == MainVersion)
+								if (IsSDKLevelAtLeast(CurrentSDKLevel, DesiredSDKLevel))
 								{
-									if (IsSDKLevelAtLeast(CurrentSDKLevel, DesiredSDKLevel))
-									{
-										return SDKStatus.Valid;
-									}
+									return SDKStatus.Valid;
 								}
 							}
 						}
