Horde: Lower log level for socket send buffer teardown from warning to info

Lower the log level to prevent non-critical messages from appearing in job step logs. These messages about socket send buffer teardown are primarily relevant to Horde developers rather than users running builds.

This change reduces log "pollution" until a better communication channel for Horde-specific diagnostics can be implemented.

diff --git a/Engine/Source/Programs/Shared/EpicGames.Horde/Compute/ComputeSocket.cs b/Engine/Source/Programs/Shared/EpicGames.Horde/Compute/ComputeSocket.cs
--- a/Engine/Source/Programs/Shared/EpicGames.Horde/Compute/ComputeSocket.cs
+++ b/Engine/Source/Programs/Shared/EpicGames.Horde/Compute/ComputeSocket.cs
@@ -809,7 +809,7 @@ namespace EpicGames.Horde.Compute
 				}
 				catch (TimeoutException ex)
 				{
-					_logger.LogWarning(ex, "Send task did not terminate gracefully in {Timeout} ms. Tearing down socket buffers...", timeout.TotalMilliseconds);
+					_logger.LogInformation(ex, "Send task did not terminate gracefully in {Timeout} ms. Tearing down socket buffers...", timeout.TotalMilliseconds);
 				}
 				catch (Exception ex)
 				{
