Linux OIDC: Prefer xdg-open, with wslview as fallback

With WSLg, xdg-open provides a simpler and more reliable path to open a browser and authenticate.

[RN] Linux WSL: OIDC requests now prefer the default browser via xdg-open, with wslview as fallback

diff --git a/Engine/Source/Programs/Shared/EpicGames.OIDC/OidcTokenManager.cs b/Engine/Source/Programs/Shared/EpicGames.OIDC/OidcTokenManager.cs
--- a/Engine/Source/Programs/Shared/EpicGames.OIDC/OidcTokenManager.cs
+++ b/Engine/Source/Programs/Shared/EpicGames.OIDC/OidcTokenManager.cs
@@ -713,25 +713,32 @@ namespace EpicGames.OIDC
 			}
 			else if (RuntimeInformation.IsOSPlatform(OSPlatform.Linux))
 			{
+				List<string> Browsers = ["xdg-open"];
 				if (IsRunningWsl())
 				{
-					_logger?.LogDebug("Running under wsl, using wslview to open browser.");
+					Browsers.Append("wslview");
+				}
+
+				foreach (string Browser in Browsers)
+				{
 					try
 					{
-						return Process.Start("wslview", url);
+						return Process.Start(Browser, url);
 					}
-					catch (Win32Exception e)
+					catch (Exception e)
 					{
-						if (e.NativeErrorCode == 2)
-						{
-							throw new Exception("Unable to start wslview, make sure its installed by running 'sudo apt install wslu' .", e);
-						}
-
-						throw;
+						_logger?.LogDebug("Failed to launch {Browser}", e);
 					}
 				}
 
-				return Process.Start("xdg-open", url);	
+				string Message = "Did not find a browser, use 'sudo apt install xdg-utils'";
+				if (IsRunningWsl())
+				{
+					Message += " or 'sudo apt install wslu'";
+				}
+				Message += ".";
+				_logger?.LogInformation(Message);
+				return null;
 			}
 			else if (RuntimeInformation.IsOSPlatform(OSPlatform.OSX))
 			{
