Fix static analysis warnings in recent Linux OIDC browser launch change

diff --git a/Engine/Source/Programs/Shared/EpicGames.OIDC/OidcTokenManager.cs b/Engine/Source/Programs/Shared/EpicGames.OIDC/OidcTokenManager.cs
--- a/Engine/Source/Programs/Shared/EpicGames.OIDC/OidcTokenManager.cs
+++ b/Engine/Source/Programs/Shared/EpicGames.OIDC/OidcTokenManager.cs
@@ -714,9 +714,10 @@ namespace EpicGames.OIDC
 			else if (RuntimeInformation.IsOSPlatform(OSPlatform.Linux))
 			{
 				List<string> Browsers = ["xdg-open"];
-				if (IsRunningWsl())
+				bool IsWsl = IsRunningWsl();
+				if (IsWsl)
 				{
-					Browsers.Append("wslview");
+					Browsers.Add("wslview");
 				}
 
 				foreach (string Browser in Browsers)
@@ -731,13 +732,14 @@ namespace EpicGames.OIDC
 					}
 				}
 
-				string Message = "Did not find a browser, use 'sudo apt install xdg-utils'";
-				if (IsRunningWsl())
+				if (IsWsl)
 				{
-					Message += " or 'sudo apt install wslu'";
+					_logger?.LogInformation("Did not find a browser, use 'sudo apt install xdg-utils' or 'sudo apt install wslu'.");
+				}
+				else
+				{
+					_logger?.LogInformation("Did not find a browser, use 'sudo apt install xdg-utils'.");
 				}
-				Message += ".";
-				_logger?.LogInformation(Message);
 				return null;
 			}
 			else if (RuntimeInformation.IsOSPlatform(OSPlatform.OSX))
