EpicGames.UHT: Ensure UHT factory sessions that generate no output files will still cull the output directory

#rnx

diff --git a/Engine/Source/Programs/Shared/EpicGames.UHT/Utils/UhtSession.cs b/Engine/Source/Programs/Shared/EpicGames.UHT/Utils/UhtSession.cs
--- a/Engine/Source/Programs/Shared/EpicGames.UHT/Utils/UhtSession.cs
+++ b/Engine/Source/Programs/Shared/EpicGames.UHT/Utils/UhtSession.cs
@@ -411,12 +411,15 @@ namespace EpicGames.UHT.Utils
 			// Invoke the exported via the delegate
 			Exporter.Delegate(this);
 
+			// These outputs are used to cull old outputs from the directories
+			Dictionary<string, HashSet<string>> outputsByDirectory = new(
+				Session.Modules.Select(x => new KeyValuePair<string, HashSet<string>>(x.Module.OutputDirectory, new HashSet<string>(StringComparer.OrdinalIgnoreCase))),
+				StringComparer.OrdinalIgnoreCase
+			);
+
 			// If outputs were exported
 			if (Outputs.Count > 0)
 			{
-
-				// These outputs are used to cull old outputs from the directories
-				Dictionary<string, HashSet<string>> outputsByDirectory = new(StringComparer.OrdinalIgnoreCase);
 				List<UhtExportFactory.Output> saves = new();
 
 				// Collect information about the outputs
@@ -457,24 +460,24 @@ namespace EpicGames.UHT.Utils
 						RenameSource(output);
 					}
 				}
+			}
 
-				// Perform the culling of the output directories
-				if (Session.CullOutput && !Session.NoOutput &&
-					(Exporter.CppFilters != null || Exporter.HeaderFilters != null || Exporter.OtherFilters != null))
+			// Perform the culling of the output directories
+			if (Session.CullOutput && !Session.NoOutput &&
+				(Exporter.CppFilters != null || Exporter.HeaderFilters != null || Exporter.OtherFilters != null))
+			{
+				if (Session.GoWide)
 				{
-					if (Session.GoWide)
+					Parallel.ForEach(outputsByDirectory, (KeyValuePair<string, HashSet<string>> kvp) =>
 					{
-						Parallel.ForEach(outputsByDirectory, (KeyValuePair<string, HashSet<string>> kvp) =>
-						{
-							CullOutputDirectory(kvp.Key, kvp.Value);
-						});
-					}
-					else
+						CullOutputDirectory(kvp.Key, kvp.Value);
+					});
+				}
+				else
+				{
+					foreach (KeyValuePair<string, HashSet<string>> kvp in outputsByDirectory)
 					{
-						foreach (KeyValuePair<string, HashSet<string>> kvp in outputsByDirectory)
-						{
-							CullOutputDirectory(kvp.Key, kvp.Value);
-						}
+						CullOutputDirectory(kvp.Key, kvp.Value);
 					}
 				}
 			}
@@ -2532,7 +2535,7 @@ namespace EpicGames.UHT.Utils
 				{
 					using StreamWriter output = new(Manifest!.ExternalDependenciesFile);
 					foreach (string dep in externalDependencies)
-					{
+					{ 
 						output.WriteLine(dep);
 					}
 				}
