Horde: Add retry policy for blob writes when getting an upload redirect

Also log headers to exception as they may contain helpful information for AWS S3 related errors.

diff --git a/Engine/Source/Programs/Shared/EpicGames.Horde/HordeHttpMessageHandler.cs b/Engine/Source/Programs/Shared/EpicGames.Horde/HordeHttpMessageHandler.cs
--- a/Engine/Source/Programs/Shared/EpicGames.Horde/HordeHttpMessageHandler.cs
+++ b/Engine/Source/Programs/Shared/EpicGames.Horde/HordeHttpMessageHandler.cs
@@ -8,6 +8,7 @@ using EpicGames.Core;
 using Microsoft.Extensions.Http;
 using Microsoft.Extensions.Logging;
 using Polly;
+using Polly.Extensions.Http;
 using Polly.Retry;
 using Polly.Timeout;
 
@@ -31,7 +32,7 @@ namespace EpicGames.Horde
 		{
 			_instance = new SocketsHttpHandler();
 			_instance = new PolicyHttpMessageHandler(request => CreateDefaultTimeoutRetryPolicy(request, logger)) { InnerHandler = _instance };
-			_instance = new PolicyHttpMessageHandler(request => CreateDefaultTransientErrorPolicy(request, logger)) { InnerHandler = _instance };
+			_instance = new PolicyHttpMessageHandler(request => CreateHordeTransientErrorPolicy(request, logger)) { InnerHandler = _instance };
 		}
 
 		/// <inheritdoc/>
@@ -64,9 +65,9 @@ namespace EpicGames.Horde
 		}
 
 		/// <summary>
-		/// Create a default timeout retry policy
+		/// Create an error and timeout retry policy taking into account Horde's use of 503 errors
 		/// </summary>
-		static IAsyncPolicy<HttpResponseMessage> CreateDefaultTransientErrorPolicy(HttpRequestMessage request, ILogger logger)
+		static IAsyncPolicy<HttpResponseMessage> CreateHordeTransientErrorPolicy(HttpRequestMessage request, ILogger logger)
 		{
 			Task OnTimeoutAsync(DelegateResult<HttpResponseMessage> outcome, TimeSpan timespan, int retryAttempt, Context context)
 			{
@@ -83,5 +84,16 @@ namespace EpicGames.Horde
 				.OrResult(x => (x.StatusCode >= HttpStatusCode.InternalServerError && x.StatusCode != HttpStatusCode.ServiceUnavailable) || x.StatusCode == HttpStatusCode.RequestTimeout)
 				.WaitAndRetryAsync(retryTimes, OnTimeoutAsync);
 		}
+		
+		/// <summary>
+		/// Create a default retry policy 
+		/// </summary>
+		public static IAsyncPolicy<HttpResponseMessage> CreateDefaultTransientErrorPolicy()
+		{
+			return HttpPolicyExtensions
+				.HandleTransientHttpError()
+				.OrResult(response => response.StatusCode == HttpStatusCode.ServiceUnavailable)
+				.WaitAndRetryAsync(3, retryAttempt => TimeSpan.FromSeconds(Math.Pow(2, retryAttempt)));
+		}
 	}
 }
diff --git a/Engine/Source/Programs/Shared/EpicGames.Horde/Storage/Backends/HttpStorageBackend.cs b/Engine/Source/Programs/Shared/EpicGames.Horde/Storage/Backends/HttpStorageBackend.cs
--- a/Engine/Source/Programs/Shared/EpicGames.Horde/Storage/Backends/HttpStorageBackend.cs
+++ b/Engine/Source/Programs/Shared/EpicGames.Horde/Storage/Backends/HttpStorageBackend.cs
@@ -4,6 +4,7 @@ using System;
 using System.Collections.Generic;
 using System.Diagnostics;
 using System.IO;
+using System.Linq;
 using System.Net;
 using System.Net.Http;
 using System.Net.Http.Headers;
@@ -13,6 +14,7 @@ using System.Threading.Tasks;
 using System.Web;
 using EpicGames.Core;
 using Microsoft.Extensions.Logging;
+using Polly;
 
 namespace EpicGames.Horde.Storage.Backends
 {
@@ -155,11 +157,14 @@ namespace EpicGames.Horde.Storage.Backends
 				if (redirectResponse.UploadUrl != null)
 				{
 					using HttpClient uploadRedirectClient = _createUploadRedirectClient();
-					using HttpResponseMessage uploadResponse = await uploadRedirectClient.PutAsync(redirectResponse.UploadUrl, streamContent, cancellationToken);
+					IAsyncPolicy<HttpResponseMessage> retryPolicy = HordeHttpMessageHandler.CreateDefaultTransientErrorPolicy();
+					
+					using HttpResponseMessage uploadResponse = await retryPolicy.ExecuteAsync(() => uploadRedirectClient.PutAsync(redirectResponse.UploadUrl, streamContent, cancellationToken));
 					if (!uploadResponse.IsSuccessStatusCode)
 					{
 						string body = await uploadResponse.Content.ReadAsStringAsync(cancellationToken);
-						throw new StorageException($"Unable to upload data to redirected URL {redirectResponse.UploadUrl}. Status:{uploadResponse.StatusCode} Reason:{uploadResponse.ReasonPhrase} Response:{body}");
+						string headersText = String.Join(", ", uploadResponse.Headers.Select(h => $"{h.Key}={String.Join(",", h.Value)}"));
+						throw new StorageException($"Unable to upload data to redirected URL {redirectResponse.UploadUrl}. Status:{uploadResponse.StatusCode} Reason:{uploadResponse.ReasonPhrase} Response:{body} Headers:{headersText}");
 					}
 
 					_logger.LogDebug("Written {Locator} (using redirect)", redirectResponse.Blob);
