UHT: inline destructor for interfaces

#rnx

diff --git a/Engine/Source/Programs/Shared/EpicGames.UHT/Exporters/CodeGen/UhtHeaderCodeGeneratorCppFile.cs b/Engine/Source/Programs/Shared/EpicGames.UHT/Exporters/CodeGen/UhtHeaderCodeGeneratorCppFile.cs
--- a/Engine/Source/Programs/Shared/EpicGames.UHT/Exporters/CodeGen/UhtHeaderCodeGeneratorCppFile.cs
+++ b/Engine/Source/Programs/Shared/EpicGames.UHT/Exporters/CodeGen/UhtHeaderCodeGeneratorCppFile.cs
@@ -1392,7 +1392,7 @@ namespace EpicGames.UHT.Exporters.CodeGen
 				builder.Append("DEFINE_VTABLE_PTR_HELPER_CTOR(").Append(classObj.SourceName).Append(");\r\n");
 			}
 
-			if (!classObj.ClassExportFlags.HasAnyFlags(UhtClassExportFlags.HasDestructor))
+			if (!classObj.ClassExportFlags.HasAnyFlags(UhtClassExportFlags.HasDestructor) && !classObj.ClassFlags.HasAnyFlags(EClassFlags.Interface))
 			{
 				builder.Append(classObj.SourceName).Append("::~").Append(classObj.SourceName).Append("() {}\r\n");
 			}
diff --git a/Engine/Source/Programs/Shared/EpicGames.UHT/Exporters/CodeGen/UhtHeaderCodeGeneratorHFile.cs b/Engine/Source/Programs/Shared/EpicGames.UHT/Exporters/CodeGen/UhtHeaderCodeGeneratorHFile.cs
--- a/Engine/Source/Programs/Shared/EpicGames.UHT/Exporters/CodeGen/UhtHeaderCodeGeneratorHFile.cs
+++ b/Engine/Source/Programs/Shared/EpicGames.UHT/Exporters/CodeGen/UhtHeaderCodeGeneratorHFile.cs
@@ -1249,11 +1249,9 @@ namespace EpicGames.UHT.Exporters.CodeGen
 		/// <returns>Output builder</returns>
 		private static StringBuilder AppendCopyConstructorDefinition(StringBuilder builder, UhtClass classObj)
 		{
-			builder.Append("private: \\\r\n");
-			builder.Append("\t/** Private move- and copy-constructors, should never be used */ \\\r\n");
-			builder.Append('\t').Append(classObj.SourceName).Append('(').Append(classObj.SourceName).Append("&&); \\\r\n");
-			builder.Append('\t').Append(classObj.SourceName).Append("(const ").Append(classObj.SourceName).Append("&); \\\r\n");
-			builder.Append("public: \\\r\n");
+			builder.Append("\t/** Deleted move- and copy-constructors, should never be used */ \\\r\n");
+			builder.Append('\t').Append(classObj.SourceName).Append('(').Append(classObj.SourceName).Append("&&) = delete; \\\r\n");
+			builder.Append('\t').Append(classObj.SourceName).Append("(const ").Append(classObj.SourceName).Append("&) = delete; \\\r\n");
 			return builder;
 		}
 
@@ -1268,7 +1266,14 @@ namespace EpicGames.UHT.Exporters.CodeGen
 		{
 			if (!classObj.ClassExportFlags.HasAnyFlags(UhtClassExportFlags.HasDestructor))
 			{
-				builder.Append('\t').Append(api).Append("virtual ~").Append(classObj.SourceName).Append("(); \\\r\n");
+				if (classObj.ClassFlags.HasAnyFlags(EClassFlags.Interface))
+				{
+					builder.Append('\t').Append("virtual ~").Append(classObj.SourceName).Append("() = default; \\\r\n");	
+				}
+				else
+				{
+					builder.Append('\t').Append(api).Append("virtual ~").Append(classObj.SourceName).Append("(); \\\r\n");	
+				}
 			}
 			return builder;
 		}
