Horde: Fix integer overflow in MsfFile.GetPage for large PDB files

diff --git a/Engine/Source/Programs/Shared/EpicGames.Horde/Symbols/SymStore.cs b/Engine/Source/Programs/Shared/EpicGames.Horde/Symbols/SymStore.cs
--- a/Engine/Source/Programs/Shared/EpicGames.Horde/Symbols/SymStore.cs
+++ b/Engine/Source/Programs/Shared/EpicGames.Horde/Symbols/SymStore.cs
@@ -261,7 +261,8 @@ namespace EpicGames.Horde.Symbols
 			public Span<byte> GetPage(int pageIdx)
 			{
 				Debug.Assert(pageIdx < _numPages);
-				return _memoryMappedView.GetMemory(pageIdx * _pageSize, _pageSize).Span;
+				// Cast to long to prevent overflow on offset for large PDBs
+				return _memoryMappedView.GetMemory((long)pageIdx * _pageSize, _pageSize).Span;
 			}
 
 			ReadOnlySpan<byte> GetStreamPage(int streamIdx, int pageIdx)
