Refactor the handling of UhtObjectProperty better unify around the notion of a CppForm.

#rnx

diff --git a/Engine/Source/Programs/Shared/EpicGames.Core/UnrealEngineTypes.cs b/Engine/Source/Programs/Shared/EpicGames.Core/UnrealEngineTypes.cs
--- a/Engine/Source/Programs/Shared/EpicGames.Core/UnrealEngineTypes.cs
+++ b/Engine/Source/Programs/Shared/EpicGames.Core/UnrealEngineTypes.cs
@@ -1275,6 +1275,11 @@ namespace EpicGames.Core
 		/// </summary>
 		ParmFlags = Parm | OutParm | ReturnParm | ReferenceParm | ConstParm | RequiredParm,
 
+		/// <summary>
+		/// Flags to keep in the inner type when propagating container flags 
+		/// </summary>
+		PropagateKeepInInner = NonNullable,
+
 		/// <summary>
 		/// Flags that are propagated to properties inside array container
 		/// </summary>
diff --git a/Engine/Source/Programs/Shared/EpicGames.UHT/Parsers/UhtPropertyParser.cs b/Engine/Source/Programs/Shared/EpicGames.UHT/Parsers/UhtPropertyParser.cs
--- a/Engine/Source/Programs/Shared/EpicGames.UHT/Parsers/UhtPropertyParser.cs
+++ b/Engine/Source/Programs/Shared/EpicGames.UHT/Parsers/UhtPropertyParser.cs
@@ -1148,28 +1148,6 @@ namespace EpicGames.UHT.Parsers
 			return inner;
 		}
 
-		[UhtPropertyType(Keyword = "TNonNullPtr")]
-		[SuppressMessage("CodeQuality", "IDE0051:Remove unused private members", Justification = "Attribute accessed method")]
-		[SuppressMessage("Style", "IDE0060:Remove unused parameter", Justification = "Attribute accessed method")]
-		private static UhtProperty? TNonNullPtrType(UhtPropertyResolveArgs args)
-		{
-			UhtClass? propertyClass = args.ParseTemplateObject(UhtTemplateObjectMode.Normal);
-			if (propertyClass == null)
-			{
-				return null;
-			}
-
-			if (propertyClass.IsChildOf(propertyClass.Session.UClass))
-			{
-				// UObject specifies that there is no limiter
-				return new UhtClassProperty(args.PropertySettings, propertyClass, propertyClass.Session.UObject);
-			}
-			else
-			{
-				return new UhtObjectProperty(args.PropertySettings, propertyClass);
-			}
-		}
-
 		[UhtPropertyType(Keyword = "verse::TPtr")]
 		[SuppressMessage("CodeQuality", "IDE0051:Remove unused private members", Justification = "Attribute accessed method")]
 		[SuppressMessage("Style", "IDE0060:Remove unused parameter", Justification = "Attribute accessed method")]
@@ -1293,11 +1271,11 @@ namespace EpicGames.UHT.Parsers
 				else if (classObj.IsChildOf(args.Session.UClass))
 				{
 					// UObject here specifies that there is no limiter
-					property = new UhtClassProperty(propertySettings, classObj, args.Session.UObject);
+					property = new UhtClassProperty(propertySettings, UhtObjectCppForm.Native, classObj, args.Session.UObject);
 				}
 				else
 				{
-					property = new UhtObjectProperty(propertySettings, classObj);
+					property = new UhtObjectProperty(propertySettings, UhtObjectCppForm.Native, classObj);
 				}
 			}
 			else
diff --git a/Engine/Source/Programs/Shared/EpicGames.UHT/Types/Properties/UhtArrayProperty.cs b/Engine/Source/Programs/Shared/EpicGames.UHT/Types/Properties/UhtArrayProperty.cs
--- a/Engine/Source/Programs/Shared/EpicGames.UHT/Types/Properties/UhtArrayProperty.cs
+++ b/Engine/Source/Programs/Shared/EpicGames.UHT/Types/Properties/UhtArrayProperty.cs
@@ -53,7 +53,7 @@ namespace EpicGames.UHT.Types
 
 			ValueProperty.SourceName = SourceName;
 			ValueProperty.EngineName = EngineName;
-			ValueProperty.PropertyFlags = PropertyFlags & EPropertyFlags.PropagateToArrayInner;
+			ValueProperty.PropertyFlags = (ValueProperty.PropertyFlags & EPropertyFlags.PropagateKeepInInner) | (PropertyFlags & EPropertyFlags.PropagateToArrayInner);
 			ValueProperty.Outer = this;
 			ValueProperty.MetaData.Clear();
 		}
@@ -117,7 +117,11 @@ namespace EpicGames.UHT.Types
 					break;
 
 				case UhtPropertyTextType.FunctionThunkParameterArgType:
-					builder.AppendFunctionThunkParameterArrayType(ValueProperty);
+					builder.AppendFunctionThunkParameterArrayType(ValueProperty, true);
+					break;
+
+				case UhtPropertyTextType.VerseMangledType:
+					builder.Append("[]").AppendPropertyText(ValueProperty, textType, true);
 					break;
 
 				default:
diff --git a/Engine/Source/Programs/Shared/EpicGames.UHT/Types/Properties/UhtClassProperty.cs b/Engine/Source/Programs/Shared/EpicGames.UHT/Types/Properties/UhtClassProperty.cs
--- a/Engine/Source/Programs/Shared/EpicGames.UHT/Types/Properties/UhtClassProperty.cs
+++ b/Engine/Source/Programs/Shared/EpicGames.UHT/Types/Properties/UhtClassProperty.cs
@@ -1,6 +1,5 @@
 // Copyright Epic Games, Inc. All Rights Reserved.
 
-using System.Diagnostics.CodeAnalysis;
 using System.Text;
 using EpicGames.Core;
 using EpicGames.UHT.Tables;
@@ -19,23 +18,17 @@ namespace EpicGames.UHT.Types
 		/// <inheritdoc/>
 		public override string EngineClassName => "ClassProperty";
 
-		/// <inheritdoc/>
-		protected override bool PGetPassAsNoPtr => PropertyFlags.HasAnyFlags(EPropertyFlags.UObjectWrapper);
-
-		/// <inheritdoc/>
-		protected override UhtPGetArgumentType PGetTypeArgument => UhtPGetArgumentType.TypeText;
-
 		/// <summary>
 		/// Construct a new property
 		/// </summary>
 		/// <param name="propertySettings">Property settings</param>
+		/// <param name="cppForm">Source code form of the property</param>
 		/// <param name="classObj">Referenced class</param>
 		/// <param name="metaClass">Reference meta class</param>
 		/// <param name="extraFlags">Extra flags to apply to the property.</param>
-		public UhtClassProperty(UhtPropertySettings propertySettings, UhtClass classObj, UhtClass metaClass, EPropertyFlags extraFlags = EPropertyFlags.None)
-			: base(propertySettings, classObj, metaClass)
+		public UhtClassProperty(UhtPropertySettings propertySettings, UhtObjectCppForm cppForm, UhtClass classObj, UhtClass metaClass, EPropertyFlags extraFlags = EPropertyFlags.None)
+			: base(propertySettings, cppForm, classObj, metaClass, extraFlags)
 		{
-			PropertyFlags |= extraFlags;
 			PropertyCaps |= UhtPropertyCaps.CanHaveConfig;
 			PropertyCaps &= ~(UhtPropertyCaps.CanBeInstanced);
 		}
@@ -46,27 +39,6 @@ namespace EpicGames.UHT.Types
 			return MetaClass != null ? $"class {MetaClass.SourceName};" : null;
 		}
 
-		/// <inheritdoc/>
-		private StringBuilder AppendSubClassText(StringBuilder builder)
-		{
-			builder.Append("TSubclassOf<").Append(MetaClass?.SourceName).Append('>');
-			return builder;
-		}
-
-		/// <inheritdoc/>
-		private StringBuilder AppendText(StringBuilder builder)
-		{
-			if (PropertyFlags.HasAnyFlags(EPropertyFlags.UObjectWrapper))
-			{
-				AppendSubClassText(builder);
-			}
-			else
-			{
-				builder.Append(Class.SourceName).Append('*');
-			}
-			return builder;
-		}
-
 		/// <inheritdoc/>
 		public override StringBuilder AppendText(StringBuilder builder, UhtPropertyTextType textType, bool isTemplateArgument)
 		{
@@ -82,43 +54,42 @@ namespace EpicGames.UHT.Types
 				case UhtPropertyTextType.InterfaceFunctionArgOrRetVal:
 				case UhtPropertyTextType.ExportMember:
 				case UhtPropertyTextType.Construction:
-				case UhtPropertyTextType.FunctionThunkParameterArrayType:
 				case UhtPropertyTextType.RigVMType:
 				case UhtPropertyTextType.GetterSetterArg:
-					AppendText(builder);
+					AppendTemplateType(builder);
 					break;
 
-				case UhtPropertyTextType.FunctionThunkRetVal:
-					if (PropertyFlags.HasAnyFlags(EPropertyFlags.ConstParm))
-					{
-						builder.Append("const ");
-					}
-					if (PropertyFlags.HasAnyFlags(EPropertyFlags.UObjectWrapper))
+				case UhtPropertyTextType.EventParameterMember:
+				case UhtPropertyTextType.EventParameterFunctionMember:
+					if (CppForm != UhtObjectCppForm.Native)
 					{
-						AppendText(builder);
+						AppendTemplateType(builder);
 					}
-					else
+					else 
 					{
 						builder.Append("UClass*");
 					}
 					break;
 
-				case UhtPropertyTextType.EventParameterMember:
-				case UhtPropertyTextType.EventParameterFunctionMember:
-					if (!PropertyFlags.HasAnyFlags(EPropertyFlags.UObjectWrapper))
+				case UhtPropertyTextType.FunctionThunkRetVal:
+					if (PropertyFlags.HasAnyFlags(EPropertyFlags.ConstParm))
 					{
-						builder.Append("UClass*");
+						builder.Append("const ");
+					}
+					if (CppForm != UhtObjectCppForm.Native)
+					{
+						AppendTemplateType(builder);
 					}
 					else
 					{
-						AppendText(builder);
+						builder.Append("UClass*");
 					}
 					break;
 
 				case UhtPropertyTextType.FunctionThunkParameterArgType:
-					if (PropertyFlags.HasAllFlags(EPropertyFlags.OutParm | EPropertyFlags.UObjectWrapper))
+					if ((PropertyFlags.HasAnyFlags(EPropertyFlags.OutParm) && CppForm != UhtObjectCppForm.Native) || isTemplateArgument)
 					{
-						AppendSubClassText(builder);
+						AppendTemplateType(builder);
 					}
 					else
 					{
@@ -126,8 +97,12 @@ namespace EpicGames.UHT.Types
 					}
 					break;
 
+				case UhtPropertyTextType.FunctionThunkParameterArrayType:
+					AppendTemplateType(builder);
+					break;
+
 				case UhtPropertyTextType.VerseMangledType:
-					AppendVerseMangledType(builder, Class);
+					AppendVerseMangledType(builder, MetaClass != null && MetaClass.IsVerseField ? MetaClass : Class);
 					break;
 			}
 			return builder;
@@ -142,28 +117,27 @@ namespace EpicGames.UHT.Types
 		/// <inheritdoc/>
 		public override StringBuilder AppendMemberDef(StringBuilder builder, IUhtPropertyMemberContext context, string name, string nameSuffix, string? offset, int tabs)
 		{
-			AppendMemberDefStart(builder, context, name, nameSuffix, offset, tabs, "FClassPropertyParams", "UECodeGen_Private::EPropertyGenFlags::Class");
+			string paramsGenFlags = CppForm == UhtObjectCppForm.TObjectPtr
+				? "UECodeGen_Private::EPropertyGenFlags::Class | UECodeGen_Private::EPropertyGenFlags::ObjectPtr"
+				: "UECodeGen_Private::EPropertyGenFlags::Class";
+			AppendMemberDefStart(builder, context, name, nameSuffix, offset, tabs, "FClassPropertyParams", paramsGenFlags);
 			AppendMemberDefRef(builder, context, Class, false);
 			AppendMemberDefRef(builder, context, MetaClass, false);
 			AppendMemberDefEnd(builder, context, name, nameSuffix);
 			return builder;
 		}
 
-		#region Keyword
-		[UhtPropertyType(Keyword = "TSubclassOf")]
-		[SuppressMessage("CodeQuality", "IDE0051:Remove unused private members", Justification = "Attribute accessed method")]
-		[SuppressMessage("Style", "IDE0060:Remove unused parameter", Justification = "Attribute accessed method")]
-		private static UhtProperty? SubclassOfProperty(UhtPropertyResolveArgs args)
+		/// <inheritdoc/>
+		public override void Validate(UhtStruct outerStruct, UhtProperty outermostProperty, UhtValidationOptions options)
 		{
-			UhtClass? metaClass = args.ParseTemplateClass();
-			if (metaClass == null)
+			base.Validate(outerStruct, outermostProperty, options);
+
+			// UFunctions with a smart pointer as input parameter wont compile anyway, because of missing P_GET_... macro.
+			// UFunctions with a smart pointer as return type will crash when called via blueprint, because they are not supported in VM.
+			if (!options.HasAnyFlags(UhtValidationOptions.IsKey) && PropertyCategory != UhtPropertyCategory.Member && CppForm == UhtObjectCppForm.TObjectPtr)
 			{
-				return null;
+				outerStruct.LogError("UFunctions cannot take a TObjectPtr as a function parameter or return value.");
 			}
-
-			// With TSubclassOf, MetaClass is used as a class limiter.  
-			return new UhtClassProperty(args.PropertySettings, metaClass.Session.UClass, metaClass, EPropertyFlags.UObjectWrapper);
 		}
-		#endregion
 	}
 }
diff --git a/Engine/Source/Programs/Shared/EpicGames.UHT/Types/Properties/UhtClassPtrProperties.cs b/Engine/Source/Programs/Shared/EpicGames.UHT/Types/Properties/UhtClassPtrProperties.cs
deleted file mode 100644
diff --git a/Engine/Source/Programs/Shared/EpicGames.UHT/Types/Properties/UhtInterfaceProperty.cs b/Engine/Source/Programs/Shared/EpicGames.UHT/Types/Properties/UhtInterfaceProperty.cs
--- a/Engine/Source/Programs/Shared/EpicGames.UHT/Types/Properties/UhtInterfaceProperty.cs
+++ b/Engine/Source/Programs/Shared/EpicGames.UHT/Types/Properties/UhtInterfaceProperty.cs
@@ -198,7 +198,7 @@ namespace EpicGames.UHT.Types
 			{
 				return new UhtInterfaceProperty(args.PropertySettings, propertyClass);
 			}
-			return new UhtObjectProperty(args.PropertySettings, propertyClass, null, EPropertyFlags.UObjectWrapper);
+			return new UhtObjectProperty(args.PropertySettings, UhtObjectCppForm.Native, propertyClass, EPropertyFlags.UObjectWrapper);
 		}
 		#endregion
 	}
diff --git a/Engine/Source/Programs/Shared/EpicGames.UHT/Types/Properties/UhtLazyObjectPtrProperty.cs b/Engine/Source/Programs/Shared/EpicGames.UHT/Types/Properties/UhtLazyObjectPtrProperty.cs
--- a/Engine/Source/Programs/Shared/EpicGames.UHT/Types/Properties/UhtLazyObjectPtrProperty.cs
+++ b/Engine/Source/Programs/Shared/EpicGames.UHT/Types/Properties/UhtLazyObjectPtrProperty.cs
@@ -38,7 +38,7 @@ namespace EpicGames.UHT.Types
 		/// <param name="propertySettings">Property settings</param>
 		/// <param name="classObj">Referenced class</param>
 		public UhtLazyObjectPtrProperty(UhtPropertySettings propertySettings, UhtClass classObj)
-			: base(propertySettings, classObj, null)
+			: base(propertySettings, UhtObjectCppForm.Native, classObj, null)
 		{
 			PropertyFlags |= EPropertyFlags.UObjectWrapper;
 			PropertyCaps |= UhtPropertyCaps.PassCppArgsByRef | UhtPropertyCaps.RequiresNullConstructorArg;
diff --git a/Engine/Source/Programs/Shared/EpicGames.UHT/Types/Properties/UhtMapProperty.cs b/Engine/Source/Programs/Shared/EpicGames.UHT/Types/Properties/UhtMapProperty.cs
--- a/Engine/Source/Programs/Shared/EpicGames.UHT/Types/Properties/UhtMapProperty.cs
+++ b/Engine/Source/Programs/Shared/EpicGames.UHT/Types/Properties/UhtMapProperty.cs
@@ -60,12 +60,12 @@ namespace EpicGames.UHT.Types
 			PropertyFlags = ValueProperty.PropertyFlags;
 			ValueProperty.SourceName = SourceName;
 			ValueProperty.EngineName = EngineName;
-			ValueProperty.PropertyFlags &= EPropertyFlags.PropagateToMapValue;
+			ValueProperty.PropertyFlags &= EPropertyFlags.PropagateKeepInInner | EPropertyFlags.PropagateToMapValue;
 			ValueProperty.Outer = this;
 			ValueProperty.MetaData.Clear();
 			KeyProperty.SourceName = $"{SourceName}_Key";
 			KeyProperty.EngineName = $"{EngineName}_Key";
-			KeyProperty.PropertyFlags &= EPropertyFlags.PropagateToMapKey;
+			KeyProperty.PropertyFlags &= EPropertyFlags.PropagateKeepInInner | EPropertyFlags.PropagateToMapKey;
 			KeyProperty.Outer = this;
 			KeyProperty.MetaData.Clear();
 
@@ -90,7 +90,7 @@ namespace EpicGames.UHT.Types
 					KeyProperty.PropertyFlags |= newFlags;
 					MetaData.Add(ValueProperty.MetaData);
 					ValueProperty.MetaData.Clear();
-					ValueProperty.PropertyFlags = PropertyFlags & EPropertyFlags.PropagateToMapValue;
+					ValueProperty.PropertyFlags = (ValueProperty.PropertyFlags & EPropertyFlags.PropagateKeepInInner) | (PropertyFlags & EPropertyFlags.PropagateToMapValue);
 
 					UpdateCaps();
 					PropagateFlagsFromInnerAndHandlePersistentInstanceMetadata(this, null, KeyProperty);
@@ -146,7 +146,7 @@ namespace EpicGames.UHT.Types
 					break;
 
 				case UhtPropertyTextType.FunctionThunkParameterArgType:
-					builder.AppendFunctionThunkParameterArrayType(KeyProperty).Append(',').AppendFunctionThunkParameterArrayType(ValueProperty);
+					builder.AppendFunctionThunkParameterArrayType(KeyProperty, true).Append(',').AppendFunctionThunkParameterArrayType(ValueProperty, true);
 					break;
 
 				case UhtPropertyTextType.VerseMangledType:
diff --git a/Engine/Source/Programs/Shared/EpicGames.UHT/Types/Properties/UhtObjectProperty.cs b/Engine/Source/Programs/Shared/EpicGames.UHT/Types/Properties/UhtObjectProperty.cs
--- a/Engine/Source/Programs/Shared/EpicGames.UHT/Types/Properties/UhtObjectProperty.cs
+++ b/Engine/Source/Programs/Shared/EpicGames.UHT/Types/Properties/UhtObjectProperty.cs
@@ -1,5 +1,6 @@
 // Copyright Epic Games, Inc. All Rights Reserved.
 
+using System;
 using System.Text;
 using EpicGames.Core;
 using EpicGames.UHT.Tables;
@@ -16,20 +17,29 @@ namespace EpicGames.UHT.Types
 		/// <inheritdoc/>
 		public override string EngineClassName => "ObjectProperty";
 
-		/// <inheritdoc/>
-		protected override UhtPGetArgumentType PGetTypeArgument => UhtPGetArgumentType.TypeText;
+		/// <summary>
+		/// Construct a new property
+		/// </summary>
+		/// <param name="propertySettings">Property settings</param>
+		/// <param name="cppForm">Source code form of the property</param>
+		/// <param name="classObj">Referenced class</param>
+		/// <param name="extraFlags">Extra flags to add to the property</param>
+		public UhtObjectProperty(UhtPropertySettings propertySettings, UhtObjectCppForm cppForm, UhtClass classObj, EPropertyFlags extraFlags = EPropertyFlags.None)
+			: this(propertySettings, cppForm, classObj, null, extraFlags)
+		{
+		}
 
 		/// <summary>
 		/// Construct a new property
 		/// </summary>
 		/// <param name="propertySettings">Property settings</param>
+		/// <param name="cppForm">Source code form of the property</param>
 		/// <param name="classObj">Referenced class</param>
 		/// <param name="metaClass">Optional reference class (used by class properties)</param>
 		/// <param name="extraFlags">Extra flags to add to the property</param>
-		public UhtObjectProperty(UhtPropertySettings propertySettings, UhtClass classObj, UhtClass? metaClass = null, EPropertyFlags extraFlags = EPropertyFlags.None)
-			: base(propertySettings, classObj, metaClass)
+		protected UhtObjectProperty(UhtPropertySettings propertySettings, UhtObjectCppForm cppForm, UhtClass classObj, UhtClass? metaClass = null, EPropertyFlags extraFlags = EPropertyFlags.None)
+			: base(propertySettings, cppForm, classObj, metaClass, extraFlags)
 		{
-			PropertyFlags |= extraFlags;
 			PropertyCaps |= UhtPropertyCaps.RequiresNullConstructorArg | UhtPropertyCaps.CanBeInstanced | UhtPropertyCaps.CanExposeOnSpawn |
 				UhtPropertyCaps.IsParameterSupportedByBlueprint | UhtPropertyCaps.IsMemberSupportedByBlueprint | UhtPropertyCaps.SupportsRigVM;
 		}
@@ -62,16 +72,49 @@ namespace EpicGames.UHT.Types
 		{
 			switch (textType)
 			{
+				case UhtPropertyTextType.GetterSetterArg:
+					if (isTemplateArgument)
+					{
+						AppendTemplateType(builder);
+					}
+					else
+					{
+						builder.Append(Class.SourceName).Append('*');
+					}
+					break;
+
 				case UhtPropertyTextType.FunctionThunkRetVal:
-					if (PropertyFlags.HasAnyFlags(EPropertyFlags.ConstParm))
+					if (isTemplateArgument)
 					{
-						builder.Append("const ");
+						AppendTemplateType(builder);
 					}
-					builder.Append(Class.SourceName).Append('*');
+					else
+					{
+						if (PropertyFlags.HasAnyFlags(EPropertyFlags.ConstParm))
+						{
+							builder.Append("const ");
+						}
+						AppendTemplateType(builder);
+					}
+					break;
+
+				case UhtPropertyTextType.FunctionThunkParameterArrayType:
+					AppendTemplateType(builder);
 					break;
 
 				case UhtPropertyTextType.FunctionThunkParameterArgType:
-					builder.Append(Class.SourceName);
+					if (isTemplateArgument)
+					{
+						AppendTemplateType(builder);
+					}
+					else if (CppForm == UhtObjectCppForm.TInterfaceInstance)
+					{
+						builder.AppendClassSourceNameOrInterfaceProxyName(Class);
+					}
+					else
+					{
+						builder.Append(Class.SourceName);
+					}
 					break;
 
 				case UhtPropertyTextType.VerseMangledType:
@@ -79,7 +122,7 @@ namespace EpicGames.UHT.Types
 					break;
 
 				default:
-					builder.Append(Class.SourceName).Append('*');
+					AppendTemplateType(builder);
 					break;
 			}
 			return builder;
@@ -94,7 +137,10 @@ namespace EpicGames.UHT.Types
 		/// <inheritdoc/>
 		public override StringBuilder AppendMemberDef(StringBuilder builder, IUhtPropertyMemberContext context, string name, string nameSuffix, string? offset, int tabs)
 		{
-			AppendMemberDefStart(builder, context, name, nameSuffix, offset, tabs, "FObjectPropertyParams", "UECodeGen_Private::EPropertyGenFlags::Object");
+			string paramsGenFlags = CppForm == UhtObjectCppForm.TObjectPtr
+				? "UECodeGen_Private::EPropertyGenFlags::Object | UECodeGen_Private::EPropertyGenFlags::ObjectPtr"
+				: "UECodeGen_Private::EPropertyGenFlags::Object";
+			AppendMemberDefStart(builder, context, name, nameSuffix, offset, tabs, "FObjectPropertyParams", paramsGenFlags);
 			AppendMemberDefRef(builder, context, Class, false);
 			AppendMemberDefEnd(builder, context, name, nameSuffix);
 			return builder;
@@ -112,9 +158,39 @@ namespace EpicGames.UHT.Types
 		{
 			if (other is UhtObjectProperty otherObject)
 			{
-				return Class == otherObject.Class && MetaClass == otherObject.MetaClass;
+				return CppForm.GetSameTypeCppForm() == otherObject.CppForm.GetSameTypeCppForm() && Class == otherObject.Class && MetaClass == otherObject.MetaClass;
 			}
 			return false;
 		}
+
+		/// <inheritdoc/>
+		public override void Validate(UhtStruct outerStruct, UhtProperty outermostProperty, UhtValidationOptions options)
+		{
+			base.Validate(outerStruct, outermostProperty, options);
+
+			// UFunctions with a smart pointer as input parameter wont compile anyway, because of missing P_GET_... macro.
+			// UFunctions with a smart pointer as return type will crash when called via blueprint, because they are not supported in VM.
+			if (PropertyCategory != UhtPropertyCategory.Member && CppForm == UhtObjectCppForm.TObjectPtr)
+			{
+				// At this point, allow this to appear in TMap keys in the UPlayerMappableInputConfig class
+				if (!options.HasAnyFlags(UhtValidationOptions.IsKey) ||
+					!outerStruct.SourceName.Equals("GetMappingContexts", StringComparison.Ordinal) ||
+					outerStruct.Outer == null ||
+					!outerStruct.Outer.SourceName.Equals("UPlayerMappableInputConfig", StringComparison.Ordinal))
+				{
+					outerStruct.LogError("UFunctions cannot take a TObjectPtr as a function parameter or return value.");
+				}
+			}
+		}
+
+		/// <inheritdoc/>
+		protected override void ValidateMember(UhtStruct structObj, UhtValidationOptions options)
+		{
+			base.ValidateMember(structObj, options);
+			if (Class.NativeInterface != null && CppForm == UhtObjectCppForm.TObjectPtr)
+			{
+				this.LogError($"UPROPERTY pointers cannot be interfaces - did you mean TScriptInterface<{Class.SourceName}>?");
+			}
+		}
 	}
 }
diff --git a/Engine/Source/Programs/Shared/EpicGames.UHT/Types/Properties/UhtObjectPropertyBase.cs b/Engine/Source/Programs/Shared/EpicGames.UHT/Types/Properties/UhtObjectPropertyBase.cs
--- a/Engine/Source/Programs/Shared/EpicGames.UHT/Types/Properties/UhtObjectPropertyBase.cs
+++ b/Engine/Source/Programs/Shared/EpicGames.UHT/Types/Properties/UhtObjectPropertyBase.cs
@@ -23,6 +23,21 @@ namespace EpicGames.UHT.Types
 		/// </summary>
 		Native,
 
+		/// <summary>
+		/// TObjectPtr
+		/// </summary>
+		TObjectPtr,
+
+		/// <summary>
+		/// TNonNullPtr
+		/// </summary>
+		TNonNullPtr,
+
+		/// <summary>
+		/// TSubclassOf
+		/// </summary>
+		TSubclassOf,
+
 		/// <summary>
 		/// verse::type
 		/// </summary>
@@ -49,9 +64,214 @@ namespace EpicGames.UHT.Types
 		TInterfaceInstance,
 	}
 
+	/// <summary>
+	/// Collection of extension methods for UhtObjectCppForm
+	/// </summary>
+	public static class UhtObjectCppFormExtensions
+	{
+
+		/// <summary>
+		/// Return a CppForm used when doing IsSameType checks
+		/// </summary>
+		/// <param name="cppForm">The form in question</param>
+		/// <returns>CppForm to be used in same type checks</returns>
+		public static UhtObjectCppForm GetSameTypeCppForm(this UhtObjectCppForm cppForm)
+		{
+			switch (cppForm)
+			{
+				case UhtObjectCppForm.Native:
+				case UhtObjectCppForm.TObjectPtr:
+					return UhtObjectCppForm.Native;
+
+				case UhtObjectCppForm.TInterfaceInstance:
+				case UhtObjectCppForm.TNonNullPtr:
+				case UhtObjectCppForm.TSubclassOf:
+				case UhtObjectCppForm.VerseType:
+				case UhtObjectCppForm.VerseCastableType:
+				case UhtObjectCppForm.VerseSubtype:
+				case UhtObjectCppForm.VerseCastableSubtype:
+					return cppForm;
+
+				default:
+					throw new NotImplementedException();
+			}
+		}
+
+		/// <summary>
+		/// Returns true if this CppForm needs a verse where clause
+		/// </summary>
+		/// <param name="cppForm">The form in question</param>
+		/// <returns>True if the verse name needs a where clause</returns>
+		public static bool NeedsWhereClause(this UhtObjectCppForm cppForm)
+		{
+			switch (cppForm)
+			{
+				case UhtObjectCppForm.Native:
+				case UhtObjectCppForm.TObjectPtr:
+				case UhtObjectCppForm.TInterfaceInstance:
+				case UhtObjectCppForm.TNonNullPtr:
+				case UhtObjectCppForm.TSubclassOf:
+					return false;
+
+				case UhtObjectCppForm.VerseType:
+				case UhtObjectCppForm.VerseCastableType:
+				case UhtObjectCppForm.VerseSubtype:
+				case UhtObjectCppForm.VerseCastableSubtype:
+					return true;
+
+				default:
+					throw new NotImplementedException();
+			}
+		}
+
+		/// <summary>
+		/// Return the P_GET macro text based on the cpp form
+		/// </summary>
+		/// <param name="cppForm">The form in question</param>
+		/// <param name="nativeText">Text for the native form</param>
+		/// <returns>P_GET macro text</returns>
+		/// <exception cref="NotImplementedException"></exception>
+		public static string GetPGetMacroText(this UhtObjectCppForm cppForm, string nativeText)
+		{
+			switch (cppForm)
+			{
+				case UhtObjectCppForm.Native:
+					return nativeText;
+
+				case UhtObjectCppForm.TObjectPtr:
+					return "OBJECTPTR";
+
+				case UhtObjectCppForm.TNonNullPtr:
+					return "NONNULLPTROBJECT";
+
+				case UhtObjectCppForm.TSubclassOf:
+					return "OBJECT";
+
+				case UhtObjectCppForm.VerseType:
+					return "VERSETYPE";
+
+				case UhtObjectCppForm.VerseCastableType:
+					return "VERSECASTABLETYPE";
+
+				case UhtObjectCppForm.VerseSubtype:
+					return "VERSESUBTYPE";
+
+				case UhtObjectCppForm.VerseCastableSubtype:
+					return "VERSECASTABLESUBTYPE";
+
+				case UhtObjectCppForm.TInterfaceInstance:
+					return "INTERFACEINSTANCE";
+
+				default:
+					throw new NotImplementedException();
+			}
+		}
+
+		/// <summary>
+		/// Returns the PGET argument type for the given CppForm
+		/// </summary>
+		/// <param name="cppForm">The form in question</param>
+		/// <returns></returns>
+		/// <exception cref="NotImplementedException"></exception>
+		public static UhtPGetArgumentType GetPGetArgumentType(this UhtObjectCppForm cppForm)
+		{
+			switch (cppForm)
+			{
+				case UhtObjectCppForm.VerseType:
+				case UhtObjectCppForm.VerseCastableType:
+					return UhtPGetArgumentType.None;
+
+				case UhtObjectCppForm.Native:
+				case UhtObjectCppForm.TObjectPtr:
+				case UhtObjectCppForm.TNonNullPtr:
+				case UhtObjectCppForm.TSubclassOf:
+				case UhtObjectCppForm.VerseSubtype:
+				case UhtObjectCppForm.VerseCastableSubtype:
+				case UhtObjectCppForm.TInterfaceInstance:
+					return UhtPGetArgumentType.TypeText;
+
+				default:
+					throw new NotImplementedException();
+			}
+		}
+
+		/// <summary>
+		/// Return true if this form needs the NO_PTR option
+		/// </summary>
+		/// <param name="cppForm">The form in question</param>
+		/// <returns></returns>
+		/// <exception cref="NotImplementedException"></exception>
+		public static bool GetPGetPassAsNoPtr(this UhtObjectCppForm cppForm)
+		{
+			switch (cppForm)
+			{
+				case UhtObjectCppForm.Native:
+				case UhtObjectCppForm.TObjectPtr:
+				case UhtObjectCppForm.TInterfaceInstance:
+				case UhtObjectCppForm.TNonNullPtr:
+				case UhtObjectCppForm.VerseType:
+				case UhtObjectCppForm.VerseCastableType:
+				case UhtObjectCppForm.VerseSubtype:
+				case UhtObjectCppForm.VerseCastableSubtype:
+					return false;
+
+				case UhtObjectCppForm.TSubclassOf:
+					return true;
+
+				default:
+					throw new NotImplementedException();
+			}
+		}
+
+		/// <summary>
+		/// Based on the cpp form, update the property settings
+		/// </summary>
+		/// <param name="cppForm">The form in question</param>
+		/// <param name="property">Property to be updated</param>
+		public static void InitializeProperty(this UhtObjectCppForm cppForm, UhtObjectPropertyBase property)
+		{
+			switch (cppForm)
+			{
+				case UhtObjectCppForm.Native:
+					break;
+
+				case UhtObjectCppForm.TInterfaceInstance:
+				case UhtObjectCppForm.VerseType:
+				case UhtObjectCppForm.VerseCastableType:
+				case UhtObjectCppForm.VerseSubtype:
+				case UhtObjectCppForm.VerseCastableSubtype:
+					// TODO: struct FInterfaceTraits::"ObjectMustImplement"
+					//if (property.Class.ClassType == UhtClassType.Interface)
+					//{
+					//	property.MetaData.Add("MustImplement", property.Class.PathName);
+					//}
+					//property.MetaData.Add("AllowAbstract", "True");
+					break;
+
+				case UhtObjectCppForm.TObjectPtr:
+					property.PropertyFlags |= EPropertyFlags.TObjectPtr | EPropertyFlags.UObjectWrapper;
+					property.PropertyCaps |= UhtPropertyCaps.PassCppArgsByRef;
+					break;
+
+				case UhtObjectCppForm.TSubclassOf:
+					property.PropertyFlags |= EPropertyFlags.UObjectWrapper;
+					break;
+
+				case UhtObjectCppForm.TNonNullPtr:
+					property.PropertyFlags |= EPropertyFlags.NonNullable;
+					break;
+
+				default:
+					throw new NotImplementedException();
+			}
+
+		}
+	}
+
 	/// <summary>
 	/// FObjectPropertyBase
 	/// </summary>
+	[UnrealHeaderTool]
 	[UhtEngineClass(Name = "ObjectPropertyBase", IsProperty = true)]
 	public abstract class UhtObjectPropertyBase : UhtProperty
 	{
@@ -62,7 +282,13 @@ namespace EpicGames.UHT.Types
 		protected override string CppTypeText => "Object";
 
 		/// <inheritdoc/>
-		protected override string PGetMacroText => "OBJECT";
+		protected override string PGetMacroText => CppForm.GetPGetMacroText("OBJECT");
+
+		/// <inheritdoc/>
+		protected override UhtPGetArgumentType PGetTypeArgument => CppForm.GetPGetArgumentType();
+
+		/// <inheritdoc/>
+		protected override bool PGetPassAsNoPtr => CppForm.GetPGetPassAsNoPtr();
 
 		/// <summary>
 		/// Referenced UCLASS
@@ -79,18 +305,25 @@ namespace EpicGames.UHT.Types
 		/// <summary>
 		/// The form that the object property appeared in the code
 		/// </summary>
-		public UhtObjectCppForm CppForm { get; set; } = UhtObjectCppForm.Native;
+		public UhtObjectCppForm CppForm { get; init; } = UhtObjectCppForm.Native;
 
 		/// <summary>
 		/// Construct a property
 		/// </summary>
 		/// <param name="propertySettings">Property settings</param>
-		/// <param name="classObj">Referenced UCLASS</param>
+		/// <param name="cppForm">Source code form of the property</param>
+		/// <param name="classObj">Referenced UCLASS (a.k.a. PropertyClass)</param>
 		/// <param name="metaClass">Referenced UCLASS used by class properties</param>
-		protected UhtObjectPropertyBase(UhtPropertySettings propertySettings, UhtClass classObj, UhtClass? metaClass) : base(propertySettings)
+		/// <param name="extraFlags">Extra flags to apply to the property.</param>
+		protected UhtObjectPropertyBase(UhtPropertySettings propertySettings, UhtObjectCppForm cppForm, UhtClass classObj, UhtClass? metaClass, EPropertyFlags extraFlags = EPropertyFlags.None) : base(propertySettings)
 		{
+			CppForm = cppForm;
 			Class = classObj;
 			MetaClass = metaClass;
+			PropertyFlags |= extraFlags;
+
+			// Based on the form, perform more initialization
+			CppForm.InitializeProperty(this);
 
 			// This applies to EVERYTHING including raw pointer
 			// Imply const if it's a parameter that is a pointer to a const class
@@ -186,7 +419,7 @@ namespace EpicGames.UHT.Types
 		{
 			Type type = GetType();
 			return type == typeof(UhtObjectProperty)
-				|| (type == typeof(UhtClassProperty) && !PropertyFlags.HasFlag(EPropertyFlags.UObjectWrapper));
+				|| (type == typeof(UhtClassProperty) && CppForm != UhtObjectCppForm.TSubclassOf);
 		}
 
 		/// <summary>
@@ -212,12 +445,185 @@ namespace EpicGames.UHT.Types
 					builder.Append("castable_subtype ").AppendVersePath(type);
 					break;
 
+				case UhtObjectCppForm.Native:
+				case UhtObjectCppForm.TObjectPtr:
+				case UhtObjectCppForm.TNonNullPtr:
 				case UhtObjectCppForm.TInterfaceInstance:
-				default:
+				case UhtObjectCppForm.TSubclassOf:
 					builder.AppendVersePath(type);
 					break;
+
+				default:
+					throw new NotImplementedException();
 			}
 			return builder;
 		}
+
+		/// <summary>
+		/// Append the template declaration for the given type
+		/// </summary>
+		/// <param name="builder">Destination builder</param>
+		/// <returns></returns>
+		/// <exception cref="NotImplementedException"></exception>
+		protected StringBuilder AppendTemplateType(StringBuilder builder)
+		{
+			switch (CppForm)
+			{
+				case UhtObjectCppForm.Native:
+					builder.Append(Class.SourceName).Append('*');
+					break;
+
+				case UhtObjectCppForm.TObjectPtr:
+					builder.Append("TObjectPtr<").Append(Class.SourceName).Append('>');
+					break;
+
+				case UhtObjectCppForm.TNonNullPtr:
+					builder.Append("TNonNullPtr<").Append(Class.SourceName).Append('>');
+					break;
+
+				case UhtObjectCppForm.TSubclassOf:
+					builder.Append("TSubclassOf<").Append(MetaClass != null ? MetaClass.SourceName : "UClass").Append('>');
+					break;
+
+				case UhtObjectCppForm.VerseType:
+					builder.Append("verse::type");
+					break;
+
+				case UhtObjectCppForm.VerseCastableType:
+					builder.Append("verse::castable_type");
+					break;
+
+				case UhtObjectCppForm.VerseSubtype:
+					builder.Append("verse::subtype<").Append(Class.SourceName).Append('>');
+					break;
+
+				case UhtObjectCppForm.VerseCastableSubtype:
+					builder.Append("verse::castable_subtype<").Append(Class.SourceName).Append('>');
+					break;
+
+				case UhtObjectCppForm.TInterfaceInstance:
+					builder.Append("TInterfaceInstance<").AppendClassSourceNameOrInterfaceProxyName(Class).Append('>');
+					break;
+
+				default:
+					throw new NotImplementedException();
+			}
+			return builder;
+		}
+
+		#region Keywords
+		[UhtPropertyType(Keyword = "TSubclassOf")]
+		[SuppressMessage("CodeQuality", "IDE0051:Remove unused private members", Justification = "Attribute accessed method")]
+		[SuppressMessage("Style", "IDE0060:Remove unused parameter", Justification = "Attribute accessed method")]
+		private static UhtProperty? SubclassOfProperty(UhtPropertyResolveArgs args)
+		{
+			UhtClass? metaClass = args.ParseTemplateClass();
+			if (metaClass == null)
+			{
+				return null;
+			}
+
+			// With TSubclassOf, MetaClass is used as a class limiter.  
+			return new UhtClassProperty(args.PropertySettings, UhtObjectCppForm.TSubclassOf, metaClass.Session.UClass, metaClass);
+		}
+
+		[UhtPropertyType(Keyword = "TNonNullPtr")]
+		[SuppressMessage("CodeQuality", "IDE0051:Remove unused private members", Justification = "Attribute accessed method")]
+		[SuppressMessage("Style", "IDE0060:Remove unused parameter", Justification = "Attribute accessed method")]
+		private static UhtProperty? TNonNullPtrType(UhtPropertyResolveArgs args)
+		{
+			return ObjectTypeProperty(args, UhtObjectCppForm.TNonNullPtr, UhtTemplateObjectMode.Normal);
+		}
+
+		[UhtPropertyType(Keyword = "TObjectPtr")]
+		[SuppressMessage("CodeQuality", "IDE0051:Remove unused private members", Justification = "Attribute accessed method")]
+		[SuppressMessage("Style", "IDE0060:Remove unused parameter", Justification = "Attribute accessed method")]
+		private static UhtProperty? ObjectPtrProperty(UhtPropertyResolveArgs args)
+		{
+			return ObjectTypeProperty(args, UhtObjectCppForm.TObjectPtr, UhtTemplateObjectMode.Normal);
+		}
+
+		[UhtPropertyType(Keyword = "verse::type", Options = UhtPropertyTypeOptions.Simple)]
+		[SuppressMessage("CodeQuality", "IDE0051:Remove unused private members", Justification = "Attribute accessed method")]
+		[SuppressMessage("Style", "IDE0060:Remove unused parameter", Justification = "Attribute accessed method")]
+		private static UhtProperty? VerseTypeProperty(UhtPropertyResolveArgs args)
+		{
+			return new UhtClassProperty(args.PropertySettings, UhtObjectCppForm.VerseType, args.Session.UClass, args.Session.UObject);
+		}
+
+		[UhtPropertyType(Keyword = "verse::castable_type", Options = UhtPropertyTypeOptions.Simple)]
+		[SuppressMessage("CodeQuality", "IDE0051:Remove unused private members", Justification = "Attribute accessed method")]
+		[SuppressMessage("Style", "IDE0060:Remove unused parameter", Justification = "Attribute accessed method")]
+		private static UhtProperty? VerseCastableTypeProperty(UhtPropertyResolveArgs args)
+		{
+			return new UhtClassProperty(args.PropertySettings, UhtObjectCppForm.VerseCastableType, args.Session.UClass, args.Session.UObject);
+		}
+
+		[UhtPropertyType(Keyword = "verse::subtype")]
+		[SuppressMessage("CodeQuality", "IDE0051:Remove unused private members", Justification = "Attribute accessed method")]
+		[SuppressMessage("Style", "IDE0060:Remove unused parameter", Justification = "Attribute accessed method")]
+		private static UhtProperty? VerseSubTypeProperty(UhtPropertyResolveArgs args)
+		{
+			return ObjectTypeProperty(args, UhtObjectCppForm.VerseSubtype, UhtTemplateObjectMode.Normal);
+		}
+
+		[UhtPropertyType(Keyword = "verse::castable_subtype")]
+		[SuppressMessage("CodeQuality", "IDE0051:Remove unused private members", Justification = "Attribute accessed method")]
+		[SuppressMessage("Style", "IDE0060:Remove unused parameter", Justification = "Attribute accessed method")]
+		private static UhtProperty? VerseCastableSubtypeProperty(UhtPropertyResolveArgs args)
+		{
+			return ObjectTypeProperty(args, UhtObjectCppForm.VerseCastableSubtype, UhtTemplateObjectMode.Normal);
+		}
+
+		[UhtPropertyType(Keyword = "TInterfaceInstance")]
+		[SuppressMessage("CodeQuality", "IDE0051:Remove unused private members", Justification = "Attribute accessed method")]
+		[SuppressMessage("Style", "IDE0060:Remove unused parameter", Justification = "Attribute accessed method")]
+		private static UhtProperty? TInterfaceInstanceProperty(UhtPropertyResolveArgs args)
+		{
+			return ObjectTypeProperty(args, UhtObjectCppForm.TInterfaceInstance, UhtTemplateObjectMode.NativeInterfaceProxy);
+		}
+
+		private static UhtProperty? ObjectTypeProperty(UhtPropertyResolveArgs args, UhtObjectCppForm cppForm, UhtTemplateObjectMode templateObjectMode)
+		{
+			int typeStartPos = args.TokenReader.PeekToken().InputStartPos;
+			UhtClass? propertyClass = args.ParseTemplateObject(templateObjectMode);
+			if (propertyClass == null)
+			{
+				return null;
+			}
+
+			if (cppForm == UhtObjectCppForm.TObjectPtr)
+			{
+				args.ConditionalLogPointerUsage(args.Config.EngineObjectPtrMemberBehavior, args.Config.EnginePluginObjectPtrMemberBehavior,
+					args.Config.NonEngineObjectPtrMemberBehavior, "ObjectPtr", typeStartPos, null);
+			}
+
+			switch (cppForm)
+			{
+				case UhtObjectCppForm.Native:
+				case UhtObjectCppForm.TObjectPtr:
+				case UhtObjectCppForm.TNonNullPtr:
+					if (propertyClass.IsChildOf(propertyClass.Session.UClass))
+					{
+						// UObject specifies that there is no limiter
+						return new UhtClassProperty(args.PropertySettings, cppForm, propertyClass, propertyClass.Session.UObject);
+					}
+					else
+					{
+						return new UhtObjectProperty(args.PropertySettings, cppForm, propertyClass);
+					}
+
+				case UhtObjectCppForm.TInterfaceInstance:
+					return new UhtObjectProperty(args.PropertySettings, cppForm, propertyClass);
+
+				case UhtObjectCppForm.VerseSubtype:
+				case UhtObjectCppForm.VerseCastableSubtype:
+					return new UhtClassProperty(args.PropertySettings, cppForm, args.Session.UClass, propertyClass);
+
+				default:
+					throw new NotImplementedException();
+			}
+		}
+		#endregion
 	}
 }
diff --git a/Engine/Source/Programs/Shared/EpicGames.UHT/Types/Properties/UhtObjectPtrProperty.cs b/Engine/Source/Programs/Shared/EpicGames.UHT/Types/Properties/UhtObjectPtrProperty.cs
deleted file mode 100644
diff --git a/Engine/Source/Programs/Shared/EpicGames.UHT/Types/Properties/UhtOptionalProperty.cs b/Engine/Source/Programs/Shared/EpicGames.UHT/Types/Properties/UhtOptionalProperty.cs
--- a/Engine/Source/Programs/Shared/EpicGames.UHT/Types/Properties/UhtOptionalProperty.cs
+++ b/Engine/Source/Programs/Shared/EpicGames.UHT/Types/Properties/UhtOptionalProperty.cs
@@ -53,7 +53,7 @@ namespace EpicGames.UHT.Types
 
 			ValueProperty.SourceName = SourceName;
 			ValueProperty.EngineName = EngineName;
-			ValueProperty.PropertyFlags = PropertyFlags & EPropertyFlags.PropagateToOptionalInner;
+			ValueProperty.PropertyFlags = (ValueProperty.PropertyFlags & EPropertyFlags.PropagateKeepInInner) | (PropertyFlags & EPropertyFlags.PropagateToOptionalInner);
 			ValueProperty.Outer = this;
 			ValueProperty.MetaData.Clear();
 		}
@@ -67,7 +67,7 @@ namespace EpicGames.UHT.Types
 				case UhtResolvePhase.Final:
 					PropertyFlags |= ResolveAndReturnNewFlags(ValueProperty, phase);
 					MetaData.Add(ValueProperty.MetaData);
-					ValueProperty.PropertyFlags = PropertyFlags & EPropertyFlags.PropagateToOptionalInner;
+					ValueProperty.PropertyFlags = (ValueProperty.PropertyFlags & EPropertyFlags.PropagateKeepInInner) | (PropertyFlags & EPropertyFlags.PropagateToOptionalInner);
 					ValueProperty.MetaData.Clear();
 					PropagateFlagsFromInnerAndHandlePersistentInstanceMetadata(this, MetaData, ValueProperty);
 					UpdateCaps();
@@ -117,7 +117,7 @@ namespace EpicGames.UHT.Types
 					break;
 
 				case UhtPropertyTextType.FunctionThunkParameterArgType:
-					builder.AppendFunctionThunkParameterArrayType(ValueProperty);
+					builder.AppendFunctionThunkParameterArrayType(ValueProperty, true);
 					break;
 
 				case UhtPropertyTextType.VerseMangledType:
diff --git a/Engine/Source/Programs/Shared/EpicGames.UHT/Types/Properties/UhtSetProperty.cs b/Engine/Source/Programs/Shared/EpicGames.UHT/Types/Properties/UhtSetProperty.cs
--- a/Engine/Source/Programs/Shared/EpicGames.UHT/Types/Properties/UhtSetProperty.cs
+++ b/Engine/Source/Programs/Shared/EpicGames.UHT/Types/Properties/UhtSetProperty.cs
@@ -46,7 +46,7 @@ namespace EpicGames.UHT.Types
 
 			ValueProperty.SourceName = SourceName;
 			ValueProperty.EngineName = EngineName;
-			ValueProperty.PropertyFlags = PropertyFlags & EPropertyFlags.PropagateToSetElement;
+			ValueProperty.PropertyFlags = (ValueProperty.PropertyFlags & EPropertyFlags.PropagateKeepInInner) | (PropertyFlags & EPropertyFlags.PropagateToSetElement);
 			ValueProperty.Outer = this;
 			ValueProperty.MetaData.Clear();
 		}
@@ -61,7 +61,7 @@ namespace EpicGames.UHT.Types
 					PropertyFlags |= ResolveAndReturnNewFlags(ValueProperty, phase);
 					MetaData.Add(ValueProperty.MetaData);
 					ValueProperty.MetaData.Clear();
-					ValueProperty.PropertyFlags = PropertyFlags & EPropertyFlags.PropagateToSetElement;
+					ValueProperty.PropertyFlags = (ValueProperty.PropertyFlags & EPropertyFlags.PropagateKeepInInner) | (PropertyFlags & EPropertyFlags.PropagateToSetElement);
 					UpdateCaps();
 					PropagateFlagsFromInnerAndHandlePersistentInstanceMetadata(this, MetaData, ValueProperty);
 					break;
@@ -110,7 +110,7 @@ namespace EpicGames.UHT.Types
 					break;
 
 				case UhtPropertyTextType.FunctionThunkParameterArgType:
-					builder.AppendFunctionThunkParameterArrayType(ValueProperty);
+					builder.AppendFunctionThunkParameterArrayType(ValueProperty, true);
 					break;
 
 				default:
diff --git a/Engine/Source/Programs/Shared/EpicGames.UHT/Types/Properties/UhtSoftObjectProperty.cs b/Engine/Source/Programs/Shared/EpicGames.UHT/Types/Properties/UhtSoftObjectProperty.cs
--- a/Engine/Source/Programs/Shared/EpicGames.UHT/Types/Properties/UhtSoftObjectProperty.cs
+++ b/Engine/Source/Programs/Shared/EpicGames.UHT/Types/Properties/UhtSoftObjectProperty.cs
@@ -35,7 +35,7 @@ namespace EpicGames.UHT.Types
 		/// <param name="classObj">UCLASS being referenced</param>
 		/// <param name="metaClass">Optional meta class (used by SoftClassProperty)</param>
 		public UhtSoftObjectProperty(UhtPropertySettings propertySettings, UhtClass classObj, UhtClass? metaClass = null)
-			: base(propertySettings, classObj, metaClass)
+			: base(propertySettings, UhtObjectCppForm.Native, classObj, metaClass)
 		{
 			PropertyFlags |= EPropertyFlags.UObjectWrapper;
 			PropertyCaps |= UhtPropertyCaps.PassCppArgsByRef | UhtPropertyCaps.RequiresNullConstructorArg | UhtPropertyCaps.CanExposeOnSpawn |
diff --git a/Engine/Source/Programs/Shared/EpicGames.UHT/Types/Properties/UhtWeakObjectPtrProperty.cs b/Engine/Source/Programs/Shared/EpicGames.UHT/Types/Properties/UhtWeakObjectPtrProperty.cs
--- a/Engine/Source/Programs/Shared/EpicGames.UHT/Types/Properties/UhtWeakObjectPtrProperty.cs
+++ b/Engine/Source/Programs/Shared/EpicGames.UHT/Types/Properties/UhtWeakObjectPtrProperty.cs
@@ -36,7 +36,7 @@ namespace EpicGames.UHT.Types
 		/// <param name="propertyClass">Class being referenced</param>
 		/// <param name="extraFlags">Extra property flags to add to the definition</param>
 		public UhtWeakObjectPtrProperty(UhtPropertySettings propertySettings, UhtClass propertyClass, EPropertyFlags extraFlags = EPropertyFlags.None)
-			: base(propertySettings, propertyClass, null)
+			: base(propertySettings, UhtObjectCppForm.Native, propertyClass, null)
 		{
 			PropertyFlags |= EPropertyFlags.UObjectWrapper | extraFlags;
 			PropertyCaps |= UhtPropertyCaps.PassCppArgsByRef | UhtPropertyCaps.RequiresNullConstructorArg | UhtPropertyCaps.IsMemberSupportedByBlueprint;
@@ -104,7 +104,7 @@ namespace EpicGames.UHT.Types
 
 			if (propertySettings.DisallowPropertyFlags.HasAnyFlags(EPropertyFlags.AutoWeak))
 			{
-				return new UhtObjectProperty(propertySettings, classObj, null, extraFlags | EPropertyFlags.UObjectWrapper);
+				return new UhtObjectProperty(propertySettings, UhtObjectCppForm.Native, classObj, extraFlags | EPropertyFlags.UObjectWrapper);
 			}
 			return new UhtWeakObjectPtrProperty(propertySettings, classObj, extraFlags);
 		}
diff --git a/Engine/Source/Programs/Shared/EpicGames.UHT/Types/UhtFunction.cs b/Engine/Source/Programs/Shared/EpicGames.UHT/Types/UhtFunction.cs
--- a/Engine/Source/Programs/Shared/EpicGames.UHT/Types/UhtFunction.cs
+++ b/Engine/Source/Programs/Shared/EpicGames.UHT/Types/UhtFunction.cs
@@ -434,7 +434,7 @@ namespace EpicGames.UHT.Types
 						isFirst = false;
 						if (!needsWhereClause)
 						{
-							needsWhereClause = type is UhtObjectPropertyBase objectProperty && (objectProperty.CppForm != UhtObjectCppForm.Native && objectProperty.CppForm != UhtObjectCppForm.TInterfaceInstance);
+							needsWhereClause = type is UhtObjectPropertyBase objectProperty && objectProperty.CppForm.NeedsWhereClause();
 						}
 					}
 				}
@@ -445,7 +445,7 @@ namespace EpicGames.UHT.Types
 					builder.Append(" where");
 					foreach (UhtType type in parameterProperties.Span)
 					{
-						if (type is UhtObjectPropertyBase objectProperty && (objectProperty.CppForm != UhtObjectCppForm.Native && objectProperty.CppForm != UhtObjectCppForm.TInterfaceInstance))
+						if (type is UhtObjectPropertyBase objectProperty && objectProperty.CppForm.NeedsWhereClause())
 						{
 							if (!isFirst)
 							{
diff --git a/Engine/Source/Programs/Shared/EpicGames.UHT/Types/UhtProperty.cs b/Engine/Source/Programs/Shared/EpicGames.UHT/Types/UhtProperty.cs
--- a/Engine/Source/Programs/Shared/EpicGames.UHT/Types/UhtProperty.cs
+++ b/Engine/Source/Programs/Shared/EpicGames.UHT/Types/UhtProperty.cs
@@ -97,11 +97,6 @@ namespace EpicGames.UHT.Types
 		/// If true, the type is a TPtr
 		/// </summary>
 		TPtr = 0x00008000,
-
-		/// <summary>
-		/// If true, the type is a TPtr
-		/// </summary>
-		TNonNullPtr = 0x00010000,
 	};
 
 	/// <summary>
@@ -2289,10 +2284,11 @@ namespace EpicGames.UHT.Types
 		/// </summary>
 		/// <param name="builder">Destination builder</param>
 		/// <param name="property">Property in question</param>
+		/// <param name="isTemplateArgument">If true, this is a template argument instead of the top level property</param>
 		/// <returns>Destination builder</returns>
-		public static StringBuilder AppendFunctionThunkParameterArrayType(this StringBuilder builder, UhtProperty property)
+		public static StringBuilder AppendFunctionThunkParameterArrayType(this StringBuilder builder, UhtProperty property, bool isTemplateArgument = false)
 		{
-			return property.AppendText(builder, UhtPropertyTextType.FunctionThunkParameterArrayType);
+			return property.AppendText(builder, UhtPropertyTextType.FunctionThunkParameterArrayType, isTemplateArgument);
 		}
 
 		/// <summary>
