[UBA]
* Added support for setting compression level from UBT command line. (Example: -UBACompressionLevel=Optimal1)

diff --git a/Engine/Source/Programs/Shared/EpicGames.UBA/IConfig.cs b/Engine/Source/Programs/Shared/EpicGames.UBA/IConfig.cs
--- a/Engine/Source/Programs/Shared/EpicGames.UBA/IConfig.cs
+++ b/Engine/Source/Programs/Shared/EpicGames.UBA/IConfig.cs
@@ -18,5 +18,13 @@ namespace EpicGames.UBA
 		{
 			return new ConfigImpl(configFile);
 		}
+
+		/// <summary>
+		/// Load a config file
+		/// </summary>
+		/// <param name="table">The name of the table where the property is added</param>
+		/// <param name="name">The name of the property</param>
+		/// <param name="value">The value of the property</param>
+		void AddValue(string table, string name, string value);
 	}
 }
\ No newline at end of file
diff --git a/Engine/Source/Programs/Shared/EpicGames.UBA/Impl/ConfigImpl.cs b/Engine/Source/Programs/Shared/EpicGames.UBA/Impl/ConfigImpl.cs
--- a/Engine/Source/Programs/Shared/EpicGames.UBA/Impl/ConfigImpl.cs
+++ b/Engine/Source/Programs/Shared/EpicGames.UBA/Impl/ConfigImpl.cs
@@ -1,5 +1,6 @@
 // Copyright Epic Games, Inc. All Rights Reserved.
 
+using System;
 using System.Runtime.InteropServices;
 
 namespace EpicGames.UBA.Impl
@@ -8,12 +9,24 @@ namespace EpicGames.UBA.Impl
 	{
 		#region DllImport
 		[DllImport("UbaHost", CharSet = CharSet.Auto)]
-		static extern bool Config_Load(string configFile);
+		static extern nint Config_Load(string configFile);
+		[DllImport("UbaHost", CharSet = CharSet.Auto)]
+		static extern nint Config_AddTable(nint config, string name);
+		[DllImport("UbaHost", CharSet = CharSet.Auto)]
+		static extern nint ConfigTable_AddValueString(nint table, string name, string value);
 		#endregion
 
+		readonly nint _handle = IntPtr.Zero;
+
 		public ConfigImpl(string configFile)
 		{
-			Config_Load(configFile);
+			_handle = Config_Load(configFile);
+		}
+
+		public void AddValue(string table,string name, string value)
+		{
+			nint tablePtr = Config_AddTable(_handle, table);
+			ConfigTable_AddValueString(tablePtr, name, value);
 		}
 	}
 }
\ No newline at end of file
diff --git a/Engine/Source/Programs/UnrealBuildTool/Executors/UnrealBuildAccelerator/UBAExecutor.cs b/Engine/Source/Programs/UnrealBuildTool/Executors/UnrealBuildAccelerator/UBAExecutor.cs
--- a/Engine/Source/Programs/UnrealBuildTool/Executors/UnrealBuildAccelerator/UBAExecutor.cs
+++ b/Engine/Source/Programs/UnrealBuildTool/Executors/UnrealBuildAccelerator/UBAExecutor.cs
@@ -411,7 +411,11 @@ namespace UnrealBuildTool
 
 				string arch = RuntimeInformation.ProcessArchitecture.ToString().ToLowerInvariant();
 				FileReference configFile = FileReference.Combine(UbaBinariesDir, "UbaHost.toml");
-				EpicGames.UBA.IConfig.LoadConfig(configFile.FullName);
+				var config = EpicGames.UBA.IConfig.LoadConfig(configFile.FullName);
+				if (!String.IsNullOrEmpty(UBAConfig.CompressionLevel))
+				{
+					config.AddValue("Storage", "CompressionLevel", UBAConfig.CompressionLevel);
+				}
 
 				using EpicGames.UBA.ILogger ubaLogger = EpicGames.UBA.ILogger.CreateLogger(logger);
 				using (Server = IServer.CreateServer(UBAConfig.MaxWorkers, UBAConfig.SendSize, ubaLogger, UBAConfig.bUseQuic))
diff --git a/Engine/Source/Programs/UnrealBuildTool/Executors/UnrealBuildAccelerator/UnrealBuildAcceleratorConfig.cs b/Engine/Source/Programs/UnrealBuildTool/Executors/UnrealBuildAccelerator/UnrealBuildAcceleratorConfig.cs
--- a/Engine/Source/Programs/UnrealBuildTool/Executors/UnrealBuildAccelerator/UnrealBuildAcceleratorConfig.cs
+++ b/Engine/Source/Programs/UnrealBuildTool/Executors/UnrealBuildAccelerator/UnrealBuildAcceleratorConfig.cs
@@ -263,5 +263,12 @@ namespace UnrealBuildTool
 		[XmlConfigFile(Category = "UnrealBuildAccelerator", Name = "ReportCacheMissReason")]
 		[CommandLine("-UBAReportCacheMissReason")]
 		public bool bReportCacheMissReason { get; set; }
+
+		/// <summary>
+		/// Compression level. Options are None, SuperFast, VeryFast, Fast, Normal, Optimal1 to Optimal5
+		/// </summary>
+		[XmlConfigFile(Category = "UnrealBuildAccelerator")]
+		[CommandLine("-UBACompressionLevel")]
+		public string CompressionLevel { get; set; } = String.Empty;
 	}
 }
\ No newline at end of file
