Added missing propagation flag for optional properties

#rnx

diff --git a/Engine/Source/Programs/Shared/EpicGames.UHT/Types/Properties/UhtOptionalProperty.cs b/Engine/Source/Programs/Shared/EpicGames.UHT/Types/Properties/UhtOptionalProperty.cs
--- a/Engine/Source/Programs/Shared/EpicGames.UHT/Types/Properties/UhtOptionalProperty.cs
+++ b/Engine/Source/Programs/Shared/EpicGames.UHT/Types/Properties/UhtOptionalProperty.cs
@@ -39,7 +39,7 @@ namespace EpicGames.UHT.Types
 		public UhtOptionalProperty(UhtPropertySettings propertySettings, UhtProperty value) : base(propertySettings, value)
 		{
 			// If the creation of the value property set more flags, then copy those flags to ourselves
-			PropertyFlags |= ValueProperty.PropertyFlags & EPropertyFlags.UObjectWrapper;
+			PropertyFlags |= ValueProperty.PropertyFlags & (EPropertyFlags.UObjectWrapper | EPropertyFlags.TObjectPtr);
 
 			if (ValueProperty.MetaData.ContainsKey(UhtNames.NativeConst))
 			{
