Sanitize path like reference in Gauntlet issue handler

#rnx

diff --git a/Engine/Source/Programs/Shared/EpicGames.Horde/Issues/Handlers/GauntletIssueHandler.cs b/Engine/Source/Programs/Shared/EpicGames.Horde/Issues/Handlers/GauntletIssueHandler.cs
--- a/Engine/Source/Programs/Shared/EpicGames.Horde/Issues/Handlers/GauntletIssueHandler.cs
+++ b/Engine/Source/Programs/Shared/EpicGames.Horde/Issues/Handlers/GauntletIssueHandler.cs
@@ -137,6 +137,7 @@ namespace EpicGames.Horde.Issues.Handlers
 			sanitized = sanitized.Length > MaxMessageLength ? sanitized.Substring(0, MaxMessageLength) : sanitized;
 			sanitized = sanitized.Trim().ToUpperInvariant();
 			sanitized = Regex.Replace(sanitized, @"(?<![A-Z])(?:[A-Z]:|/)[^ :]+[/\\]SYNC[/\\]", "{root}/"); // Redact things that look like workspace roots; may be different between agents
+			sanitized = Regex.Replace(sanitized, @"(?:^|\s)/[\w/.:-]+", "{path}"); // Redact path like reference
 			sanitized = Regex.Replace(sanitized, @"0X[0-9A-F]+", "H"); // Redact hex strings
 			sanitized = Regex.Replace(sanitized, @"\d[\d.,:]*", "n"); // Redact numbers and timestamp like things
 
