Fixing Verse issues and improving FObjectProperty/FClassProperty setup.

#rnx
#prefilight 67daf1a35f405205b64164f6

diff --git a/Engine/Source/Programs/Shared/EpicGames.UHT/Parsers/UhtPropertyParser.cs b/Engine/Source/Programs/Shared/EpicGames.UHT/Parsers/UhtPropertyParser.cs
--- a/Engine/Source/Programs/Shared/EpicGames.UHT/Parsers/UhtPropertyParser.cs
+++ b/Engine/Source/Programs/Shared/EpicGames.UHT/Parsers/UhtPropertyParser.cs
@@ -1270,12 +1270,11 @@ namespace EpicGames.UHT.Parsers
 				}
 				else if (classObj.IsChildOf(args.Session.UClass))
 				{
-					// UObject here specifies that there is no limiter
-					property = new UhtClassProperty(propertySettings, UhtObjectCppForm.Native, classObj, args.Session.UObject);
+					property = new UhtClassProperty(propertySettings, UhtObjectCppForm.NativeClass, classObj);
 				}
 				else
 				{
-					property = new UhtObjectProperty(propertySettings, UhtObjectCppForm.Native, classObj);
+					property = new UhtObjectProperty(propertySettings, UhtObjectCppForm.NativeObject, classObj);
 				}
 			}
 			else
diff --git a/Engine/Source/Programs/Shared/EpicGames.UHT/Types/Properties/UhtClassProperty.cs b/Engine/Source/Programs/Shared/EpicGames.UHT/Types/Properties/UhtClassProperty.cs
--- a/Engine/Source/Programs/Shared/EpicGames.UHT/Types/Properties/UhtClassProperty.cs
+++ b/Engine/Source/Programs/Shared/EpicGames.UHT/Types/Properties/UhtClassProperty.cs
@@ -23,12 +23,15 @@ namespace EpicGames.UHT.Types
 		/// </summary>
 		/// <param name="propertySettings">Property settings</param>
 		/// <param name="cppForm">Source code form of the property</param>
-		/// <param name="classObj">Referenced class</param>
-		/// <param name="metaClass">Reference meta class</param>
+		/// <param name="referencedClass">Referenced class</param>
 		/// <param name="extraFlags">Extra flags to apply to the property.</param>
-		public UhtClassProperty(UhtPropertySettings propertySettings, UhtObjectCppForm cppForm, UhtClass classObj, UhtClass metaClass, EPropertyFlags extraFlags = EPropertyFlags.None)
-			: base(propertySettings, cppForm, classObj, metaClass, extraFlags)
+		public UhtClassProperty(UhtPropertySettings propertySettings, UhtObjectCppForm cppForm, UhtClass referencedClass, EPropertyFlags extraFlags = EPropertyFlags.None)
+			: base(propertySettings, cppForm, referencedClass, extraFlags, 0)
 		{
+			if (!cppForm.IsValidForClassProperty())
+			{
+				throw new UhtIceException($"Improper UhtObjectCppForm.{cppForm} for an UhtClassProperty");
+			}
 			PropertyCaps |= UhtPropertyCaps.CanHaveConfig;
 			PropertyCaps &= ~(UhtPropertyCaps.CanBeInstanced);
 		}
@@ -61,7 +64,7 @@ namespace EpicGames.UHT.Types
 
 				case UhtPropertyTextType.EventParameterMember:
 				case UhtPropertyTextType.EventParameterFunctionMember:
-					if (CppForm != UhtObjectCppForm.Native)
+					if (CppForm != UhtObjectCppForm.NativeClass)
 					{
 						AppendTemplateType(builder);
 					}
@@ -76,7 +79,7 @@ namespace EpicGames.UHT.Types
 					{
 						builder.Append("const ");
 					}
-					if (CppForm != UhtObjectCppForm.Native)
+					if (CppForm != UhtObjectCppForm.NativeClass)
 					{
 						AppendTemplateType(builder);
 					}
@@ -87,7 +90,7 @@ namespace EpicGames.UHT.Types
 					break;
 
 				case UhtPropertyTextType.FunctionThunkParameterArgType:
-					if ((PropertyFlags.HasAnyFlags(EPropertyFlags.OutParm) && CppForm != UhtObjectCppForm.Native) || isTemplateArgument)
+					if ((PropertyFlags.HasAnyFlags(EPropertyFlags.OutParm) && CppForm != UhtObjectCppForm.NativeClass) || isTemplateArgument)
 					{
 						AppendTemplateType(builder);
 					}
@@ -102,7 +105,7 @@ namespace EpicGames.UHT.Types
 					break;
 
 				case UhtPropertyTextType.VerseMangledType:
-					AppendVerseMangledType(builder, MetaClass != null && MetaClass.IsVerseField ? MetaClass : Class);
+					AppendVerseMangledType(builder, ReferencedClass);
 					break;
 			}
 			return builder;
@@ -117,7 +120,7 @@ namespace EpicGames.UHT.Types
 		/// <inheritdoc/>
 		public override StringBuilder AppendMemberDef(StringBuilder builder, IUhtPropertyMemberContext context, string name, string nameSuffix, string? offset, int tabs)
 		{
-			string paramsGenFlags = CppForm == UhtObjectCppForm.TObjectPtr
+			string paramsGenFlags = CppForm == UhtObjectCppForm.TObjectPtrClass
 				? "UECodeGen_Private::EPropertyGenFlags::Class | UECodeGen_Private::EPropertyGenFlags::ObjectPtr"
 				: "UECodeGen_Private::EPropertyGenFlags::Class";
 			AppendMemberDefStart(builder, context, name, nameSuffix, offset, tabs, "FClassPropertyParams", paramsGenFlags);
@@ -134,7 +137,7 @@ namespace EpicGames.UHT.Types
 
 			// UFunctions with a smart pointer as input parameter wont compile anyway, because of missing P_GET_... macro.
 			// UFunctions with a smart pointer as return type will crash when called via blueprint, because they are not supported in VM.
-			if (!options.HasAnyFlags(UhtValidationOptions.IsKey) && PropertyCategory != UhtPropertyCategory.Member && CppForm == UhtObjectCppForm.TObjectPtr)
+			if (!options.HasAnyFlags(UhtValidationOptions.IsKey) && PropertyCategory != UhtPropertyCategory.Member && CppForm == UhtObjectCppForm.TObjectPtrClass)
 			{
 				outerStruct.LogError("UFunctions cannot take a TObjectPtr as a function parameter or return value.");
 			}
diff --git a/Engine/Source/Programs/Shared/EpicGames.UHT/Types/Properties/UhtInterfaceProperty.cs b/Engine/Source/Programs/Shared/EpicGames.UHT/Types/Properties/UhtInterfaceProperty.cs
--- a/Engine/Source/Programs/Shared/EpicGames.UHT/Types/Properties/UhtInterfaceProperty.cs
+++ b/Engine/Source/Programs/Shared/EpicGames.UHT/Types/Properties/UhtInterfaceProperty.cs
@@ -198,7 +198,7 @@ namespace EpicGames.UHT.Types
 			{
 				return new UhtInterfaceProperty(args.PropertySettings, propertyClass);
 			}
-			return new UhtObjectProperty(args.PropertySettings, UhtObjectCppForm.Native, propertyClass, EPropertyFlags.UObjectWrapper);
+			return new UhtObjectProperty(args.PropertySettings, UhtObjectCppForm.NativeObject, propertyClass, EPropertyFlags.UObjectWrapper);
 		}
 		#endregion
 	}
diff --git a/Engine/Source/Programs/Shared/EpicGames.UHT/Types/Properties/UhtLazyObjectPtrProperty.cs b/Engine/Source/Programs/Shared/EpicGames.UHT/Types/Properties/UhtLazyObjectPtrProperty.cs
--- a/Engine/Source/Programs/Shared/EpicGames.UHT/Types/Properties/UhtLazyObjectPtrProperty.cs
+++ b/Engine/Source/Programs/Shared/EpicGames.UHT/Types/Properties/UhtLazyObjectPtrProperty.cs
@@ -36,9 +36,9 @@ namespace EpicGames.UHT.Types
 		/// Construct new property
 		/// </summary>
 		/// <param name="propertySettings">Property settings</param>
-		/// <param name="classObj">Referenced class</param>
-		public UhtLazyObjectPtrProperty(UhtPropertySettings propertySettings, UhtClass classObj)
-			: base(propertySettings, UhtObjectCppForm.Native, classObj, null)
+		/// <param name="referencedClass">Referenced class</param>
+		public UhtLazyObjectPtrProperty(UhtPropertySettings propertySettings, UhtClass referencedClass)
+			: base(propertySettings, UhtObjectCppForm.NativeObject, referencedClass)
 		{
 			PropertyFlags |= EPropertyFlags.UObjectWrapper;
 			PropertyCaps |= UhtPropertyCaps.PassCppArgsByRef | UhtPropertyCaps.RequiresNullConstructorArg;
diff --git a/Engine/Source/Programs/Shared/EpicGames.UHT/Types/Properties/UhtObjectProperty.cs b/Engine/Source/Programs/Shared/EpicGames.UHT/Types/Properties/UhtObjectProperty.cs
--- a/Engine/Source/Programs/Shared/EpicGames.UHT/Types/Properties/UhtObjectProperty.cs
+++ b/Engine/Source/Programs/Shared/EpicGames.UHT/Types/Properties/UhtObjectProperty.cs
@@ -22,11 +22,15 @@ namespace EpicGames.UHT.Types
 		/// </summary>
 		/// <param name="propertySettings">Property settings</param>
 		/// <param name="cppForm">Source code form of the property</param>
-		/// <param name="classObj">Referenced class</param>
+		/// <param name="referencedClass">Referenced class</param>
 		/// <param name="extraFlags">Extra flags to add to the property</param>
-		public UhtObjectProperty(UhtPropertySettings propertySettings, UhtObjectCppForm cppForm, UhtClass classObj, EPropertyFlags extraFlags = EPropertyFlags.None)
-			: this(propertySettings, cppForm, classObj, null, extraFlags)
+		public UhtObjectProperty(UhtPropertySettings propertySettings, UhtObjectCppForm cppForm, UhtClass referencedClass, EPropertyFlags extraFlags = EPropertyFlags.None)
+			: this(propertySettings, cppForm, referencedClass, extraFlags, 0)
 		{
+			if (!cppForm.IsValidForObjectProperty())
+			{
+				throw new UhtIceException($"Improper UhtObjectCppForm.{cppForm} for an UhtObjectProperty");
+			}
 		}
 
 		/// <summary>
@@ -34,11 +38,12 @@ namespace EpicGames.UHT.Types
 		/// </summary>
 		/// <param name="propertySettings">Property settings</param>
 		/// <param name="cppForm">Source code form of the property</param>
-		/// <param name="classObj">Referenced class</param>
-		/// <param name="metaClass">Optional reference class (used by class properties)</param>
+		/// <param name="referencedClass">Referenced class</param>
 		/// <param name="extraFlags">Extra flags to add to the property</param>
-		protected UhtObjectProperty(UhtPropertySettings propertySettings, UhtObjectCppForm cppForm, UhtClass classObj, UhtClass? metaClass = null, EPropertyFlags extraFlags = EPropertyFlags.None)
-			: base(propertySettings, cppForm, classObj, metaClass, extraFlags)
+		/// <param name="dummyArg">Extra argument that the UhtClassProperty uses to invoke the right constructor</param>
+		[System.Diagnostics.CodeAnalysis.SuppressMessage("Style", "IDE0060:Remove unused parameter", Justification = "<Pending>")]
+		protected UhtObjectProperty(UhtPropertySettings propertySettings, UhtObjectCppForm cppForm, UhtClass referencedClass, EPropertyFlags extraFlags, int dummyArg)
+			: base(propertySettings, cppForm, referencedClass, extraFlags)
 		{
 			PropertyCaps |= UhtPropertyCaps.RequiresNullConstructorArg | UhtPropertyCaps.CanBeInstanced | UhtPropertyCaps.CanExposeOnSpawn |
 				UhtPropertyCaps.IsParameterSupportedByBlueprint | UhtPropertyCaps.IsMemberSupportedByBlueprint | UhtPropertyCaps.SupportsRigVM;
@@ -109,7 +114,7 @@ namespace EpicGames.UHT.Types
 					}
 					else if (CppForm == UhtObjectCppForm.TInterfaceInstance)
 					{
-						builder.AppendClassSourceNameOrInterfaceProxyName(Class);
+						builder.AppendClassSourceNameOrInterfaceProxyName(ReferencedClass);
 					}
 					else
 					{
@@ -118,7 +123,7 @@ namespace EpicGames.UHT.Types
 					break;
 
 				case UhtPropertyTextType.VerseMangledType:
-					AppendVerseMangledType(builder, Class);
+					AppendVerseMangledType(builder, ReferencedClass);
 					break;
 
 				default:
@@ -137,7 +142,7 @@ namespace EpicGames.UHT.Types
 		/// <inheritdoc/>
 		public override StringBuilder AppendMemberDef(StringBuilder builder, IUhtPropertyMemberContext context, string name, string nameSuffix, string? offset, int tabs)
 		{
-			string paramsGenFlags = CppForm == UhtObjectCppForm.TObjectPtr
+			string paramsGenFlags = CppForm == UhtObjectCppForm.TObjectPtrObject
 				? "UECodeGen_Private::EPropertyGenFlags::Object | UECodeGen_Private::EPropertyGenFlags::ObjectPtr"
 				: "UECodeGen_Private::EPropertyGenFlags::Object";
 			AppendMemberDefStart(builder, context, name, nameSuffix, offset, tabs, "FObjectPropertyParams", paramsGenFlags);
@@ -170,7 +175,7 @@ namespace EpicGames.UHT.Types
 
 			// UFunctions with a smart pointer as input parameter wont compile anyway, because of missing P_GET_... macro.
 			// UFunctions with a smart pointer as return type will crash when called via blueprint, because they are not supported in VM.
-			if (PropertyCategory != UhtPropertyCategory.Member && CppForm == UhtObjectCppForm.TObjectPtr)
+			if (PropertyCategory != UhtPropertyCategory.Member && (CppForm == UhtObjectCppForm.TObjectPtrObject || CppForm == UhtObjectCppForm.TObjectPtrClass))
 			{
 				// At this point, allow this to appear in TMap keys in the UPlayerMappableInputConfig class
 				if (!options.HasAnyFlags(UhtValidationOptions.IsKey) ||
@@ -187,7 +192,7 @@ namespace EpicGames.UHT.Types
 		protected override void ValidateMember(UhtStruct structObj, UhtValidationOptions options)
 		{
 			base.ValidateMember(structObj, options);
-			if (Class.NativeInterface != null && CppForm == UhtObjectCppForm.TObjectPtr)
+			if (Class.NativeInterface != null && (CppForm == UhtObjectCppForm.TObjectPtrObject || CppForm == UhtObjectCppForm.TObjectPtrClass))
 			{
 				this.LogError($"UPROPERTY pointers cannot be interfaces - did you mean TScriptInterface<{Class.SourceName}>?");
 			}
diff --git a/Engine/Source/Programs/Shared/EpicGames.UHT/Types/Properties/UhtObjectPropertyBase.cs b/Engine/Source/Programs/Shared/EpicGames.UHT/Types/Properties/UhtObjectPropertyBase.cs
--- a/Engine/Source/Programs/Shared/EpicGames.UHT/Types/Properties/UhtObjectPropertyBase.cs
+++ b/Engine/Source/Programs/Shared/EpicGames.UHT/Types/Properties/UhtObjectPropertyBase.cs
@@ -8,10 +8,13 @@ using System.Text.Json.Serialization;
 using EpicGames.Core;
 using EpicGames.UHT.Tables;
 using EpicGames.UHT.Tokenizer;
+using EpicGames.UHT.Types;
 using EpicGames.UHT.Utils;
 
 namespace EpicGames.UHT.Types
 {
+
+	#region CppForm
 	/// <summary>
 	/// How the object property appeared in the code
 	/// </summary>
@@ -19,14 +22,34 @@ namespace EpicGames.UHT.Types
 	{
 
 		/// <summary>
-		/// Existing native (raw pointer or TObjectPtr)
+		/// Existing native object raw pointer
 		/// </summary>
-		Native,
+		NativeObject,
 
 		/// <summary>
-		/// TObjectPtr
+		/// Existing native class raw pointer
 		/// </summary>
-		TObjectPtr,
+		NativeClass,
+
+		/// <summary>
+		/// TObjectPtr for an object
+		/// </summary>
+		TObjectPtrObject,
+
+		/// <summary>
+		/// TObjectPtr for a class
+		/// </summary>
+		TObjectPtrClass,
+
+		/// <summary>
+		/// TSoftObjectPtr
+		/// </summary>
+		TSoftObjectPtr,
+
+		/// <summary>
+		/// TSoftClassPtr
+		/// </summary>
+		TSoftClassPtr,
 
 		/// <summary>
 		/// TNonNullPtr
@@ -70,6 +93,70 @@ namespace EpicGames.UHT.Types
 	public static class UhtObjectCppFormExtensions
 	{
 
+		/// <summary>
+		/// Check to see if the given CPP form is valid for a object property
+		/// </summary>
+		/// <param name="cppForm">The form in question</param>
+		/// <returns>Is valid if true</returns>
+		/// <exception cref="NotImplementedException"></exception>
+		public static bool IsValidForObjectProperty(this UhtObjectCppForm cppForm)
+		{
+			switch (cppForm)
+			{
+				case UhtObjectCppForm.NativeObject:
+				case UhtObjectCppForm.TObjectPtrObject:
+				case UhtObjectCppForm.TSoftObjectPtr:
+				case UhtObjectCppForm.TNonNullPtr:
+				case UhtObjectCppForm.TInterfaceInstance:
+					return true;
+
+				case UhtObjectCppForm.NativeClass:
+				case UhtObjectCppForm.TObjectPtrClass:
+				case UhtObjectCppForm.TSoftClassPtr:
+				case UhtObjectCppForm.TSubclassOf:
+				case UhtObjectCppForm.VerseType:
+				case UhtObjectCppForm.VerseCastableType:
+				case UhtObjectCppForm.VerseSubtype:
+				case UhtObjectCppForm.VerseCastableSubtype:
+					return false;
+
+				default:
+					throw new NotImplementedException();
+			}
+		}
+
+		/// <summary>
+		/// Check to see if the given CPP form is valid for a class property
+		/// </summary>
+		/// <param name="cppForm">The form in question</param>
+		/// <returns>Is valid if true</returns>
+		/// <exception cref="NotImplementedException"></exception>
+		public static bool IsValidForClassProperty(this UhtObjectCppForm cppForm)
+		{
+			switch (cppForm)
+			{
+				case UhtObjectCppForm.NativeClass:
+				case UhtObjectCppForm.TObjectPtrClass:
+				case UhtObjectCppForm.TSoftClassPtr:
+				case UhtObjectCppForm.VerseType:
+				case UhtObjectCppForm.VerseCastableType:
+				case UhtObjectCppForm.VerseSubtype:
+				case UhtObjectCppForm.VerseCastableSubtype:
+				case UhtObjectCppForm.TSubclassOf:
+					return true;
+
+				case UhtObjectCppForm.NativeObject:
+				case UhtObjectCppForm.TObjectPtrObject:
+				case UhtObjectCppForm.TSoftObjectPtr:
+				case UhtObjectCppForm.TNonNullPtr:
+				case UhtObjectCppForm.TInterfaceInstance:
+					return false;
+
+				default:
+					throw new NotImplementedException();
+			}
+		}
+
 		/// <summary>
 		/// Return a CppForm used when doing IsSameType checks
 		/// </summary>
@@ -79,10 +166,14 @@ namespace EpicGames.UHT.Types
 		{
 			switch (cppForm)
 			{
-				case UhtObjectCppForm.Native:
-				case UhtObjectCppForm.TObjectPtr:
-					return UhtObjectCppForm.Native;
-
+				case UhtObjectCppForm.NativeObject:
+				case UhtObjectCppForm.NativeClass:
+				case UhtObjectCppForm.TObjectPtrObject:
+				case UhtObjectCppForm.TObjectPtrClass:
+					return UhtObjectCppForm.NativeObject;
+
+				case UhtObjectCppForm.TSoftObjectPtr:
+				case UhtObjectCppForm.TSoftClassPtr:
 				case UhtObjectCppForm.TInterfaceInstance:
 				case UhtObjectCppForm.TNonNullPtr:
 				case UhtObjectCppForm.TSubclassOf:
@@ -106,8 +197,12 @@ namespace EpicGames.UHT.Types
 		{
 			switch (cppForm)
 			{
-				case UhtObjectCppForm.Native:
-				case UhtObjectCppForm.TObjectPtr:
+				case UhtObjectCppForm.NativeObject:
+				case UhtObjectCppForm.NativeClass:
+				case UhtObjectCppForm.TObjectPtrObject:
+				case UhtObjectCppForm.TObjectPtrClass:
+				case UhtObjectCppForm.TSoftObjectPtr:
+				case UhtObjectCppForm.TSoftClassPtr:
 				case UhtObjectCppForm.TInterfaceInstance:
 				case UhtObjectCppForm.TNonNullPtr:
 				case UhtObjectCppForm.TSubclassOf:
@@ -128,19 +223,26 @@ namespace EpicGames.UHT.Types
 		/// Return the P_GET macro text based on the cpp form
 		/// </summary>
 		/// <param name="cppForm">The form in question</param>
-		/// <param name="nativeText">Text for the native form</param>
 		/// <returns>P_GET macro text</returns>
 		/// <exception cref="NotImplementedException"></exception>
-		public static string GetPGetMacroText(this UhtObjectCppForm cppForm, string nativeText)
+		public static string GetPGetMacroText(this UhtObjectCppForm cppForm)
 		{
 			switch (cppForm)
 			{
-				case UhtObjectCppForm.Native:
-					return nativeText;
+				case UhtObjectCppForm.NativeObject:
+				case UhtObjectCppForm.NativeClass:
+					return "OBJECT";
 
-				case UhtObjectCppForm.TObjectPtr:
+				case UhtObjectCppForm.TObjectPtrObject:
+				case UhtObjectCppForm.TObjectPtrClass:
 					return "OBJECTPTR";
 
+				case UhtObjectCppForm.TSoftObjectPtr:
+					return "SOFTOBJECT";
+
+				case UhtObjectCppForm.TSoftClassPtr:
+					return "SOFTCLASS";
+
 				case UhtObjectCppForm.TNonNullPtr:
 					return "NONNULLPTROBJECT";
 
@@ -181,8 +283,12 @@ namespace EpicGames.UHT.Types
 				case UhtObjectCppForm.VerseCastableType:
 					return UhtPGetArgumentType.None;
 
-				case UhtObjectCppForm.Native:
-				case UhtObjectCppForm.TObjectPtr:
+				case UhtObjectCppForm.NativeObject:
+				case UhtObjectCppForm.NativeClass:
+				case UhtObjectCppForm.TObjectPtrObject:
+				case UhtObjectCppForm.TObjectPtrClass:
+				case UhtObjectCppForm.TSoftObjectPtr:
+				case UhtObjectCppForm.TSoftClassPtr:
 				case UhtObjectCppForm.TNonNullPtr:
 				case UhtObjectCppForm.TSubclassOf:
 				case UhtObjectCppForm.VerseSubtype:
@@ -205,8 +311,12 @@ namespace EpicGames.UHT.Types
 		{
 			switch (cppForm)
 			{
-				case UhtObjectCppForm.Native:
-				case UhtObjectCppForm.TObjectPtr:
+				case UhtObjectCppForm.NativeObject:
+				case UhtObjectCppForm.NativeClass:
+				case UhtObjectCppForm.TObjectPtrObject:
+				case UhtObjectCppForm.TObjectPtrClass:
+				case UhtObjectCppForm.TSoftObjectPtr:
+				case UhtObjectCppForm.TSoftClassPtr:
 				case UhtObjectCppForm.TInterfaceInstance:
 				case UhtObjectCppForm.TNonNullPtr:
 				case UhtObjectCppForm.VerseType:
@@ -228,45 +338,64 @@ namespace EpicGames.UHT.Types
 		/// </summary>
 		/// <param name="cppForm">The form in question</param>
 		/// <param name="property">Property to be updated</param>
-		public static void InitializeProperty(this UhtObjectCppForm cppForm, UhtObjectPropertyBase property)
+		public static (UhtClass, UhtClass?) InitializeProperty(this UhtObjectCppForm cppForm, UhtObjectPropertyBase property)
 		{
 			switch (cppForm)
 			{
-				case UhtObjectCppForm.Native:
-					break;
+				case UhtObjectCppForm.NativeObject:
+					return (property.ReferencedClass, null);
+
+				case UhtObjectCppForm.NativeClass:
+					return (property.ReferencedClass, property.Session.UObject);
+
+				case UhtObjectCppForm.TObjectPtrObject:
+					property.PropertyFlags |= EPropertyFlags.TObjectPtr | EPropertyFlags.UObjectWrapper;
+					property.PropertyCaps |= UhtPropertyCaps.PassCppArgsByRef;
+					return (property.ReferencedClass, null);
+
+				case UhtObjectCppForm.TObjectPtrClass:
+					property.PropertyFlags |= EPropertyFlags.TObjectPtr | EPropertyFlags.UObjectWrapper;
+					property.PropertyCaps |= UhtPropertyCaps.PassCppArgsByRef;
+					return (property.ReferencedClass, property.Session.UObject);
+
+				case UhtObjectCppForm.TSoftObjectPtr:
+					property.PropertyFlags |= EPropertyFlags.UObjectWrapper;
+					return (property.ReferencedClass, null);
+
+				case UhtObjectCppForm.TSoftClassPtr:
+					property.PropertyFlags |= EPropertyFlags.UObjectWrapper;
+					return (property.Session.UClass, property.ReferencedClass);
 
 				case UhtObjectCppForm.TInterfaceInstance:
+					property.MetaData.Add("ObjectMustImplement", property.ReferencedClass.PathName);
+					return (property.Session.UObject, null);
+
 				case UhtObjectCppForm.VerseType:
 				case UhtObjectCppForm.VerseCastableType:
 				case UhtObjectCppForm.VerseSubtype:
 				case UhtObjectCppForm.VerseCastableSubtype:
-					// TODO: struct FInterfaceTraits::"ObjectMustImplement"
-					//if (property.Class.ClassType == UhtClassType.Interface)
-					//{
-					//	property.MetaData.Add("MustImplement", property.Class.PathName);
-					//}
-					//property.MetaData.Add("AllowAbstract", "True");
-					break;
-
-				case UhtObjectCppForm.TObjectPtr:
-					property.PropertyFlags |= EPropertyFlags.TObjectPtr | EPropertyFlags.UObjectWrapper;
-					property.PropertyCaps |= UhtPropertyCaps.PassCppArgsByRef;
-					break;
+					property.MetaData.Add("AllowAbstract", "True");
+					if (property.ReferencedClass.ClassType == UhtClassType.Interface)
+					{
+						property.MetaData.Add("MustImplement", property.ReferencedClass.PathName);
+						return (property.Session.UClass, property.Session.UObject);
+					}
+					return (property.Session.UClass, property.ReferencedClass);
 
 				case UhtObjectCppForm.TSubclassOf:
 					property.PropertyFlags |= EPropertyFlags.UObjectWrapper;
-					break;
+					return (property.Session.UClass, property.ReferencedClass);
 
 				case UhtObjectCppForm.TNonNullPtr:
 					property.PropertyFlags |= EPropertyFlags.NonNullable;
-					break;
+					return (property.ReferencedClass, null);
 
 				default:
 					throw new NotImplementedException();
 			}
-
 		}
 	}
+	#endregion
 
 	/// <summary>
 	/// FObjectPropertyBase
@@ -282,7 +411,7 @@ namespace EpicGames.UHT.Types
 		protected override string CppTypeText => "Object";
 
 		/// <inheritdoc/>
-		protected override string PGetMacroText => CppForm.GetPGetMacroText("OBJECT");
+		protected override string PGetMacroText => CppForm.GetPGetMacroText();
 
 		/// <inheritdoc/>
 		protected override UhtPGetArgumentType PGetTypeArgument => CppForm.GetPGetArgumentType();
@@ -291,7 +420,13 @@ namespace EpicGames.UHT.Types
 		protected override bool PGetPassAsNoPtr => CppForm.GetPGetPassAsNoPtr();
 
 		/// <summary>
-		/// Referenced UCLASS
+		/// Originally referenced class when property is created.
+		/// </summary>
+		[JsonConverter(typeof(UhtTypeSourceNameJsonConverter<UhtClass>))]
+		public UhtClass ReferencedClass { get; set; }
+
+		/// <summary>
+		/// Referenced UCLASS that might later be changed depending on the type (i.e. Verse's TInterfaceInstance)
 		/// </summary>
 		[JsonConverter(typeof(UhtTypeSourceNameJsonConverter<UhtClass>))]
 		public UhtClass Class { get; set; }
@@ -305,25 +440,23 @@ namespace EpicGames.UHT.Types
 		/// <summary>
 		/// The form that the object property appeared in the code
 		/// </summary>
-		public UhtObjectCppForm CppForm { get; init; } = UhtObjectCppForm.Native;
+		public UhtObjectCppForm CppForm { get; init; }
 
 		/// <summary>
 		/// Construct a property
 		/// </summary>
 		/// <param name="propertySettings">Property settings</param>
 		/// <param name="cppForm">Source code form of the property</param>
-		/// <param name="classObj">Referenced UCLASS (a.k.a. PropertyClass)</param>
-		/// <param name="metaClass">Referenced UCLASS used by class properties</param>
+		/// <param name="referencedClass">Referenced UCLASS</param>
 		/// <param name="extraFlags">Extra flags to apply to the property.</param>
-		protected UhtObjectPropertyBase(UhtPropertySettings propertySettings, UhtObjectCppForm cppForm, UhtClass classObj, UhtClass? metaClass, EPropertyFlags extraFlags = EPropertyFlags.None) : base(propertySettings)
+		protected UhtObjectPropertyBase(UhtPropertySettings propertySettings, UhtObjectCppForm cppForm, UhtClass referencedClass, EPropertyFlags extraFlags = EPropertyFlags.None) : base(propertySettings)
 		{
 			CppForm = cppForm;
-			Class = classObj;
-			MetaClass = metaClass;
+			ReferencedClass = referencedClass;
 			PropertyFlags |= extraFlags;
 
 			// Based on the form, perform more initialization
-			CppForm.InitializeProperty(this);
+			(Class, MetaClass) = CppForm.InitializeProperty(this);
 
 			// This applies to EVERYTHING including raw pointer
 			// Imply const if it's a parameter that is a pointer to a const class
@@ -434,19 +567,17 @@ namespace EpicGames.UHT.Types
 			{
 				case UhtObjectCppForm.VerseType:
 				case UhtObjectCppForm.VerseCastableType:
-					builder.Append(SourceName);
-					break;
-
 				case UhtObjectCppForm.VerseSubtype:
-					builder.Append("subtype ").AppendVersePath(type);
-					break;
-
 				case UhtObjectCppForm.VerseCastableSubtype:
-					builder.Append("castable_subtype ").AppendVersePath(type);
+					builder.Append(SourceName);
 					break;
 
-				case UhtObjectCppForm.Native:
-				case UhtObjectCppForm.TObjectPtr:
+				case UhtObjectCppForm.NativeObject:
+				case UhtObjectCppForm.NativeClass:
+				case UhtObjectCppForm.TObjectPtrObject:
+				case UhtObjectCppForm.TObjectPtrClass:
+				case UhtObjectCppForm.TSoftObjectPtr:
+				case UhtObjectCppForm.TSoftClassPtr:
 				case UhtObjectCppForm.TNonNullPtr:
 				case UhtObjectCppForm.TInterfaceInstance:
 				case UhtObjectCppForm.TSubclassOf:
@@ -469,14 +600,24 @@ namespace EpicGames.UHT.Types
 		{
 			switch (CppForm)
 			{
-				case UhtObjectCppForm.Native:
+				case UhtObjectCppForm.NativeObject:
+				case UhtObjectCppForm.NativeClass:
 					builder.Append(Class.SourceName).Append('*');
 					break;
 
-				case UhtObjectCppForm.TObjectPtr:
+				case UhtObjectCppForm.TObjectPtrObject:
+				case UhtObjectCppForm.TObjectPtrClass:
 					builder.Append("TObjectPtr<").Append(Class.SourceName).Append('>');
 					break;
 
+				case UhtObjectCppForm.TSoftObjectPtr:
+					builder.Append("TSoftObjectPtr<").Append(Class.SourceName).Append('>');
+					break;
+
+				case UhtObjectCppForm.TSoftClassPtr:
+					builder.Append("TSoftObjectPtr<").Append(Class.SourceName).Append('>');
+					break;
+
 				case UhtObjectCppForm.TNonNullPtr:
 					builder.Append("TNonNullPtr<").Append(Class.SourceName).Append('>');
 					break;
@@ -502,7 +643,7 @@ namespace EpicGames.UHT.Types
 					break;
 
 				case UhtObjectCppForm.TInterfaceInstance:
-					builder.Append("TInterfaceInstance<").AppendClassSourceNameOrInterfaceProxyName(Class).Append('>');
+					builder.Append("TInterfaceInstance<").AppendClassSourceNameOrInterfaceProxyName(ReferencedClass).Append('>');
 					break;
 
 				default:
@@ -524,7 +665,7 @@ namespace EpicGames.UHT.Types
 			}
 
 			// With TSubclassOf, MetaClass is used as a class limiter.  
-			return new UhtClassProperty(args.PropertySettings, UhtObjectCppForm.TSubclassOf, metaClass.Session.UClass, metaClass);
+			return new UhtClassProperty(args.PropertySettings, UhtObjectCppForm.TSubclassOf, metaClass);
 		}
 
 		[UhtPropertyType(Keyword = "TNonNullPtr")]
@@ -532,7 +673,9 @@ namespace EpicGames.UHT.Types
 		[SuppressMessage("Style", "IDE0060:Remove unused parameter", Justification = "Attribute accessed method")]
 		private static UhtProperty? TNonNullPtrType(UhtPropertyResolveArgs args)
 		{
-			return ObjectTypeProperty(args, UhtObjectCppForm.TNonNullPtr, UhtTemplateObjectMode.Normal);
+			return args.ParseTemplateObjectProperty(UhtTemplateObjectMode.Normal, (UhtClass referencedClass, int typeStartPos) => {
+				return new UhtObjectProperty(args.PropertySettings, UhtObjectCppForm.TNonNullPtr, referencedClass);
+			});
 		}
 
 		[UhtPropertyType(Keyword = "TObjectPtr")]
@@ -540,7 +683,19 @@ namespace EpicGames.UHT.Types
 		[SuppressMessage("Style", "IDE0060:Remove unused parameter", Justification = "Attribute accessed method")]
 		private static UhtProperty? ObjectPtrProperty(UhtPropertyResolveArgs args)
 		{
-			return ObjectTypeProperty(args, UhtObjectCppForm.TObjectPtr, UhtTemplateObjectMode.Normal);
+			return args.ParseTemplateObjectProperty(UhtTemplateObjectMode.Normal, (UhtClass referencedClass, int typeStartPos) => {
+				args.ConditionalLogPointerUsage(args.Config.EngineObjectPtrMemberBehavior, args.Config.EnginePluginObjectPtrMemberBehavior,
+					args.Config.NonEngineObjectPtrMemberBehavior, "ObjectPtr", typeStartPos, null);
+				if (referencedClass.IsChildOf(args.Session.UClass))
+				{
+					// UObject specifies that there is no limiter
+					return new UhtClassProperty(args.PropertySettings, UhtObjectCppForm.TObjectPtrClass, referencedClass);
+				}
+				else
+				{
+					return new UhtObjectProperty(args.PropertySettings, UhtObjectCppForm.TObjectPtrObject, referencedClass);
+				}
+			});
 		}
 
 		[UhtPropertyType(Keyword = "verse::type", Options = UhtPropertyTypeOptions.Simple)]
@@ -548,7 +703,7 @@ namespace EpicGames.UHT.Types
 		[SuppressMessage("Style", "IDE0060:Remove unused parameter", Justification = "Attribute accessed method")]
 		private static UhtProperty? VerseTypeProperty(UhtPropertyResolveArgs args)
 		{
-			return new UhtClassProperty(args.PropertySettings, UhtObjectCppForm.VerseType, args.Session.UClass, args.Session.UObject);
+			return new UhtClassProperty(args.PropertySettings, UhtObjectCppForm.VerseType, args.Session.UObject);
 		}
 
 		[UhtPropertyType(Keyword = "verse::castable_type", Options = UhtPropertyTypeOptions.Simple)]
@@ -556,7 +711,7 @@ namespace EpicGames.UHT.Types
 		[SuppressMessage("Style", "IDE0060:Remove unused parameter", Justification = "Attribute accessed method")]
 		private static UhtProperty? VerseCastableTypeProperty(UhtPropertyResolveArgs args)
 		{
-			return new UhtClassProperty(args.PropertySettings, UhtObjectCppForm.VerseCastableType, args.Session.UClass, args.Session.UObject);
+			return new UhtClassProperty(args.PropertySettings, UhtObjectCppForm.VerseCastableType, args.Session.UObject);
 		}
 
 		[UhtPropertyType(Keyword = "verse::subtype")]
@@ -564,7 +719,9 @@ namespace EpicGames.UHT.Types
 		[SuppressMessage("Style", "IDE0060:Remove unused parameter", Justification = "Attribute accessed method")]
 		private static UhtProperty? VerseSubTypeProperty(UhtPropertyResolveArgs args)
 		{
-			return ObjectTypeProperty(args, UhtObjectCppForm.VerseSubtype, UhtTemplateObjectMode.Normal);
+			return args.ParseTemplateObjectProperty(UhtTemplateObjectMode.Normal, (UhtClass referencedClass, int typeStartPos) => {
+				return new UhtClassProperty(args.PropertySettings, UhtObjectCppForm.VerseSubtype, referencedClass);
+			});
 		}
 
 		[UhtPropertyType(Keyword = "verse::castable_subtype")]
@@ -572,7 +729,9 @@ namespace EpicGames.UHT.Types
 		[SuppressMessage("Style", "IDE0060:Remove unused parameter", Justification = "Attribute accessed method")]
 		private static UhtProperty? VerseCastableSubtypeProperty(UhtPropertyResolveArgs args)
 		{
-			return ObjectTypeProperty(args, UhtObjectCppForm.VerseCastableSubtype, UhtTemplateObjectMode.Normal);
+			return args.ParseTemplateObjectProperty(UhtTemplateObjectMode.Normal, (UhtClass referencedClass, int typeStartPos) => {
+				return new UhtClassProperty(args.PropertySettings, UhtObjectCppForm.VerseCastableSubtype, referencedClass);
+			});
 		}
 
 		[UhtPropertyType(Keyword = "TInterfaceInstance")]
@@ -580,49 +739,9 @@ namespace EpicGames.UHT.Types
 		[SuppressMessage("Style", "IDE0060:Remove unused parameter", Justification = "Attribute accessed method")]
 		private static UhtProperty? TInterfaceInstanceProperty(UhtPropertyResolveArgs args)
 		{
-			return ObjectTypeProperty(args, UhtObjectCppForm.TInterfaceInstance, UhtTemplateObjectMode.NativeInterfaceProxy);
-		}
-
-		private static UhtProperty? ObjectTypeProperty(UhtPropertyResolveArgs args, UhtObjectCppForm cppForm, UhtTemplateObjectMode templateObjectMode)
-		{
-			int typeStartPos = args.TokenReader.PeekToken().InputStartPos;
-			UhtClass? propertyClass = args.ParseTemplateObject(templateObjectMode);
-			if (propertyClass == null)
-			{
-				return null;
-			}
-
-			if (cppForm == UhtObjectCppForm.TObjectPtr)
-			{
-				args.ConditionalLogPointerUsage(args.Config.EngineObjectPtrMemberBehavior, args.Config.EnginePluginObjectPtrMemberBehavior,
-					args.Config.NonEngineObjectPtrMemberBehavior, "ObjectPtr", typeStartPos, null);
-			}
-
-			switch (cppForm)
-			{
-				case UhtObjectCppForm.Native:
-				case UhtObjectCppForm.TObjectPtr:
-				case UhtObjectCppForm.TNonNullPtr:
-					if (propertyClass.IsChildOf(propertyClass.Session.UClass))
-					{
-						// UObject specifies that there is no limiter
-						return new UhtClassProperty(args.PropertySettings, cppForm, propertyClass, propertyClass.Session.UObject);
-					}
-					else
-					{
-						return new UhtObjectProperty(args.PropertySettings, cppForm, propertyClass);
-					}
-
-				case UhtObjectCppForm.TInterfaceInstance:
-					return new UhtObjectProperty(args.PropertySettings, cppForm, propertyClass);
-
-				case UhtObjectCppForm.VerseSubtype:
-				case UhtObjectCppForm.VerseCastableSubtype:
-					return new UhtClassProperty(args.PropertySettings, cppForm, args.Session.UClass, propertyClass);
-
-				default:
-					throw new NotImplementedException();
-			}
+			return args.ParseTemplateObjectProperty(UhtTemplateObjectMode.NativeInterfaceProxy, (UhtClass referencedClass, int typeStartPos) => {
+				return new UhtObjectProperty(args.PropertySettings, UhtObjectCppForm.TInterfaceInstance, referencedClass);
+			});
 		}
 		#endregion
 	}
diff --git a/Engine/Source/Programs/Shared/EpicGames.UHT/Types/Properties/UhtSoftClassProperty.cs b/Engine/Source/Programs/Shared/EpicGames.UHT/Types/Properties/UhtSoftClassProperty.cs
--- a/Engine/Source/Programs/Shared/EpicGames.UHT/Types/Properties/UhtSoftClassProperty.cs
+++ b/Engine/Source/Programs/Shared/EpicGames.UHT/Types/Properties/UhtSoftClassProperty.cs
@@ -22,20 +22,13 @@ namespace EpicGames.UHT.Types
 		/// <inheritdoc/>
 		protected override string CppTypeText => "SoftClassPtr";
 
-		/// <inheritdoc/>
-		protected override string PGetMacroText => "SOFTCLASS";
-
-		/// <inheritdoc/>
-		protected override UhtPGetArgumentType PGetTypeArgument => UhtPGetArgumentType.TypeText;
-
 		/// <summary>
 		/// Construct a new class property
 		/// </summary>
 		/// <param name="propertySettings">Property setting</param>
-		/// <param name="classObj">Referenced class (Always UClass)</param>
-		/// <param name="metaClass">Referenced meta class</param>
-		public UhtSoftClassProperty(UhtPropertySettings propertySettings, UhtClass classObj, UhtClass metaClass)
-			: base(propertySettings, classObj, metaClass)
+		/// <param name="referencedClass">Referenced class</param>
+		public UhtSoftClassProperty(UhtPropertySettings propertySettings, UhtClass referencedClass)
+			: base(propertySettings, UhtObjectCppForm.TSoftClassPtr, referencedClass)
 		{
 		}
 
@@ -91,7 +84,7 @@ namespace EpicGames.UHT.Types
 			}
 
 			// With TSubclassOf, MetaClass is used as a class limiter.  
-			return new UhtSoftClassProperty(args.PropertySettings, metaClass.Session.UClass, metaClass);
+			return new UhtSoftClassProperty(args.PropertySettings, metaClass);
 		}
 		#endregion
 	}
diff --git a/Engine/Source/Programs/Shared/EpicGames.UHT/Types/Properties/UhtSoftObjectProperty.cs b/Engine/Source/Programs/Shared/EpicGames.UHT/Types/Properties/UhtSoftObjectProperty.cs
--- a/Engine/Source/Programs/Shared/EpicGames.UHT/Types/Properties/UhtSoftObjectProperty.cs
+++ b/Engine/Source/Programs/Shared/EpicGames.UHT/Types/Properties/UhtSoftObjectProperty.cs
@@ -2,9 +2,7 @@
 
 using System.Diagnostics.CodeAnalysis;
 using System.Text;
-using EpicGames.Core;
 using EpicGames.UHT.Tables;
-using EpicGames.UHT.Tokenizer;
 using EpicGames.UHT.Utils;
 
 namespace EpicGames.UHT.Types
@@ -22,22 +20,25 @@ namespace EpicGames.UHT.Types
 		/// <inheritdoc/>
 		protected override string CppTypeText => "SoftObjectPtr";
 
-		/// <inheritdoc/>
-		protected override string PGetMacroText => "SOFTOBJECT";
-
-		/// <inheritdoc/>
-		protected override UhtPGetArgumentType PGetTypeArgument => UhtPGetArgumentType.TypeText;
+		/// <summary>
+		/// Construct a new property
+		/// </summary>
+		/// <param name="propertySettings">Property settings</param>
+		/// <param name="referencedClass">UCLASS being referenced</param>
+		public UhtSoftObjectProperty(UhtPropertySettings propertySettings, UhtClass referencedClass)
+			: this(propertySettings, UhtObjectCppForm.TSoftObjectPtr, referencedClass)
+		{
+		}
 
 		/// <summary>
 		/// Construct a new property
 		/// </summary>
 		/// <param name="propertySettings">Property settings</param>
-		/// <param name="classObj">UCLASS being referenced</param>
-		/// <param name="metaClass">Optional meta class (used by SoftClassProperty)</param>
-		public UhtSoftObjectProperty(UhtPropertySettings propertySettings, UhtClass classObj, UhtClass? metaClass = null)
-			: base(propertySettings, UhtObjectCppForm.Native, classObj, metaClass)
+		/// <param name="cppForm">Form of the object property</param>
+		/// <param name="referencedClass">UCLASS being referenced</param>
+		protected UhtSoftObjectProperty(UhtPropertySettings propertySettings, UhtObjectCppForm cppForm, UhtClass referencedClass)
+			: base(propertySettings, cppForm, referencedClass)
 		{
-			PropertyFlags |= EPropertyFlags.UObjectWrapper;
 			PropertyCaps |= UhtPropertyCaps.PassCppArgsByRef | UhtPropertyCaps.RequiresNullConstructorArg | UhtPropertyCaps.CanExposeOnSpawn |
 				UhtPropertyCaps.CanHaveConfig | UhtPropertyCaps.IsParameterSupportedByBlueprint | UhtPropertyCaps.IsMemberSupportedByBlueprint;
 		}
diff --git a/Engine/Source/Programs/Shared/EpicGames.UHT/Types/Properties/UhtWeakObjectPtrProperty.cs b/Engine/Source/Programs/Shared/EpicGames.UHT/Types/Properties/UhtWeakObjectPtrProperty.cs
--- a/Engine/Source/Programs/Shared/EpicGames.UHT/Types/Properties/UhtWeakObjectPtrProperty.cs
+++ b/Engine/Source/Programs/Shared/EpicGames.UHT/Types/Properties/UhtWeakObjectPtrProperty.cs
@@ -36,7 +36,7 @@ namespace EpicGames.UHT.Types
 		/// <param name="propertyClass">Class being referenced</param>
 		/// <param name="extraFlags">Extra property flags to add to the definition</param>
 		public UhtWeakObjectPtrProperty(UhtPropertySettings propertySettings, UhtClass propertyClass, EPropertyFlags extraFlags = EPropertyFlags.None)
-			: base(propertySettings, UhtObjectCppForm.Native, propertyClass, null)
+			: base(propertySettings, UhtObjectCppForm.NativeObject, propertyClass)
 		{
 			PropertyFlags |= EPropertyFlags.UObjectWrapper | extraFlags;
 			PropertyCaps |= UhtPropertyCaps.PassCppArgsByRef | UhtPropertyCaps.RequiresNullConstructorArg | UhtPropertyCaps.IsMemberSupportedByBlueprint;
@@ -104,7 +104,7 @@ namespace EpicGames.UHT.Types
 
 			if (propertySettings.DisallowPropertyFlags.HasAnyFlags(EPropertyFlags.AutoWeak))
 			{
-				return new UhtObjectProperty(propertySettings, UhtObjectCppForm.Native, classObj, extraFlags | EPropertyFlags.UObjectWrapper);
+				return new UhtObjectProperty(propertySettings, UhtObjectCppForm.NativeObject, classObj, extraFlags | EPropertyFlags.UObjectWrapper);
 			}
 			return new UhtWeakObjectPtrProperty(propertySettings, classObj, extraFlags);
 		}
diff --git a/Engine/Source/Programs/Shared/EpicGames.UHT/Utils/UhtPropertyResolveArgs.cs b/Engine/Source/Programs/Shared/EpicGames.UHT/Utils/UhtPropertyResolveArgs.cs
--- a/Engine/Source/Programs/Shared/EpicGames.UHT/Utils/UhtPropertyResolveArgs.cs
+++ b/Engine/Source/Programs/Shared/EpicGames.UHT/Utils/UhtPropertyResolveArgs.cs
@@ -235,6 +235,24 @@ namespace EpicGames.UHT.Utils
 			return returnClass;
 		}
 
+		/// <summary>
+		/// Parse the token stream for a class type and then invoke the factory to create the property
+		/// </summary>
+		/// <param name="mode">Adjust what is returned and/or what is searched</param>
+		/// <param name="propertyFactory">Factory invoked to create the property</param>
+		/// <returns>Created property or null on error</returns>
+		public readonly UhtProperty? ParseTemplateObjectProperty(UhtTemplateObjectMode mode, Func<UhtClass, int, UhtProperty?> propertyFactory)
+		{
+			int typeStartPos = TokenReader.PeekToken().InputStartPos;
+			UhtClass? referencedClass = ParseTemplateObject(mode);
+			if (referencedClass == null)
+			{
+				return null;
+			}
+
+			return propertyFactory(referencedClass, typeStartPos);
+		}
+
 		/// <summary>
 		/// Parse a template type
 		/// </summary>
