Initial UHT support for Verse classes and interfaces defined in UHT

#rnx

diff --git a/Engine/Source/Programs/Shared/EpicGames.UHT/Exporters/CodeGen/UhtCodeGenerator.cs b/Engine/Source/Programs/Shared/EpicGames.UHT/Exporters/CodeGen/UhtCodeGenerator.cs
--- a/Engine/Source/Programs/Shared/EpicGames.UHT/Exporters/CodeGen/UhtCodeGenerator.cs
+++ b/Engine/Source/Programs/Shared/EpicGames.UHT/Exporters/CodeGen/UhtCodeGenerator.cs
@@ -42,10 +42,7 @@ namespace EpicGames.UHT.Exporters.CodeGen
 			public bool NeedsFastArrayHeaders { get; set; }
 			public bool NeedsVerseValue { get; set; }
 			public bool NeedsVerseString { get; set; }
-			public bool NeedsVerseClass { get; set; }
-			public bool NeedsVerseStruct { get; set; }
-			public bool NeedsVerseEnum { get; set; }
-			public bool NeedsVerseInterop { get; set; }
+			public bool NeedsVerseCodeGen { get; set; }
 			public bool NeedsVerseField { get; set; }
 			public bool NeedsVersePointer { get; set; }
 		}
@@ -300,8 +297,6 @@ namespace EpicGames.UHT.Exporters.CodeGen
 					InitObjectInfo(builder, package, ref headerInfo, obj);
 				}
 			}
-
-			headerInfo.NeedsVerseInterop = headerInfo.NeedsVerseStruct;
 		}
 
 		private void InitObjectInfo(StringBuilder builder, UhtPackage package, ref HeaderInfo headerInfo, UhtObject obj)
@@ -380,7 +375,7 @@ namespace EpicGames.UHT.Exporters.CodeGen
 			{
 				if (enumObj.IsVerseField)
 				{
-					headerInfo.NeedsVerseEnum = true;
+					headerInfo.NeedsVerseCodeGen = true;
 				}
 			}
 
@@ -418,7 +413,7 @@ namespace EpicGames.UHT.Exporters.CodeGen
 		{
 			if (structObj.IsVerseField)
 			{
-				headerInfo.NeedsVerseStruct = true;
+				headerInfo.NeedsVerseCodeGen = true;
 			}
 			foreach (UhtType childType in structObj.Children)
 			{
diff --git a/Engine/Source/Programs/Shared/EpicGames.UHT/Exporters/CodeGen/UhtHeaderCodeGenerator.cs b/Engine/Source/Programs/Shared/EpicGames.UHT/Exporters/CodeGen/UhtHeaderCodeGenerator.cs
--- a/Engine/Source/Programs/Shared/EpicGames.UHT/Exporters/CodeGen/UhtHeaderCodeGenerator.cs
+++ b/Engine/Source/Programs/Shared/EpicGames.UHT/Exporters/CodeGen/UhtHeaderCodeGenerator.cs
@@ -19,6 +19,7 @@ namespace EpicGames.UHT.Exporters.CodeGen
 		public const string SparseDataPropertyAccessorsMacroSuffix = "SPARSE_DATA_PROPERTY_ACCESSORS";
 		public const string RpcWrappersMacroSuffix = "RPC_WRAPPERS";
 		public const string RpcWrappersNoPureDeclsMacroSuffix = "RPC_WRAPPERS_NO_PURE_DECLS";
+		public const string VerseNativeCallableWrappersMacroSuffix = "VERSE_NATIVE_CALLABLE_WRAPPERS";
 		public const string AccessorsMacroSuffix = "ACCESSORS";
 		public const string ArchiveSerializerMacroSuffix = "ARCHIVESERIALIZER";
 		public const string StandardConstructorsMacroSuffix = "STANDARD_CONSTRUCTORS";
@@ -34,6 +35,7 @@ namespace EpicGames.UHT.Exporters.CodeGen
 		public const string PrologMacroSuffix = "PROLOG";
 		public const string DelegateMacroSuffix = "DELEGATE";
 		public const string AutoGettersSettersMacroSuffix = "AUTOGETTERSETTER_DECLS";
+		public const string VInterfacesMacroSuffix = "VINTERFACES";
 		#endregion
 
 		public readonly UhtHeaderFile HeaderFile;
@@ -446,7 +448,8 @@ namespace EpicGames.UHT.Exporters.CodeGen
 
 		protected static bool IsRpcFunction(UhtFunction function)
 		{
-			return function.FunctionFlags.HasAnyFlags(EFunctionFlags.Native) &&
+			return
+				function.FunctionFlags.HasAnyFlags(EFunctionFlags.Native) &&
 				!function.FunctionExportFlags.HasAnyFlags(UhtFunctionExportFlags.CustomThunk);
 		}
 
@@ -462,6 +465,12 @@ namespace EpicGames.UHT.Exporters.CodeGen
 		/// <returns>True if the glue version of the function should be exported.</returns>
 		protected static bool ShouldExportFunction(UhtFunction function)
 		{
+			// Always create a stub for verse native callables
+			if (function.FunctionExportFlags.HasAnyFlags(UhtFunctionExportFlags.VerseNativeCallable))
+			{
+				return true;
+			}
+
 			// export any script stubs for native functions declared in interface classes
 			bool isBlueprintNativeEvent = function.FunctionFlags.HasAllFlags(EFunctionFlags.BlueprintEvent | EFunctionFlags.Native);
 			if (!isBlueprintNativeEvent)
@@ -699,5 +708,160 @@ namespace EpicGames.UHT.Exporters.CodeGen
 				throw new UhtException(type, "Can not use given type to create a macro");
 			}
 		}
+
+		/// <summary>
+		/// Append just the overall return type of the native callable verse function
+		/// </summary>
+		/// <param name="builder">Destination builder</param>
+		/// <param name="function">Function in question</param>
+		/// <returns>Builder</returns>
+		public static StringBuilder AppendVerseNativeCallableReturnType(this StringBuilder builder, UhtFunction function)
+		{
+			UhtProperty? returnProperty = function.ReturnProperty;
+			// if (coroutine) "FVerseTask"
+			if (returnProperty == null)
+			{
+				builder.Append("void ");
+			}
+			else if (returnProperty is UhtEnumProperty enumProperty && enumProperty.Enum == function.Session.EVerseNativeCallResult)
+			{
+				builder.Append("EVerseNativeCallResult ");
+			}
+			else
+			{
+				builder.Append("TIfNoRuntimeError<").AppendTokens(returnProperty.TypeTokens.AllTokens).Append("> ");
+			}
+			return builder;
+		}
+
+		/// <summary>
+		/// Append the complete signature of the native callable verse function
+		/// </summary>
+		/// <param name="builder">Destination builder</param>
+		/// <param name="classObj">If supplied, used to append the ClassSourceName:: part of the declaration</param>
+		/// <param name="function">Function in question</param>
+		/// <param name="classNameSuffix">Appended to the end of the class name.  Used to generate the Proxy implementations</param>
+		/// <returns>Builder</returns>
+		public static StringBuilder AppendVerseNativeCallableSignature(this StringBuilder builder, UhtClass? classObj, UhtFunction function, string classNameSuffix)
+		{
+			builder.AppendVerseNativeCallableReturnType(function);
+
+			if (classObj != null)
+			{
+				builder.AppendClassSourceNameOrInterfaceName(classObj).Append(classNameSuffix).Append("::");
+			}
+			builder.Append(function.SourceName).Append("(V_EXEC_CONTEXT_PARAM");
+			foreach (UhtType type in function.ParameterProperties.Span)
+			{
+				if (type is UhtProperty property)
+				{
+					builder.Append(", ").AppendTokens(property.TypeTokens.AllTokens).Append(' ').Append(property.SourceName);
+				}
+			}
+			builder.Append(')');
+			return builder;
+		}
+
+		public static StringBuilder AppendVerseNativeCallableTypeDef(this StringBuilder builder, UhtFunction function)
+		{
+			UhtProperty? returnProperty = function.ReturnProperty;
+			bool isIncompleteReturn = returnProperty is UhtEnumProperty enumProperty && enumProperty.Enum == function.Session.EVerseNativeCallResult;
+			if (!isIncompleteReturn && function.FunctionExportFlags.HasAnyFlags(UhtFunctionExportFlags.VerseDecides) && returnProperty is UhtOptionalProperty optionalReturn)
+			{
+				returnProperty = optionalReturn.ValueProperty;
+			}
+
+			builder.Append("TVerseFunction<");
+			// if (coroutine) "FVerseTask"
+			if (returnProperty == null)
+			{
+				builder.Append("void");
+			}
+			else if (isIncompleteReturn)
+			{
+				builder.Append("EVerseNativeCallResult");
+			}
+			else
+			{
+				builder.AppendTokens(returnProperty.TypeTokens.MyTokens);
+			}
+			builder.Append('(').AppendVerseFunctionArgs(function, false, true, property => builder.AppendTokens(property.TypeTokens.MyTokens), () => builder.Append("FDecidesContext")).Append(")>");
+			return builder;
+		}
+
+		public static StringBuilder AppendVerseFunctionArgs(this StringBuilder builder, UhtFunction function, bool startWithSeparator, bool decidesBeforeReturn, Action<UhtProperty> propertyAction, Action decidesAction)
+		{
+			bool insertSeparator = startWithSeparator;
+			if (function.HasVerseIncompleteReturn)
+			{
+				if (decidesBeforeReturn && function.FunctionExportFlags.HasAnyFlags(UhtFunctionExportFlags.VerseDecides))
+				{
+					if (insertSeparator)
+					{
+						builder.Append(", ");
+					}
+					insertSeparator = true;
+					decidesAction();
+				}
+				{
+					if (function.ParameterProperties.Span[^1] is UhtProperty property)
+					{
+						if (insertSeparator)
+						{
+							builder.Append(", ");
+						}
+						insertSeparator = true;
+						propertyAction(property);
+					}
+				}
+				if (!decidesBeforeReturn && function.FunctionExportFlags.HasAnyFlags(UhtFunctionExportFlags.VerseDecides))
+				{
+					if (insertSeparator)
+					{
+						builder.Append(", ");
+					}
+					insertSeparator = true;
+					decidesAction();
+				}
+				foreach (UhtType type in function.ParameterProperties.Span[..^1])
+				{
+					if (type is UhtProperty property)
+					{
+						if (insertSeparator)
+						{
+							builder.Append(", ");
+						}
+						insertSeparator = true;
+						propertyAction(property);
+					}
+				}
+			}
+			else
+			{
+				if (function.FunctionExportFlags.HasAnyFlags(UhtFunctionExportFlags.VerseDecides))
+				{
+					if (insertSeparator)
+					{
+						builder.Append(", ");
+					}
+					insertSeparator = true;
+					decidesAction();
+				}
+
+				foreach (UhtType type in function.ParameterProperties.Span)
+				{
+					if (type is UhtProperty property)
+					{
+						if (insertSeparator)
+						{
+							builder.Append(", ");
+						}
+						insertSeparator = true;
+						propertyAction(property);
+					}
+				}
+			}
+			return builder;
+		}
 	}
 }
diff --git a/Engine/Source/Programs/Shared/EpicGames.UHT/Exporters/CodeGen/UhtHeaderCodeGeneratorCppFile.cs b/Engine/Source/Programs/Shared/EpicGames.UHT/Exporters/CodeGen/UhtHeaderCodeGeneratorCppFile.cs
--- a/Engine/Source/Programs/Shared/EpicGames.UHT/Exporters/CodeGen/UhtHeaderCodeGeneratorCppFile.cs
+++ b/Engine/Source/Programs/Shared/EpicGames.UHT/Exporters/CodeGen/UhtHeaderCodeGeneratorCppFile.cs
@@ -59,13 +59,10 @@ namespace EpicGames.UHT.Exporters.CodeGen
 				{
 					includesToAdd.Add("Net/Serialization/FastArraySerializerImplementation.h");
 				}
-				if (headerInfo.NeedsVerseClass || headerInfo.NeedsVerseStruct || headerInfo.NeedsVerseEnum)
+				if (headerInfo.NeedsVerseCodeGen)
 				{
 					includesToAdd.Add("VerseVM/VVMUECodeGen.h");
-				}
-				if (headerInfo.NeedsVerseEnum)
-				{
-					includesToAdd.Add("VerseInteropMacros.h");
+					includesToAdd.Add("VerseInteropUtils.h");
 				}
 
 				foreach (UhtType type in HeaderFile.Children)
@@ -220,6 +217,14 @@ namespace EpicGames.UHT.Exporters.CodeGen
 							using UhtCodeBlockComment blockComment = new(builder, field);
 							AppendClass(builder, classObj, functions);
 							GetRegistrations(packageRegistrations, field).Classes.Add(classObj);
+
+							if (classObj.ClassType == UhtClassType.Interface && classObj.IsVerseField)
+							{
+								foreach (UhtClass baseClass in classObj.FlattenedVerseInterfaces)
+								{
+									AppendNativeInterfaceVerseProxyFunctions(builder, classObj, baseClass.AlternateObject as UhtClass);
+								}
+							}
 						}
 					}
 				}
@@ -844,7 +849,7 @@ namespace EpicGames.UHT.Exporters.CodeGen
 
 		private StringBuilder AppendDelegate(StringBuilder builder, UhtFunction function)
 		{
-			AppendFunction(builder, function, false);
+			AppendFunction(builder, function, false, false);
 
 			int tabs = 0;
 			string strippedFunctionName = function.StrippedFunctionName;
@@ -866,7 +871,7 @@ namespace EpicGames.UHT.Exporters.CodeGen
 			return builder;
 		}
 
-		private StringBuilder AppendFunction(StringBuilder builder, UhtFunction function, bool isNoExport)
+		private StringBuilder AppendFunction(StringBuilder builder, UhtFunction function, bool isNoExport, bool isVerseNativeCallable)
 		{
 			const string MetaDataParamsName = "Function_MetaDataParams";
 			string singletonName = GetSingletonName(function, true);
@@ -916,12 +921,13 @@ namespace EpicGames.UHT.Exporters.CodeGen
 				{
 					AppendPropertiesDefs(builder, context, properties, 0);
 
+					string functionEngineName = isVerseNativeCallable ? function.EncodedVersePath : function.EngineName;
 					builder
 						.Append("const UECodeGen_Private::").AppendFunctionParamsType(function).Append(' ')
 						.Append(staticsName).Append("::FuncParams = { { ")
 						.Append("(UObject*(*)())").Append(GetFunctionOuterFunc(function)).Append(", ")
 						.Append(GetSingletonName(function.SuperFunction, true)).Append(", ")
-						.AppendUTF8LiteralString(function.EngineName).Append(", ");
+						.AppendUTF8LiteralString(functionEngineName).Append(", ");
 
 					builder.AppendArrayPtrLine(properties, UhtDefineScopeNames.Standard, staticsName, "PropPointers", 0, ", ");
 					builder.AppendArrayCountLine(properties, UhtDefineScopeNames.Standard, staticsName, "PropPointers", 0, ", ");
@@ -960,11 +966,28 @@ namespace EpicGames.UHT.Exporters.CodeGen
 						.AppendMetaDataParams(function, staticsName, MetaDataParamsName)
 						.Append("}, ");
 
-					if (function.FunctionType == UhtFunctionType.SparseDelegate)
+					switch (function.FunctionType)
 					{
-						builder
-							.AppendUTF8LiteralString(function.SparseOwningClassName).Append(", ")
-							.AppendUTF8LiteralString(function.SparseDelegateName).Append(", ");
+						case UhtFunctionType.Function:
+							if (function.IsVerseField)
+							{
+								builder.Append("\r\n");
+								builder.Append("#if WITH_VERSE_VM\r\n");
+								builder.AppendUTF8LiteralString(function.VersePath).Append(",\r\n");
+								builder.Append("#else\r\n");
+								builder.AppendUTF8LiteralString(function.EncodedVersePath).Append(",\r\n");
+								builder.Append("#endif\r\n");
+							}
+							break;
+
+						case UhtFunctionType.Delegate:
+							break;
+
+						case UhtFunctionType.SparseDelegate:
+							builder
+								.AppendUTF8LiteralString(function.SparseOwningClassName).Append(", ")
+								.AppendUTF8LiteralString(function.SparseDelegateName).Append(", ");
+							break;
 					}
 
 					builder.Append(" };\r\n");
@@ -982,7 +1005,7 @@ namespace EpicGames.UHT.Exporters.CodeGen
 					builder.Append("\tstatic UFunction* ReturnFunction = nullptr;\r\n");
 					builder.Append("\tif (!ReturnFunction)\r\n");
 					builder.Append("\t{\r\n");
-					builder.Append("\t\tUECodeGen_Private::").AppendFunctionConstructorName(function).Append("(&ReturnFunction, ").Append(staticsName).Append("::FuncParams);\r\n");
+					builder.Append("\t\t").AppendFunctionConstructorName(function).Append("(&ReturnFunction, ").Append(staticsName).Append("::FuncParams);\r\n");
 					builder.Append("\t}\r\n");
 					builder.Append("\treturn ReturnFunction;\r\n");
 					builder.Append("}\r\n");
@@ -1081,6 +1104,7 @@ namespace EpicGames.UHT.Exporters.CodeGen
 			bool isNotDelegate = !function.FunctionFlags.HasAnyFlags(EFunctionFlags.Delegate);
 			bool isRpc = IsRpcFunction(function) && ShouldExportFunction(function);
 			bool isCallback = IsCallbackFunction(function);
+			bool isVerseNativeCallable = function.FunctionExportFlags.HasAnyFlags(UhtFunctionExportFlags.VerseNativeCallable);
 			if (isNotDelegate || isRpc || isCallback)
 			{
 				using UhtCodeBlockComment blockCommand = new(builder, classObj, function);
@@ -1095,7 +1119,16 @@ namespace EpicGames.UHT.Exporters.CodeGen
 				}
 				if (!function.FunctionFlags.HasAnyFlags(EFunctionFlags.Delegate))
 				{
-					AppendFunction(builder, function, classObj.ClassExportFlags.HasAnyFlags(UhtClassExportFlags.NoExport));
+					AppendFunction(builder, function, classObj.ClassExportFlags.HasAnyFlags(UhtClassExportFlags.NoExport), false);
+					if (isVerseNativeCallable)
+					{
+						AppendVerseNativeCallableFunctionStub(builder, classObj, function);
+						AppendVerseNativeCallableFunction(builder, classObj, function, "");
+						if (classObj.ClassFlags.HasAnyFlags(EClassFlags.Interface))
+						{
+							AppendVerseNativeCallableFunction(builder, classObj, function, UhtNames.VerseProxySuffix);
+						}
+					}
 				}
 				if (isRpc)
 				{
@@ -1109,6 +1142,171 @@ namespace EpicGames.UHT.Exporters.CodeGen
 			return builder;
 		}
 
+		private static StringBuilder AppendNativeInterfaceVerseProxyFunctions(StringBuilder builder, UhtClass classObj, UhtClass? funcClassObj)
+		{
+			if (funcClassObj != null)
+			{
+				UhtUsedDefineScopes<UhtFunction> verseNativeCallableFunctions = new(funcClassObj.Functions.Where(x => x.FunctionExportFlags.HasAnyFlags(UhtFunctionExportFlags.VerseNativeCallable)));
+				builder.Append("\r\n");
+				builder.Append("// Verse native callable thunks for ").Append(classObj.SourceName).Append("\r\n");
+				AppendVerseNativeCallableProxyFunctions(builder, classObj, verseNativeCallableFunctions);
+			}
+			return builder;
+		}
+
+		private static StringBuilder AppendVerseNativeCallableProxyFunctions(StringBuilder builder, UhtClass classObj, UhtUsedDefineScopes<UhtFunction> verseNativeCallableFunctions)
+		{
+			return builder.AppendInstances(verseNativeCallableFunctions, UhtDefineScopeNames.Standard,
+				(StringBuilder builder, UhtFunction function) =>
+				{
+					AppendVerseNativeCallableFunction(builder, classObj, function, UhtNames.VerseProxySuffix);
+				});
+		}
+
+		private static StringBuilder AppendVerseNativeCallableFunctionStub(StringBuilder builder, UhtClass classObj, UhtFunction function)
+		{
+			UhtProperty? returnProperty = function.ReturnProperty;
+			if (returnProperty != null)
+			{
+				builder.AppendTokens(returnProperty.TypeTokens.AllTokens).Append(' ');
+			}
+			else
+			{
+				builder.Append("void ");
+			}
+
+			builder.AppendClassSourceNameOrInterfaceName(classObj).Append("::").Append(function.SourceName).Append('(');
+			bool needsComma = false;
+			foreach (UhtType type in function.ParameterProperties.Span)
+			{
+				if (type is UhtProperty property)
+				{
+					if (needsComma)
+					{
+						builder.Append(", ");
+					}
+					needsComma = true;
+					builder.AppendTokens(property.TypeTokens.AllTokens).Append(' ').Append(property.SourceName);
+				}
+			}
+			builder.Append(")\r\n");
+			builder.Append("{\r\n");
+			if (returnProperty != null)
+			{
+				builder.Append("\treturn ");
+			}
+			else
+			{
+				builder.Append('\t');
+			}
+			builder.Append(function.SourceName).Append("(verse::FExecutionContext::GetActiveContext()");
+			foreach (UhtType type in function.ParameterProperties.Span)
+			{
+				if (type is UhtProperty property)
+				{
+					builder.Append(", std::forward<decltype(").Append(property.SourceName).Append(")>(").Append(property.SourceName).Append(')');
+				}
+			}
+			if (returnProperty != null)
+			{
+				if (returnProperty is UhtEnumProperty enumProperty && enumProperty.Enum == function.Session.EVerseNativeCallResult)
+				{
+					builder.Append(");\r\n");
+				}
+				else
+				{
+					builder.Append(").GetValue();\r\n");
+				}
+			}
+			else
+			{
+				builder.Append(");\r\n");
+			}
+			builder.Append("}\r\n");
+			return builder;
+		}
+
+		private static StringBuilder AppendVerseNativeCallableFunction(StringBuilder builder, UhtClass classObj, UhtFunction function, string classNameSuffix)
+		{
+			UhtProperty? returnProperty = function.ReturnProperty;
+			bool isIncompleteReturn = returnProperty is UhtEnumProperty enumProperty && enumProperty.Enum == function.Session.EVerseNativeCallResult;
+
+			builder.AppendVerseNativeCallableSignature(classObj, function, classNameSuffix).Append("\r\n");
+			builder.Append("{\r\n");
+
+			// Result value declaration
+			if (isIncompleteReturn)
+			{
+				builder.Append("\tEVerseNativeCallResult Result;\r\n");
+			}
+			else if (returnProperty != null)
+			{
+				builder.Append("\tTIfNoRuntimeError<").AppendTokens(returnProperty.TypeTokens.AllTokens).Append("> Result;\r\n");
+			}
+
+			// The bodies must differ based on the VM
+			builder.Append("#if WITH_VERSE_VM\r\n");
+			AppendVerseNativeCallableFunctionBody(builder, classObj, function, classNameSuffix, true);
+			builder.Append("#else\r\n");
+			AppendVerseNativeCallableFunctionBody(builder, classObj, function, classNameSuffix, false);
+			builder.Append("#endif\r\n");
+
+			if (returnProperty != null)
+			{
+				builder.Append("\treturn Result;\r\n");
+			}
+			builder.Append("}\r\n");
+			return builder;
+		}
+
+		private static StringBuilder AppendVerseNativeCallableFunctionBody(StringBuilder builder, UhtClass classObj, UhtFunction function, string classNameSuffix, bool isVerseVM)
+		{
+			builder.Append(isVerseVM ? "\tAutoRTFM::Open([&] {\r\n" : "\t{\r\n");
+
+			// Callee type declaration
+			builder.Append("\t\tusing CalleeType = ").AppendVerseNativeCallableTypeDef(function).Append(";\r\n");
+
+			// Callee paths
+			string versePath = function.VersePath;
+			if (isVerseVM)
+			{
+				builder.Append("\t\tconst char* CalleePath = \"").Append(versePath).Append("\";\r\n");
+			}
+			else
+			{
+				builder.Append("\t\tconst TCHAR* CalleePath = TEXT(\"").AppendEncodedVerseName(versePath).Append("\");\r\n");
+			}
+
+			// Callee declaration
+			if (classObj.ClassType == UhtClassType.Interface)
+			{
+				builder.Append("\t\tCalleeType Callee{_V_EXEC_CONTEXT_PARAM_NAME, _getUObject(), CalleePath};\r\n");
+			}
+			else
+			{
+				builder.Append("\t\tCalleeType Callee{_V_EXEC_CONTEXT_PARAM_NAME, V_SAFE_INSTANCE_CONTEXT(\"").Append(classObj.VerseName).Append("\"), CalleePath};\r\n");
+			}
+			/* MODULE
+			else
+			{
+				builder.Append("\t\tCalleeType Callee{_V_EXEC_CONTEXT_PARAM_NAME, V_SAFE_STATIC_CONTEXT(\"").Append(classObj.VerseName).Append("\"), CalleePath};\r\n");
+			}
+			 */
+
+			// Invocation
+			builder.Append("\t\t");
+			if (function.ReturnProperty != null)
+			{
+				builder.Append("Result = ");
+			}
+			builder.Append("Callee(_V_EXEC_CONTEXT_PARAM_NAME");
+			builder.AppendVerseFunctionArgs(function, true, false, property => builder.Append(property.SourceName), () => builder.Append("FDecidesContext(FDecidesContext::EDefaultContruct::UnsafeDoNotUse)"));
+			builder.Append(");\r\n");
+
+			builder.Append(isVerseVM ? "\t});\r\n" : "\t}\r\n");
+			return builder;
+		}
+
 		private StringBuilder AppendClassFunctionCallback(StringBuilder builder, UhtClass classObj, UhtFunction function)
 		{
 			// Net response functions don't go into the VM
@@ -1634,7 +1832,7 @@ namespace EpicGames.UHT.Exporters.CodeGen
 
 			PropertyMemberContextImpl context = new(CodeGenerator, classObj, classObj.SourceName, staticsName);
 
-			bool hasInterfaces = classObj.Bases.Any(x => x is UhtClass baseClass && baseClass.ClassFlags.HasAnyFlags(EClassFlags.Interface));
+			bool hasInterfaces = classObj.Bases.Any(x => x is UhtClass baseClass && baseClass.ClassFlags.HasAnyFlags(EClassFlags.Interface)) || classObj.FlattenedVerseInterfaces.Count > 0;
 
 			// This block of code is the contents of IMPLEMENT_CLASS_NO_AUTO_REGISTRATION 
 			builder.Append("FClassRegistrationInfo ").Append(registrationName).Append(";\r\n");
@@ -1747,24 +1945,17 @@ namespace EpicGames.UHT.Exporters.CodeGen
 
 					foreach (UhtStruct structObj in classObj.Bases)
 					{
-						if (structObj is UhtClass baseClass)
+						if (structObj is UhtClass interfaceObj)
 						{
-							if (baseClass.ClassFlags.HasAnyFlags(EClassFlags.Interface))
-							{
-								builder
-									.Append("\t{ ")
-									.Append(GetSingletonName(baseClass.AlternateObject, false))
-									.Append(", (int32)VTABLE_OFFSET(")
-									.Append(classObj.SourceName)
-									.Append(", ")
-									.Append(baseClass.SourceName)
-									.Append("), false }, ")
-									.AppendObjectHash(classObj, context, baseClass.AlternateObject)
-									.Append("\r\n");
-							}
+							AppendImplementedInterfaceParam(builder, context, classObj, interfaceObj);
 						}
 					}
 
+					foreach (UhtClass interfaceObj in classObj.FlattenedVerseInterfaces)
+					{
+						AppendImplementedInterfaceParam(builder, context, classObj, interfaceObj);
+					}
+
 					builder.Append("};\r\n");
 				}
 
@@ -1912,7 +2103,8 @@ namespace EpicGames.UHT.Exporters.CodeGen
 			builder.Append("void ").Append(classObj.SourceName).Append("::StaticRegisterNatives").Append(classObj.SourceName).Append("()\r\n");
 			builder.Append("{\r\n");
 
-			UhtUsedDefineScopes<UhtFunction> functions = new(classObj.Functions.Where(x => x.FunctionFlags.HasExactFlags(EFunctionFlags.Native | EFunctionFlags.NetRequest, EFunctionFlags.Native)));
+			UhtUsedDefineScopes<UhtFunction> functions = new(classObj.Functions.Where(
+				x => x.FunctionFlags.HasExactFlags(EFunctionFlags.Native | EFunctionFlags.NetRequest, EFunctionFlags.Native)));
 			functions.Instances.Sort((x, y) => StringComparerUE.OrdinalIgnoreCase.Compare(x.EngineName, y.EngineName));
 
 			if (!functions.IsEmpty)
@@ -1947,6 +2139,39 @@ namespace EpicGames.UHT.Exporters.CodeGen
 			return builder;
 		}
 
+		private StringBuilder AppendImplementedInterfaceParam(StringBuilder builder, PropertyMemberContextImpl context, UhtClass classObj, UhtClass interfaceObj)
+		{
+			if (interfaceObj.ClassFlags.HasAnyFlags(EClassFlags.Interface))
+			{
+				builder
+					.Append("\t{ ")
+					.Append(GetSingletonName(interfaceObj.AlternateObject, false));
+				if (interfaceObj.IsVerseField)
+				{
+					builder
+						.Append(", 0, true");
+					if (classObj.VerseInterfaces.Contains(interfaceObj))
+					{
+						builder.Append(", true");
+					}
+				}
+				else
+				{
+					builder
+						.Append(", (int32)VTABLE_OFFSET(")
+						.Append(classObj.SourceName)
+						.Append(", ")
+						.Append(interfaceObj.SourceName)
+						.Append("), false");
+				}
+				builder
+					.Append(" }, ")
+					.AppendObjectHash(classObj, context, interfaceObj.AlternateObject)
+					.Append("\r\n");
+			}
+			return builder;
+		}
+
 		private static StringBuilder AppendFunctionThunk(StringBuilder builder, UhtFunction function)
 		{
 			// Export the GET macro for the parameters
@@ -2135,7 +2360,7 @@ namespace EpicGames.UHT.Exporters.CodeGen
 			switch (function.FunctionType)
 			{
 				case UhtFunctionType.Function:
-					builder.Append("FFunctionParams");
+					builder.Append(function.IsVerseField ? "FVerseFunctionParams" : "FFunctionParams");
 					break;
 				case UhtFunctionType.Delegate:
 					builder.Append("FDelegateFunctionParams");
@@ -2158,13 +2383,13 @@ namespace EpicGames.UHT.Exporters.CodeGen
 			switch (function.FunctionType)
 			{
 				case UhtFunctionType.Function:
-					builder.Append("ConstructUFunction");
+					builder.Append(function.IsVerseField ? "Verse::CodeGen::Private::ConstructUVerseFunction" : "UECodeGen_Private::ConstructUFunction");
 					break;
 				case UhtFunctionType.Delegate:
-					builder.Append("ConstructUDelegateFunction");
+					builder.Append("UECodeGen_Private::ConstructUDelegateFunction");
 					break;
 				case UhtFunctionType.SparseDelegate:
-					builder.Append("ConstructUSparseDelegateFunction");
+					builder.Append("UECodeGen_Private::ConstructUSparseDelegateFunction");
 					break;
 			}
 			return builder;
diff --git a/Engine/Source/Programs/Shared/EpicGames.UHT/Exporters/CodeGen/UhtHeaderCodeGeneratorHFile.cs b/Engine/Source/Programs/Shared/EpicGames.UHT/Exporters/CodeGen/UhtHeaderCodeGeneratorHFile.cs
--- a/Engine/Source/Programs/Shared/EpicGames.UHT/Exporters/CodeGen/UhtHeaderCodeGeneratorHFile.cs
+++ b/Engine/Source/Programs/Shared/EpicGames.UHT/Exporters/CodeGen/UhtHeaderCodeGeneratorHFile.cs
@@ -78,7 +78,7 @@ namespace EpicGames.UHT.Exporters.CodeGen
 				{
 					builder.Append("#include \"VerseValue.h\"\r\n");
 				}
-				if (headerInfo.NeedsVerseInterop)
+				if (headerInfo.NeedsVerseCodeGen)
 				{
 					builder.Append("#include \"VerseInteropMacros.h\"\r\n");
 				}
@@ -138,6 +138,18 @@ namespace EpicGames.UHT.Exporters.CodeGen
 					}
 				}
 
+				foreach (UhtField field in HeaderFile.References.ExportTypes)
+				{
+					if (field is UhtClass classObj)
+					{
+						if (!classObj.ClassFlags.HasAnyFlags(EClassFlags.Intrinsic) && classObj.ClassType == UhtClassType.Interface && classObj.IsVerseField)
+						{
+							using UhtCodeBlockComment blockComment = new(builder, field, "Verse Interface Proxy");
+							AppendNativeInterfaceVerseProxy(builder, classObj);
+						}
+					}
+				}
+
 				builder.Append("\r\n");
 				builder.Append("#undef CURRENT_FILE_ID\r\n");
 				builder.Append("#define CURRENT_FILE_ID ").Append(headerInfo.FileId).Append("\r\n");
@@ -765,6 +777,12 @@ namespace EpicGames.UHT.Exporters.CodeGen
 			// Output the RPC methods
 			AppendRpcFunctions(builder, classObj, alternateUsesLegacy, rpcFunctions);
 
+			// Collect the verse native callable functions
+			UhtUsedDefineScopes<UhtFunction> verseNativeCallableFunctions = new(classObj.Functions.Where(x => x.FunctionExportFlags.HasAnyFlags(UhtFunctionExportFlags.VerseNativeCallable)));
+
+			// Output the verse native callback functions
+			AppendVerseNativeCallableFunctions(builder, classObj, verseNativeCallableFunctions);
+
 			// Output property accessors
 			IEnumerable<UhtProperty> getterSetterProperties = classObj.Properties.Where(x => x.PropertyExportFlags.HasAnyFlags(UhtPropertyExportFlags.GetterFound | UhtPropertyExportFlags.SetterFound));
 			AppendPropertyAccessors(builder, classObj, getterSetterProperties);
@@ -836,8 +854,14 @@ namespace EpicGames.UHT.Exporters.CodeGen
 				}
 
 				AppendProlog(builder, classObj);
-				AppendGeneratedBodyMacroBlock(builder, classObj, nativeInterface!, alternateUsesLegacy, rpcFunctions, autoGetterSetterProperties, 
-					sparseScriptStructs.Any(), sparseProperties, getterSetterProperties, hasCallbacks, null);
+				AppendGeneratedBodyMacroBlock(builder, classObj, nativeInterface!, alternateUsesLegacy, rpcFunctions, verseNativeCallableFunctions,
+					autoGetterSetterProperties, sparseScriptStructs.Any(), sparseProperties, getterSetterProperties, hasCallbacks, null);
+
+				// For interfaces, the VINTERFACES macro generates nothing
+				if (classObj.VerseInterfacesLine != 0)
+				{
+					using UhtMacroCreator macro = new(builder, this, classObj.VerseInterfacesLine, VInterfacesMacroSuffix);
+				}
 			}
 			else
 			{
@@ -861,8 +885,9 @@ namespace EpicGames.UHT.Exporters.CodeGen
 				AppendFieldNotify(builder, classObj);
 				AppendAutoGettersSetters(builder, classObj, autoGetterSetterProperties);
 				AppendProlog(builder, classObj);
-				AppendGeneratedBodyMacroBlock(builder, classObj, classObj, usesLegacy, rpcFunctions, autoGetterSetterProperties, 
-					sparseScriptStructs.Any(), sparseProperties, getterSetterProperties, hasCallbacks, usesLegacy ? "GENERATED_UCLASS_BODY" : null);
+				AppendGeneratedBodyMacroBlock(builder, classObj, classObj, usesLegacy, rpcFunctions, verseNativeCallableFunctions,
+					autoGetterSetterProperties, sparseScriptStructs.Any(), sparseProperties, getterSetterProperties, hasCallbacks, 
+					usesLegacy ? "GENERATED_UCLASS_BODY" : null);
 			}
 
 			// For back compat we still need to forward declare the class here
@@ -873,7 +898,10 @@ namespace EpicGames.UHT.Exporters.CodeGen
 			{
 				if (!classObj.FieldExportFlags.HasAnyFlags(UhtFieldExportFlags.NoVerseAlias))
 				{
-					builder.Append("namespace verse { using ").Append(classObj.EngineName).Append(" = ::").Append(classObj.SourceName).Append("; }\r\n");
+					if (classObj.ClassType == UhtClassType.Class)
+					{
+						builder.Append("namespace verse { using ").Append(classObj.EngineName).Append(" = ").AppendClassSourceNameOrInterfaceProxyName(classObj).Append("; }\r\n");
+					}
 				}
 			}
 
@@ -882,6 +910,52 @@ namespace EpicGames.UHT.Exporters.CodeGen
 			return builder;
 		}
 
+		private static StringBuilder AppendNativeInterfaceVerseProxy(StringBuilder builder, UhtClass classObj)
+		{
+			builder.Append("class ").AppendClassSourceNameOrInterfaceProxyName(classObj).Append("\r\n");
+			builder.Append("{\r\n");
+			builder.Append("public:\r\n");
+			builder.Append("\tusing StaticClassType = ").Append(classObj.SourceName).Append(";\r\n");
+			foreach (UhtClass baseClass in classObj.FlattenedVerseInterfaces)
+			{
+				AppendNativeInterfaceVerseProxyFunctions(builder, baseClass.AlternateObject as UhtClass);
+			}
+			AppendNativeInterfaceVerseProxyFunctions(builder, classObj);
+			builder.Append('\t').AppendClassSourceNameOrInterfaceProxyName(classObj).Append("(const TNonNullPtr<UObject>& InSelf) : Self(InSelf) {}\r\n");
+			builder.Append("\tUObject* _getUObject() const { return Self; }\r\n");
+			builder.Append('\t').AppendClassSourceNameOrInterfaceProxyName(classObj).Append("* operator->() { return this; }\r\n");
+			builder.Append("private:\r\n");
+			builder.Append("\tTNonNullPtr<UObject> Self;\r\n");
+			builder.Append("};\r\n");
+			foreach (UhtClass baseClass in classObj.FlattenedVerseInterfaces)
+			{
+				builder
+					.Append("template<> struct TIsSubinterfaceOf<")
+					.AppendClassSourceNameOrInterfaceProxyName(baseClass)
+					.Append(", ")
+					.AppendClassSourceNameOrInterfaceProxyName(classObj)
+					.Append("> { enum { Value = true }; };\r\n");
+			}
+
+			if (!classObj.FieldExportFlags.HasAnyFlags(UhtFieldExportFlags.NoVerseAlias))
+			{
+				builder.Append("namespace verse { using ").Append(classObj.EngineName).Append(" = ").AppendClassSourceNameOrInterfaceProxyName(classObj).Append("; }\r\n");
+			}
+			return builder;
+		}
+
+		private static StringBuilder AppendNativeInterfaceVerseProxyFunctions(StringBuilder builder, UhtClass? classObj)
+		{
+			if (classObj != null)
+			{
+				UhtUsedDefineScopes<UhtFunction> verseNativeCallableFunctions = new(classObj.Functions.Where(x => x.FunctionExportFlags.HasAnyFlags(UhtFunctionExportFlags.VerseNativeCallable)));
+				builder.Append("\t// Verse native callable thunks for ").Append(classObj.SourceName).Append("\r\n");
+				AppendVerseNativeCallableProxyFunctions(builder, classObj, verseNativeCallableFunctions);
+				builder.Append("\r\n");
+			}
+			return builder;
+		}
+
 		private StringBuilder AppendProlog(StringBuilder builder, UhtClass classObj/*, List<UhtFunction> classbackFunctions*/)
 		{
 			using (UhtMacroCreator macro = new(builder, this, classObj.PrologLineNumber, PrologMacroSuffix))
@@ -1051,6 +1125,37 @@ namespace EpicGames.UHT.Exporters.CodeGen
 			return builder;
 		}
 
+		private StringBuilder AppendVerseNativeCallableFunctions(StringBuilder builder, UhtClass classObj, UhtUsedDefineScopes<UhtFunction> functions)
+		{
+			builder.AppendMultiScopedMacros(functions, UhtDefineScopeNames.WithEditor, this, classObj, VerseNativeCallableWrappersMacroSuffix,
+				(builder, functions) =>
+				{
+					foreach (UhtFunction function in functions)
+					{
+						builder.Append('\t').AppendVerseNativeCallableSignature(null, function, "").Append("; \\\r\n");
+					}
+				});
+			return builder;
+		}
+
+		private static StringBuilder AppendVerseNativeCallableProxyFunctions(StringBuilder builder, UhtClass classObj, UhtUsedDefineScopes<UhtFunction> functions)
+		{
+			if (classObj.SuperClass != null && classObj.SuperClass.IsVerseField && classObj.SuperClass.ClassType == UhtClassType.Interface)
+			{
+				AppendVerseNativeCallableProxyFunctions(builder, classObj.SuperClass, new(classObj.SuperClass.Functions.Where(x => x.FunctionExportFlags.HasAnyFlags(UhtFunctionExportFlags.VerseNativeCallable))));
+			}
+
+			builder.AppendMultiScoped(functions, UhtDefineScopeNames.WithEditor,
+				(builder, functions) =>
+				{
+					foreach (UhtFunction function in functions)
+					{
+						builder.Append('\t').AppendVerseNativeCallableSignature(null, function, "").Append(";\r\n");
+					}
+				});
+			return builder;
+		}
+
 		private StringBuilder AppendPropertyAccessors(StringBuilder builder, UhtClass classObj, IEnumerable<UhtProperty> getterSetterProperties)
 		{
 			if (getterSetterProperties.Any())
@@ -1687,7 +1792,8 @@ namespace EpicGames.UHT.Exporters.CodeGen
 		}
 
 		private StringBuilder AppendGeneratedBodyMacroBlock(StringBuilder builder, UhtClass classObj, UhtClass bodyClassObj, bool isLegacy, 
-			UhtUsedDefineScopes<UhtFunction> rpcFunctions, UhtUsedDefineScopes<UhtProperty> autoGetterSetterProperties, bool hasSparseStructs,
+			UhtUsedDefineScopes<UhtFunction> rpcFunctions, UhtUsedDefineScopes<UhtFunction> verseNativeCallableFunctions,
+			UhtUsedDefineScopes<UhtProperty> autoGetterSetterProperties, bool hasSparseStructs,
 			UhtUsedDefineScopes<UhtProperty> sparseProperties, IEnumerable<UhtProperty> getterSetterProperties, bool hasCallbacks, string? deprecatedMacroName)
 		{
 			bool isInterface = classObj.ClassFlags.HasAnyFlags(EClassFlags.Interface);
@@ -1705,6 +1811,7 @@ namespace EpicGames.UHT.Exporters.CodeGen
 					builder.AppendMultiMacroRefs(sparseProperties, this, classObj, SparseDataPropertyAccessorsMacroSuffix);
 				}
 				builder.AppendMultiMacroRefs(rpcFunctions, this, classObj, isLegacy ? RpcWrappersMacroSuffix : RpcWrappersNoPureDeclsMacroSuffix);
+				builder.AppendMultiMacroRefs(verseNativeCallableFunctions, this, classObj, VerseNativeCallableWrappersMacroSuffix);
 				if (getterSetterProperties.Any())
 				{
 					builder.Append('\t').AppendMacroName(this, classObj, AccessorsMacroSuffix).Append(" \\\r\n");
diff --git a/Engine/Source/Programs/Shared/EpicGames.UHT/Parsers/UhtClassParser.cs b/Engine/Source/Programs/Shared/EpicGames.UHT/Parsers/UhtClassParser.cs
--- a/Engine/Source/Programs/Shared/EpicGames.UHT/Parsers/UhtClassParser.cs
+++ b/Engine/Source/Programs/Shared/EpicGames.UHT/Parsers/UhtClassParser.cs
@@ -102,6 +102,8 @@ namespace EpicGames.UHT.Parsers
 					UhtParserHelpers.ParseInheritance(parentScope.HeaderParser, topScope.Session.Config!, out UhtInheritance inheritance);
 					classObj.SuperIdentifier = inheritance.SuperIdentifier;
 					classObj.BaseIdentifiers = inheritance.BaseIdentifiers;
+					classObj.VerseInterfaceIdentifiers = inheritance.VerseInterfaceIdentifiers;
+					classObj.VerseInterfacesLine = inheritance.VerseInterfacesLine;
 
 					if (apiMacroToken)
 					{
diff --git a/Engine/Source/Programs/Shared/EpicGames.UHT/Parsers/UhtFunctionParser.cs b/Engine/Source/Programs/Shared/EpicGames.UHT/Parsers/UhtFunctionParser.cs
--- a/Engine/Source/Programs/Shared/EpicGames.UHT/Parsers/UhtFunctionParser.cs
+++ b/Engine/Source/Programs/Shared/EpicGames.UHT/Parsers/UhtFunctionParser.cs
@@ -452,6 +452,12 @@ namespace EpicGames.UHT.Parsers
 						// Get function or operator name.
 						function.SourceName = funcNameToken.Value.ToString();
 
+						// If the function is a verse function, then we need to apply the source name to the verse name
+						if (function.IsVerseField)
+						{
+							function.VerseName = function.SourceName;
+						}
+
 						scopeName = $"{scopeName} '{function.SourceName}'";
 						tokenContext.Reset(scopeName);
 
diff --git a/Engine/Source/Programs/Shared/EpicGames.UHT/Parsers/UhtInterfaceClassParser.cs b/Engine/Source/Programs/Shared/EpicGames.UHT/Parsers/UhtInterfaceClassParser.cs
--- a/Engine/Source/Programs/Shared/EpicGames.UHT/Parsers/UhtInterfaceClassParser.cs
+++ b/Engine/Source/Programs/Shared/EpicGames.UHT/Parsers/UhtInterfaceClassParser.cs
@@ -80,6 +80,12 @@ namespace EpicGames.UHT.Parsers
 					// Parse the inheritance
 					UhtParserHelpers.ParseInheritance(parentScope.HeaderParser, topScope.Session.Config!, out UhtInheritance inheritance);
 					classObj.SuperIdentifier = inheritance.SuperIdentifier;
+					classObj.VerseInterfaceIdentifiers = inheritance.VerseInterfaceIdentifiers;
+					classObj.VerseInterfacesLine = inheritance.VerseInterfacesLine;
+					if (inheritance.BaseIdentifiers.Count > 0)
+					{
+						topScope.TokenReader.LogError($"Interfaces can not contain any extra base classes");
+					}
 
 					// Split the source name into the different parts
 					UhtEngineNameParts nameParts = UhtUtilities.GetEngineNameParts(classObj.SourceName);
diff --git a/Engine/Source/Programs/Shared/EpicGames.UHT/Parsers/UhtParserHelpers.cs b/Engine/Source/Programs/Shared/EpicGames.UHT/Parsers/UhtParserHelpers.cs
--- a/Engine/Source/Programs/Shared/EpicGames.UHT/Parsers/UhtParserHelpers.cs
+++ b/Engine/Source/Programs/Shared/EpicGames.UHT/Parsers/UhtParserHelpers.cs
@@ -26,6 +26,16 @@ namespace EpicGames.UHT.Parsers
 		/// </summary>
 		public List<UhtToken[]> BaseIdentifiers { get; } = new();
 
+		/// <summary>
+		/// Collection of verse interface identifiers
+		/// </summary>
+		public List<UhtToken> VerseInterfaceIdentifiers { get; } = new();
+
+		/// <summary>
+		/// Line number that the VIDENTIFIER macro appeared
+		/// </summary>
+		public int VerseInterfacesLine { get; set; } = 0;
+
 		/// <summary>
 		/// Constructor
 		/// </summary>
@@ -73,6 +83,18 @@ namespace EpicGames.UHT.Parsers
 								{
 									scratch.BaseIdentifiers.Add(identifier.ToArray());
 								});
+						})
+						.Optional("VINTERFACES", () =>
+						{
+							tokenReader.Require('(');
+							tokenReader.RequireList(')', ',', false, () =>
+							{
+								tokenReader.RequireIdentifier((ref UhtToken identifier) =>
+								{
+									scratch.VerseInterfaceIdentifiers.Add(identifier);
+								});
+							});
+							scratch.VerseInterfacesLine = tokenReader.InputLine;
 						});
 				});
 				inheritance = scratch;
diff --git a/Engine/Source/Programs/Shared/EpicGames.UHT/Parsers/UhtScriptStructParser.cs b/Engine/Source/Programs/Shared/EpicGames.UHT/Parsers/UhtScriptStructParser.cs
--- a/Engine/Source/Programs/Shared/EpicGames.UHT/Parsers/UhtScriptStructParser.cs
+++ b/Engine/Source/Programs/Shared/EpicGames.UHT/Parsers/UhtScriptStructParser.cs
@@ -113,6 +113,10 @@ namespace EpicGames.UHT.Parsers
 					UhtParserHelpers.ParseInheritance(parentScope.HeaderParser, topScope.Session.Config!, out UhtInheritance inheritance);
 					scriptStruct.SuperIdentifier = inheritance.SuperIdentifier;
 					scriptStruct.BaseIdentifiers = inheritance.BaseIdentifiers;
+					if (inheritance.VerseInterfacesLine != 0)
+					{
+						topScope.TokenReader.LogError("VINTERFACES() isn't allowed on USTRUCT definitions");
+					}
 
 					// Add the comments here for compatibility with old UHT
 					//COMPATIBILITY-TODO - Move this back to where the AddModuleRelativePathToMetaData is called.
diff --git a/Engine/Source/Programs/Shared/EpicGames.UHT/Specifiers/UhtFieldSpecifiers.cs b/Engine/Source/Programs/Shared/EpicGames.UHT/Specifiers/UhtFieldSpecifiers.cs
--- a/Engine/Source/Programs/Shared/EpicGames.UHT/Specifiers/UhtFieldSpecifiers.cs
+++ b/Engine/Source/Programs/Shared/EpicGames.UHT/Specifiers/UhtFieldSpecifiers.cs
@@ -18,9 +18,11 @@ namespace EpicGames.UHT.Parsers
 	[SuppressMessage("CodeQuality", "IDE0051:Remove unused private members", Justification = "Attribute accessed method")]
 	public static class UhtFieldSpecifiers
 	{
-		[UhtSpecifier(Extends = UhtTableNames.Struct, ValueType = UhtSpecifierValueType.OptionalEqualsKeyValuePairList)]
+		[UhtSpecifier(Extends = UhtTableNames.ScriptStruct, ValueType = UhtSpecifierValueType.OptionalEqualsKeyValuePairList)]
 		[UhtSpecifier(Extends = UhtTableNames.Class, ValueType = UhtSpecifierValueType.OptionalEqualsKeyValuePairList)]
 		[UhtSpecifier(Extends = UhtTableNames.Enum, ValueType = UhtSpecifierValueType.OptionalEqualsKeyValuePairList)]
+		[UhtSpecifier(Extends = UhtTableNames.Interface, ValueType = UhtSpecifierValueType.OptionalEqualsKeyValuePairList)]
+		[UhtSpecifier(Extends = UhtTableNames.NativeInterface, ValueType = UhtSpecifierValueType.OptionalEqualsKeyValuePairList)]
 		private static void VerseSpecifier(UhtSpecifierContext specifierContext, List<KeyValuePair<StringView, StringView>> values)
 		{
 			UhtField fieldObj = (UhtField)specifierContext.Type;
diff --git a/Engine/Source/Programs/Shared/EpicGames.UHT/Specifiers/UhtFunctionSpecifiers.cs b/Engine/Source/Programs/Shared/EpicGames.UHT/Specifiers/UhtFunctionSpecifiers.cs
--- a/Engine/Source/Programs/Shared/EpicGames.UHT/Specifiers/UhtFunctionSpecifiers.cs
+++ b/Engine/Source/Programs/Shared/EpicGames.UHT/Specifiers/UhtFunctionSpecifiers.cs
@@ -410,6 +410,32 @@ namespace EpicGames.UHT.Parsers
 			}
 		}
 
+		[UhtSpecifier(Extends = UhtTableNames.Function, ValueType = UhtSpecifierValueType.OptionalEqualsKeyValuePairList)]
+		private static void VerseSpecifier(UhtSpecifierContext specifierContext, List<KeyValuePair<StringView, StringView>> values)
+		{
+			UhtFunction functionObj = (UhtFunction)specifierContext.Type;
+
+			// Extract all the elements of the meta data
+			foreach (KeyValuePair<StringView, StringView> kvp in values)
+			{
+				ReadOnlySpan<char> key = kvp.Key.Span;
+				if (key.Equals("NativeCallable", StringComparison.OrdinalIgnoreCase))
+				{
+					functionObj.FunctionExportFlags |= UhtFunctionExportFlags.VerseNativeCallable;
+				}
+				else if (key.Equals("Decides", StringComparison.OrdinalIgnoreCase))
+				{
+					functionObj.FunctionExportFlags |= UhtFunctionExportFlags.VerseDecides;
+				}
+				else
+				{
+					functionObj.LogError($"Verse specifier option '{key}' is unknown");
+				}
+			}
+
+			functionObj.VerseName = functionObj.SourceName;
+		}
+
 		private static void ParseNetServiceIdentifiers(UhtSpecifierContext specifierContext, List<KeyValuePair<StringView, StringView>> identifiers)
 		{
 			IUhtTokenReader tokenReader = specifierContext.TokenReader;
diff --git a/Engine/Source/Programs/Shared/EpicGames.UHT/Types/UhtClass.cs b/Engine/Source/Programs/Shared/EpicGames.UHT/Types/UhtClass.cs
--- a/Engine/Source/Programs/Shared/EpicGames.UHT/Types/UhtClass.cs
+++ b/Engine/Source/Programs/Shared/EpicGames.UHT/Types/UhtClass.cs
@@ -2,6 +2,7 @@
 
 using System;
 using System.Collections.Generic;
+using System.Diagnostics.CodeAnalysis;
 using System.Text;
 using System.Text.Json.Serialization;
 using EpicGames.Core;
@@ -561,6 +562,34 @@ namespace EpicGames.UHT.Types
 			}
 		}
 
+		/// <summary>
+		/// Line number containing the VINTERFACES macro
+		/// </summary>
+		public int VerseInterfacesLine { get; set; } = 0;
+
+		/// <summary>
+		/// Verse interfaces directly specified on the interface definition
+		/// </summary>
+		[JsonConverter(typeof(UhtTypeReadOnlyListSourceNameJsonConverter<UhtClass>))]
+		public IReadOnlyList<UhtClass> VerseInterfaces => _verseInterfaces ?? s_emptyVerseInterfaces;
+
+		/// <summary>
+		/// Flattened collection of all verse interfaces. 
+		/// </summary>
+		[JsonConverter(typeof(UhtTypeReadOnlyListSourceNameJsonConverter<UhtClass>))]
+		public IReadOnlyList<UhtClass> FlattenedVerseInterfaces => _flattenedVerseInterfaces ?? s_emptyVerseInterfaces;
+
+		/// <summary>
+		/// List of all verse interfaces implemented by this interface
+		/// </summary>
+		[JsonIgnore]
+		[SuppressMessage("Usage", "CA2227:Collection properties should be read only", Justification = "<Pending>")]
+		public List<UhtToken>? VerseInterfaceIdentifiers { get; set; } = null;
+
+		private static readonly List<UhtClass> s_emptyVerseInterfaces = new();
+		private List<UhtClass>? _verseInterfaces = null;
+		private List<UhtClass>? _flattenedVerseInterfaces = null;
+
 		/// <summary>
 		/// Add the given list of tokens as a possible declaration
 		/// </summary>
@@ -647,6 +676,7 @@ namespace EpicGames.UHT.Types
 		{
 			BindSuper(SuperIdentifier, UhtFindOptions.Class);
 			BindBases(BaseIdentifiers, UhtFindOptions.Class);
+			BindVerseInterfaces();       
 			base.BindSuperAndBases();
 		}
 
@@ -655,6 +685,14 @@ namespace EpicGames.UHT.Types
 		{
 			base.ResolveSuper(resolvePhase);
 
+			if (_verseInterfaces != null)
+			{
+				foreach (UhtClass verseInterface in _verseInterfaces)
+				{
+					verseInterface.ResolveSuper(resolvePhase);
+				}
+			}
+
 			// Make sure the class within is resolved.  We have to make sure we don't try to resolve ourselves.
 			if (ClassWithin != null && ClassWithin != this)
 			{
@@ -672,10 +710,37 @@ namespace EpicGames.UHT.Types
 						ClassFlags |= EClassFlags.MatchedSerializers;
 					}
 
+					// Roll up the verse interfaces
+					if (_verseInterfaces != null)
+					{
+						_flattenedVerseInterfaces = [];
+						_flattenedVerseInterfaces.AddRange(_verseInterfaces);
+						int stopIndex = _flattenedVerseInterfaces.Count;
+						for (int index = 0; index < stopIndex; index++)
+						{
+							UhtClass verseInterface = _flattenedVerseInterfaces[index];
+							if (verseInterface._flattenedVerseInterfaces != null)
+							{
+								foreach (UhtClass subInterface in verseInterface._flattenedVerseInterfaces)
+								{
+									if (!_flattenedVerseInterfaces.Contains(subInterface))
+									{
+										_flattenedVerseInterfaces.Add(subInterface);
+									}
+								}
+							}
+						}
+					}
+
 					switch (ClassType)
 					{
 						case UhtClassType.Class:
 							{
+								// Validate some verse settings
+								if (!IsVerseField && VerseInterfacesLine != 0)
+								{
+									this.LogError($"VINTERFACES can only appear of classes marked as Verse classes");
+								}
 
 								// Merge the super class flags
 								if (superClass != null)
@@ -691,10 +756,23 @@ namespace EpicGames.UHT.Types
 										{
 											this.LogError($"Class '{baseClass.SourceName}' is not an interface; Can only inherit from non-UObjects or UInterface derived interfaces");
 										}
+										if (baseClass.IsVerseField)
+										{
+											if (IsVerseField)
+											{
+												this.LogError($"Interface '{baseClass.SourceName}' must be listed in a VINTERFACES() macro");
+											}
+											else
+											{
+												this.LogError($"Non-Verse classes can not inherit the Verse interface '{baseClass.SourceName}");
+											}
+										}
 										ClassFlags |= baseClass.ClassFlags & EClassFlags.ScriptInherit & ~RemovedClassFlags;
 									}
 								}
 
+								ValidateVerseInterfaces();
+
 								// These following collections only inherit from he parent if they are empty in this class
 								if (SparseClassDataTypes.Count == 0)
 								{
@@ -727,6 +805,15 @@ namespace EpicGames.UHT.Types
 
 						case UhtClassType.Interface:
 							{
+								if (IsVerseField)
+								{
+									if (superClass != Session.UInterface)
+									{
+										this.LogError($"Verse interface must derived from UInterface");
+									}
+									ValidateVerseInterfaces();
+								}
+
 								if (superClass != null)
 								{
 									ClassFlags |= superClass.ClassFlags & EClassFlags.ScriptInherit;
@@ -744,6 +831,13 @@ namespace EpicGames.UHT.Types
 							break;
 
 						case UhtClassType.NativeInterface:
+							if (AlternateObject is UhtClass interfaceObj && interfaceObj.IsVerseField)
+							{
+								if (superClass != Session.IInterface)
+								{
+									this.LogError($"Verse native interface must derived from IInterface");
+								}
+							}
 							break;
 
 						default:
@@ -753,6 +847,21 @@ namespace EpicGames.UHT.Types
 			}
 		}
 
+		private void ValidateVerseInterfaces()
+		{
+			foreach (UhtClass verseInterface in VerseInterfaces)
+			{
+				if (verseInterface.ClassType != UhtClassType.NativeInterface)
+				{
+					this.LogError($"Class '{verseInterface.SourceName}' is not a native interface");
+				}
+				if (!verseInterface.IsVerseField)
+				{
+					this.LogError($"Class '{verseInterface.SourceName}' is not a Verse interface");
+				}
+			}
+		}
+
 		/// <summary>
 		/// Return the native interface name from an interface
 		/// </summary>
@@ -829,6 +938,10 @@ namespace EpicGames.UHT.Types
 								else
 								{
 									AlternateObject = interfaceObj;
+									if (interfaceObj.IsVerseField)
+									{
+										VerseName = SourceName;
+									}
 									interfaceObj.NativeInterface = this;
 									//COMPATIBILITY-TODO - Use the native interface access specifier
 									interfaceObj.GeneratedBodyAccessSpecifier = GeneratedBodyAccessSpecifier;
@@ -1997,14 +2110,20 @@ namespace EpicGames.UHT.Types
 			}
 
 			// Add any other base classes
-			foreach (UhtStruct @base in Bases)
+			foreach (UhtStruct baseObj in Bases)
 			{
-				if (@base is UhtClass baseClass)
+				if (baseObj is UhtClass baseClass)
 				{
 					collector.AddCrossModuleReference(baseClass, false);
 				}
 			}
 
+			// Add any verse interfaces
+			foreach (UhtStruct interfaceObj in FlattenedVerseInterfaces)
+			{
+				collector.AddCrossModuleReference(interfaceObj, false);
+			}
+
 			// Add the package
 			collector.AddDeclaration(Package, true);
 			collector.AddCrossModuleReference(Package, true);
@@ -2048,5 +2167,27 @@ namespace EpicGames.UHT.Types
 			}
 			return false;
 		}
+
+		private void BindVerseInterfaces()
+		{
+			if (VerseInterfaceIdentifiers != null)
+			{
+				foreach (UhtToken identifier in VerseInterfaceIdentifiers)
+				{
+					// We really only case about interfaces, but we can also handle structs
+					UhtToken scratch = identifier;
+					UhtClass? verseInterface = (UhtClass?)FindType(UhtFindOptions.Class | UhtFindOptions.SourceName | UhtFindOptions.NoSelf, ref scratch);
+					if (verseInterface != null)
+					{
+						_verseInterfaces ??= new List<UhtClass>();
+						_verseInterfaces.Add(verseInterface);
+					}
+					else
+					{
+						this.LogError($"Verse interface {identifier.Value} not found");
+					}
+				}
+			}
+		}
 	}
 }
diff --git a/Engine/Source/Programs/Shared/EpicGames.UHT/Types/UhtFunction.cs b/Engine/Source/Programs/Shared/EpicGames.UHT/Types/UhtFunction.cs
--- a/Engine/Source/Programs/Shared/EpicGames.UHT/Types/UhtFunction.cs
+++ b/Engine/Source/Programs/Shared/EpicGames.UHT/Types/UhtFunction.cs
@@ -116,6 +116,16 @@ namespace EpicGames.UHT.Types
 		/// True if the function specifier has a getter/setter specified
 		/// </summary>
 		SawPropertyAccessor = 1 << 17,
+
+		/// <summary>
+		/// True if the function is verse function callable from native code
+		/// </summary>
+		VerseNativeCallable = 1 << 18,
+
+		/// <summary>
+		/// True if the function is a decides function
+		/// </summary>
+		VerseDecides = 1 << 19,
 	}
 
 	/// <summary>
@@ -372,6 +382,12 @@ namespace EpicGames.UHT.Types
 			}
 		}
 
+		/// <summary>
+		/// True if this function has an incomplete return type
+		/// </summary>
+		public bool HasVerseIncompleteReturn => ReturnProperty is UhtEnumProperty enumProperty && enumProperty.Enum == Session.EVerseNativeCallResult;
+
+
 		/// <summary>
 		/// Construct a new instance of a function
 		/// </summary>
@@ -799,6 +815,24 @@ namespace EpicGames.UHT.Types
 							this.LogError($"FieldNotify function '{SourceName}' must be pure.");
 						}
 					}
+
+					// Verse validations
+					if (IsVerseField)
+					{
+						if (FunctionExportFlags.HasAnyFlags(UhtFunctionExportFlags.VerseDecides))
+						{
+							UhtProperty? returnProperty = ReturnProperty;
+							if (returnProperty == null || returnProperty is not UhtOptionalProperty)
+							{
+								this.LogError($"Verse function '{SourceName}' marked as 'decides' must have a 'TOptional<...>' return value");
+							}							
+						}
+
+						if (Outer is UhtClass parent && !parent.IsVerseField)
+						{
+							this.LogError($"Function '{SourceName}' is marked as a Verse function but appears in a class that isn't a Verse class");
+						}
+					}
 					break;
 
 				case UhtFunctionType.Delegate:
diff --git a/Engine/Source/Programs/Shared/EpicGames.UHT/Utils/UhtSession.cs b/Engine/Source/Programs/Shared/EpicGames.UHT/Utils/UhtSession.cs
--- a/Engine/Source/Programs/Shared/EpicGames.UHT/Utils/UhtSession.cs
+++ b/Engine/Source/Programs/Shared/EpicGames.UHT/Utils/UhtSession.cs
@@ -2292,6 +2292,13 @@ namespace EpicGames.UHT.Utils
 					{
 						TopologicalStructVisit(states, baseStruct, structStack);
 					}
+					if (visit is UhtClass classObj)
+					{
+						foreach (UhtStruct verseInterface in classObj.VerseInterfaces)
+						{
+							TopologicalStructVisit(states, verseInterface, structStack);
+						}
+					}
 					TopologicalStructVisitChildren(states, visit, structStack);
 					states[visit.ObjectTypeIndex] = TopologicalState.Permanent;
 					break;
diff --git a/Engine/Source/Programs/Shared/EpicGames.UHT/Utils/UhtStringBuilderExtensions.cs b/Engine/Source/Programs/Shared/EpicGames.UHT/Utils/UhtStringBuilderExtensions.cs
--- a/Engine/Source/Programs/Shared/EpicGames.UHT/Utils/UhtStringBuilderExtensions.cs
+++ b/Engine/Source/Programs/Shared/EpicGames.UHT/Utils/UhtStringBuilderExtensions.cs
@@ -460,8 +460,9 @@ namespace EpicGames.UHT.Utils
 			switch (classObj.ClassType)
 			{
 				case UhtClassType.Class:
-				case UhtClassType.NativeInterface:
 					return builder.Append(classObj.SourceName);
+				case UhtClassType.NativeInterface:
+					return builder.Append(classObj.SourceName).Append(UhtNames.VerseProxySuffix);
 				case UhtClassType.Interface:
 					return builder.Append('I').Append(classObj.SourceName[1..]).Append(UhtNames.VerseProxySuffix);
 				default:
