UHT support for Verse classes and interfaces defined in UHT (Part 2, update to use new Verse class constructor)

#rnx

diff --git a/Engine/Source/Programs/Shared/EpicGames.UHT/Exporters/CodeGen/UhtHeaderCodeGeneratorCppFile.cs b/Engine/Source/Programs/Shared/EpicGames.UHT/Exporters/CodeGen/UhtHeaderCodeGeneratorCppFile.cs
--- a/Engine/Source/Programs/Shared/EpicGames.UHT/Exporters/CodeGen/UhtHeaderCodeGeneratorCppFile.cs
+++ b/Engine/Source/Programs/Shared/EpicGames.UHT/Exporters/CodeGen/UhtHeaderCodeGeneratorCppFile.cs
@@ -1991,7 +1991,9 @@ namespace EpicGames.UHT.Exporters.CodeGen
 				builder.Append("{\r\n");
 				builder.Append("\tif (!").Append(registrationName).Append(".OuterSingleton)\r\n");
 				builder.Append("\t{\r\n");
-				builder.Append("\t\tUECodeGen_Private::ConstructUClass(").Append(registrationName).Append(".OuterSingleton, ").Append(staticsName).Append("::ClassParams);\r\n");
+				builder.Append("\t\t")
+					.Append(classObj.IsVerseField ? "Verse::CodeGen::Private::ConstructUVerseClass(" : "UECodeGen_Private::ConstructUClass(")
+					.Append(registrationName).Append(".OuterSingleton, ").Append(staticsName).Append("::ClassParams);\r\n");
 				if (sparseDataTypes != null)
 				{
 					foreach (string sparseClass in sparseDataTypes)
