Compact Binary - When serializing a dict with a string key, serialize it as a object instead of an array of pairs. This follows more closely how we typically serialize dictionaries and is similar to conventions in other structured object types like json. If this causes a breakage you can opt out of it by explicitly setting the CbConverter on the property.

diff --git a/Engine/Source/Programs/Shared/EpicGames.Serialization.Tests/SerializerTests.cs b/Engine/Source/Programs/Shared/EpicGames.Serialization.Tests/SerializerTests.cs
--- a/Engine/Source/Programs/Shared/EpicGames.Serialization.Tests/SerializerTests.cs
+++ b/Engine/Source/Programs/Shared/EpicGames.Serialization.Tests/SerializerTests.cs
@@ -183,5 +183,40 @@ namespace EpicGames.Serialization.Tests
 			Assert.AreEqual(123, output.Dictionary["hello"]);
 			Assert.AreEqual(456, output.Dictionary["world"]);
 		}
+
+		class ObjectWithStringDictCollection
+		{
+			[CbField("d")]
+			public Dictionary<string, int> Dictionary { get; set; } = new Dictionary<string, int>(StringComparer.Ordinal);
+		}
+
+		[TestMethod]
+		public void ObjectWithStringDictCollectionTest()
+		{
+			ObjectWithStringDictCollection input = new ObjectWithStringDictCollection();
+			input.Dictionary.Add("hello", 123);
+			input.Dictionary.Add("world", 456);
+			CbObject obj = CbSerializer.Serialize(input);
+
+			ObjectWithStringDictCollection output = CbSerializer.Deserialize<ObjectWithStringDictCollection>(obj.AsField());
+			Assert.AreEqual(2, output.Dictionary.Count);
+			Assert.AreEqual(123, output.Dictionary["hello"]);
+			Assert.AreEqual(456, output.Dictionary["world"]);
+
+			// make sure string keyed dictionaries are serialized as objects
+			CbWriter writer = new CbWriter();
+			writer.BeginObject();
+			writer.BeginObject("d");
+			writer.WriteInteger("hello", 123);
+			writer.WriteInteger("world", 456);
+			writer.EndObject();
+			writer.EndObject();
+			CbObject manualCbObject = writer.ToObject();
+
+			ObjectWithStringDictCollection output2 = CbSerializer.Deserialize<ObjectWithStringDictCollection>(manualCbObject.AsField());
+			Assert.AreEqual(2, output2.Dictionary.Count);
+			Assert.AreEqual(123, output2.Dictionary["hello"]);
+			Assert.AreEqual(456, output2.Dictionary["world"]);
+		}
 	}
 }
diff --git a/Engine/Source/Programs/Shared/EpicGames.Serialization/Converters/CbDictionaryConverter.cs b/Engine/Source/Programs/Shared/EpicGames.Serialization/Converters/CbDictionaryConverter.cs
--- a/Engine/Source/Programs/Shared/EpicGames.Serialization/Converters/CbDictionaryConverter.cs
+++ b/Engine/Source/Programs/Shared/EpicGames.Serialization/Converters/CbDictionaryConverter.cs
@@ -2,6 +2,7 @@
 
 using System;
 using System.Collections.Generic;
+using System.Linq;
 
 namespace EpicGames.Serialization.Converters
 {
@@ -93,10 +94,82 @@ namespace EpicGames.Serialization.Converters
 		{
 			if (type.IsGenericType && type.GetGenericTypeDefinition() == typeof(Dictionary<,>))
 			{
-				Type converterType = typeof(CbDictionaryConverter<,>).MakeGenericType(type.GenericTypeArguments);
-				return (CbConverter)Activator.CreateInstance(converterType)!;
+				Type? keyType = type.GetGenericArguments().FirstOrDefault();
+				if (keyType == typeof(string))
+				{
+					// if the key type is a string serialize as a CbObject to more closely follow typical conventions
+					Type converterType = typeof(CbStringDictionaryConverter<>).MakeGenericType(type.GenericTypeArguments[1..]);
+					return (CbConverter)Activator.CreateInstance(converterType)!;
+				}
+				else
+				{
+					Type converterType = typeof(CbDictionaryConverter<,>).MakeGenericType(type.GenericTypeArguments);
+					return (CbConverter)Activator.CreateInstance(converterType)!;
+				}
 			}
 			return null;
 		}
 	}
+
+	/// <summary>
+	/// Converter for dictionary types with string keys, serialized as objects
+	/// </summary>
+	class CbStringDictionaryConverter<TValue> : CbConverter<Dictionary<string, TValue>>
+	{
+		/// <inheritdoc/>
+		public override Dictionary<string, TValue> Read(CbField field)
+		{
+			if (field.IsNull())
+			{
+				return null!;
+			}
+
+			Dictionary<string, TValue> dictionary = [];
+			foreach (CbField element in field)
+			{
+				string key = element.Name.ToString();
+				
+				TValue value = CbSerializer.Deserialize<TValue>(element);
+
+				dictionary.Add(key, value);
+			}
+			return dictionary;
+		}
+
+		/// <inheritdoc/>
+		public override void Write(CbWriter writer, Dictionary<string, TValue> value)
+		{
+			if (value == null)
+			{
+				writer.WriteNullValue();
+			}
+			else
+			{
+				writer.BeginObject();
+				foreach (KeyValuePair<string, TValue> pair in value)
+				{
+					CbSerializer.Serialize(writer, pair.Key, pair.Value);
+				}
+				writer.EndObject();
+			}
+		}
+
+		/// <inheritdoc/>
+		public override void WriteNamed(CbWriter writer, CbFieldName name, Dictionary<string, TValue> value)
+		{
+			if (value == null)
+			{
+				writer.WriteNull(name);
+			}
+			else if (value.Count > 0)
+			{
+				writer.BeginObject(name);
+				foreach (KeyValuePair<string, TValue> pair in value)
+				{
+					CbSerializer.Serialize(writer, pair.Key, pair.Value);
+				}
+				writer.EndObject();
+			}
+		}
+	}
 }
