UAT - ProjectStoreTypes: Add documentation xml comments to fix CS1591 and remove setter to address CA2227

#rnx

diff --git a/Engine/Source/Programs/Shared/EpicGames.ProjectStore/ProjectStoreTypes.cs b/Engine/Source/Programs/Shared/EpicGames.ProjectStore/ProjectStoreTypes.cs
--- a/Engine/Source/Programs/Shared/EpicGames.ProjectStore/ProjectStoreTypes.cs
+++ b/Engine/Source/Programs/Shared/EpicGames.ProjectStore/ProjectStoreTypes.cs
@@ -25,7 +25,13 @@ namespace EpicGames.ProjectStore
 	[JsonConverter(typeof(JsonStringEnumConverter))]
 	public enum ZenOperatingMode
 	{
+		/// <summary>
+		/// Use Zen streaming
+		/// </summary>
 		Stream,
+		/// <summary>
+		/// Access data from local filesystem
+		/// </summary>
 		Filesystem
 	}
 
@@ -45,7 +51,7 @@ namespace EpicGames.ProjectStore
 		/// <summary>
 		/// The list of remote host names for zenserver
 		/// </summary>
-		public List<string> RemoteHostNames { get; set; } = new List<string>();
+		public List<string> RemoteHostNames { get; } = new List<string>();
 		/// <summary>
 		/// The port for zenserver on the host
 		/// </summary>
