UHT: Added support for WITH_TESTS

diff --git a/Engine/Source/Programs/Shared/EpicGames.UHT/Exporters/CodeGen/UhtHeaderCodeGeneratorCppFile.cs b/Engine/Source/Programs/Shared/EpicGames.UHT/Exporters/CodeGen/UhtHeaderCodeGeneratorCppFile.cs
--- a/Engine/Source/Programs/Shared/EpicGames.UHT/Exporters/CodeGen/UhtHeaderCodeGeneratorCppFile.cs
+++ b/Engine/Source/Programs/Shared/EpicGames.UHT/Exporters/CodeGen/UhtHeaderCodeGeneratorCppFile.cs
@@ -183,12 +183,14 @@ namespace EpicGames.UHT.Exporters.CodeGen
 					if (field is UhtEnum enumObj)
 					{
 						using UhtCodeBlockComment blockComment = new(builder, field);
+						using UhtMacroBlockEmitter macroBlockEmitter = new(builder, UhtDefineScopeNames.Standard, enumObj.DefineScope);
 						AppendEnum(builder, enumObj);
 						GetRegistrations(packageRegistrations, field).Enumerations.Add(enumObj);
 					}
 					else if (field is UhtScriptStruct scriptStruct)
 					{
 						using UhtCodeBlockComment blockComment = new(builder, field);
+						using UhtMacroBlockEmitter macroBlockEmitter = new(builder, UhtDefineScopeNames.Standard, scriptStruct.DefineScope);
 						AppendScriptStruct(builder, scriptStruct);
 						if (scriptStruct.ScriptStructFlags.HasAnyFlags(EStructFlags.Native))
 						{
@@ -198,6 +200,7 @@ namespace EpicGames.UHT.Exporters.CodeGen
 					else if (field is UhtFunction function)
 					{
 						using UhtCodeBlockComment blockComment = new(builder, field);
+						using UhtMacroBlockEmitter macroBlockEmitter = new(builder, UhtDefineScopeNames.Standard, function.DefineScope);
 						AppendDelegate(builder, function);
 					}
 					else if (field is UhtClass classObj)
@@ -215,6 +218,7 @@ namespace EpicGames.UHT.Exporters.CodeGen
 							}
 
 							using UhtCodeBlockComment blockComment = new(builder, field);
+							using UhtMacroBlockEmitter macroBlockEmitter = new(builder, UhtDefineScopeNames.Standard, classObj.DefineScope);
 							AppendClass(builder, classObj, functions);
 							GetRegistrations(packageRegistrations, field).Classes.Add(classObj);
 
@@ -399,131 +403,127 @@ namespace EpicGames.UHT.Exporters.CodeGen
 				enumDisplayNameFn = "nullptr";
 			}
 
-			using (UhtMacroBlockEmitter macroBlockEmitter = new(builder, UhtDefineScopeNames.Standard, enumObj.DefineScope))
+			// If we don't have a zero 0 then we emit a static assert to verify we have one
+			if (!enumObj.IsValidEnumValue(0) && enumObj.MetaData.ContainsKey(UhtNames.BlueprintType))
 			{
-
-				// If we don't have a zero 0 then we emit a static assert to verify we have one
-				if (!enumObj.IsValidEnumValue(0) && enumObj.MetaData.ContainsKey(UhtNames.BlueprintType))
+				bool hasUnparsedValue = enumObj.EnumValues.Exists(x => x.Value == -1);
+				if (hasUnparsedValue)
 				{
-					bool hasUnparsedValue = enumObj.EnumValues.Exists(x => x.Value == -1);
-					if (hasUnparsedValue)
+					builder.Append("static_assert(");
+					bool doneFirst = false;
+					foreach (UhtEnumValue value in enumObj.EnumValues)
 					{
-						builder.Append("static_assert(");
-						bool doneFirst = false;
-						foreach (UhtEnumValue value in enumObj.EnumValues)
+						if (value.Value == -1)
 						{
-							if (value.Value == -1)
+							if (doneFirst)
 							{
-								if (doneFirst)
-								{
-									builder.Append("||");
-								}
-								doneFirst = true;
-								builder.Append("!int64(").Append(value.Name).Append(')');
+								builder.Append("||");
 							}
+							doneFirst = true;
+							builder.Append("!int64(").Append(value.Name).Append(')');
 						}
-						builder.Append(", \"'").Append(enumObj.SourceName).Append("' does not have a 0 entry!(This is a problem when the enum is initialized by default)\");\r\n");
 					}
+					builder.Append(", \"'").Append(enumObj.SourceName).Append("' does not have a 0 entry!(This is a problem when the enum is initialized by default)\");\r\n");
 				}
+			}
 
-				builder.Append("static FEnumRegistrationInfo ").Append(registrationName).Append(";\r\n");
-				builder.Append("static UEnum* ").Append(enumObj.SourceName).Append("_StaticEnum()\r\n");
-				builder.Append("{\r\n");
+			builder.Append("static FEnumRegistrationInfo ").Append(registrationName).Append(";\r\n");
+			builder.Append("static UEnum* ").Append(enumObj.SourceName).Append("_StaticEnum()\r\n");
+			builder.Append("{\r\n");
 
-				builder.Append("\tif (!").Append(registrationName).Append(".OuterSingleton)\r\n");
-				builder.Append("\t{\r\n");
-				builder.Append("\t\t").Append(registrationName).Append(".OuterSingleton = GetStaticEnum(").Append(singletonName).Append(", (UObject*)")
-					.Append(GetSingletonName(enumObj.Package, true)).Append("(), TEXT(\"").Append(enumObj.SourceName).Append("\"));\r\n");
-				builder.Append("\t}\r\n");
-				builder.Append("\treturn ").Append(registrationName).Append(".OuterSingleton;\r\n");
+			builder.Append("\tif (!").Append(registrationName).Append(".OuterSingleton)\r\n");
+			builder.Append("\t{\r\n");
+			builder.Append("\t\t").Append(registrationName).Append(".OuterSingleton = GetStaticEnum(").Append(singletonName).Append(", (UObject*)")
+				.Append(GetSingletonName(enumObj.Package, true)).Append("(), TEXT(\"").Append(enumObj.SourceName).Append("\"));\r\n");
+			builder.Append("\t}\r\n");
+			builder.Append("\treturn ").Append(registrationName).Append(".OuterSingleton;\r\n");
 
-				builder.Append("}\r\n");
+			builder.Append("}\r\n");
 
-				builder.Append("template<> ").Append(Module.Api).Append("UEnum* StaticEnum<").Append(enumObj.CppType).Append(">()\r\n");
-				builder.Append("{\r\n");
-				builder.Append("\treturn ").Append(enumObj.SourceName).Append("_StaticEnum();\r\n");
-				builder.Append("}\r\n");
+			builder.Append("template<> ").Append(Module.Api).Append("UEnum* StaticEnum<").Append(enumObj.CppType).Append(">()\r\n");
+			builder.Append("{\r\n");
+			builder.Append("\treturn ").Append(enumObj.SourceName).Append("_StaticEnum();\r\n");
+			builder.Append("}\r\n");
 
-				if (enumObj.IsVerseField)
-				{
-					builder.Append("V_DEFINE_IMPORTED_ENUM(").Append(Module.Api.TrimEnd()).Append(", \"").AppendVerseUEPackageName(enumObj).Append("\", \"").Append(enumObj.VersePath).Append("\", ").Append(enumObj.SourceName).Append(");\r\n");
-				}
+			if (enumObj.IsVerseField)
+			{
+				builder.Append("V_DEFINE_IMPORTED_ENUM(").Append(Module.Api.TrimEnd()).Append(", \"").AppendVerseUEPackageName(enumObj).Append("\", \"").Append(enumObj.VersePath).Append("\", ").Append(enumObj.SourceName).Append(");\r\n");
+			}
 
-				// Everything from this point on will be part of the definition hash
-				int hashCodeBlockStart = builder.Length;
+			// Everything from this point on will be part of the definition hash
+			int hashCodeBlockStart = builder.Length;
 
-				// Statics declaration
-				{
-					builder.Append("struct ").Append(staticsName).Append("\r\n");
-					builder.Append("{\r\n");
-					builder.AppendMetaDataDecl(enumObj, null, null, MetaDataParamsName, 1);
+			// Statics declaration
+			{
+				builder.Append("struct ").Append(staticsName).Append("\r\n");
+				builder.Append("{\r\n");
+				builder.AppendMetaDataDecl(enumObj, null, null, MetaDataParamsName, 1);
 
-					// Enumerators
-					if (enumObj.EnumValues.Count > 0)
+				// Enumerators
+				if (enumObj.EnumValues.Count > 0)
+				{
+					builder.Append("\tstatic constexpr UECodeGen_Private::FEnumeratorParam Enumerators[] = {\r\n");
+					int enumIndex = 0;
+					foreach (UhtEnumValue value in enumObj.EnumValues)
 					{
-						builder.Append("\tstatic constexpr UECodeGen_Private::FEnumeratorParam Enumerators[] = {\r\n");
-						int enumIndex = 0;
-						foreach (UhtEnumValue value in enumObj.EnumValues)
+						if (!enumObj.MetaData.TryGetValue("OverrideName", enumIndex, out string? keyName))
 						{
-							if (!enumObj.MetaData.TryGetValue("OverrideName", enumIndex, out string? keyName))
-							{
-								keyName = value.Name.ToString();
-							}
-							builder.Append("\t\t{ ").AppendUTF8LiteralString(keyName).Append(", (int64)").Append(value.Name).Append(" },\r\n");
-							++enumIndex;
+							keyName = value.Name.ToString();
 						}
-						builder.Append("\t};\r\n");
+						builder.Append("\t\t{ ").AppendUTF8LiteralString(keyName).Append(", (int64)").Append(value.Name).Append(" },\r\n");
+						++enumIndex;
 					}
-					builder.Append("\tstatic const UECodeGen_Private::FEnumParams EnumParams;\r\n");
-					builder.Append("};\r\n");
-				}
-
-				// Statics definition
-				{
-					builder.Append("const UECodeGen_Private::FEnumParams ").Append(staticsName).Append("::EnumParams = {\r\n");
-					builder.Append("\t(UObject*(*)())").Append(GetSingletonName(enumObj.Package, true)).Append(",\r\n");
-					builder.Append('\t').Append(enumDisplayNameFn).Append(",\r\n");
-					builder.Append('\t').AppendUTF8LiteralString(enumObj.SourceName).Append(",\r\n");
-					builder.Append('\t').AppendUTF8LiteralString(enumObj.CppType).Append(",\r\n");
-					if (enumObj.EnumValues.Count > 0)
-					{
-						builder.Append('\t').Append(staticsName).Append("::Enumerators,\r\n");
-						builder.Append('\t').Append(ObjectFlags).Append(",\r\n");
-						builder.Append("\tUE_ARRAY_COUNT(").Append(staticsName).Append("::Enumerators),\r\n");
-					}
-					else
-					{
-						builder.Append('\t').Append("nullptr,\r\n");
-						builder.Append('\t').Append(ObjectFlags).Append(",\r\n");
-						builder.Append("\t0,\r\n");
-					}
-					builder.Append('\t').Append(enumObj.EnumFlags.HasAnyFlags(EEnumFlags.Flags) ? "EEnumFlags::Flags" : "EEnumFlags::None").Append(",\r\n");
-					builder.Append("\t(uint8)UEnum::ECppForm::").Append(enumObj.CppForm.ToString()).Append(",\r\n");
-					builder.Append('\t').AppendMetaDataParams(enumObj, staticsName, MetaDataParamsName).Append("\r\n");
-					builder.Append("};\r\n");
+					builder.Append("\t};\r\n");
 				}
+				builder.Append("\tstatic const UECodeGen_Private::FEnumParams EnumParams;\r\n");
+				builder.Append("};\r\n");
+			}
 
-				// Registration singleton
-				builder.Append("UEnum* ").Append(singletonName).Append("()\r\n");
-				builder.Append("{\r\n");
-				builder.Append("\tif (!").Append(registrationName).Append(".InnerSingleton)\r\n");
-				builder.Append("\t{\r\n");
-				if (enumObj.IsVerseField)
+			// Statics definition
+			{
+				builder.Append("const UECodeGen_Private::FEnumParams ").Append(staticsName).Append("::EnumParams = {\r\n");
+				builder.Append("\t(UObject*(*)())").Append(GetSingletonName(enumObj.Package, true)).Append(",\r\n");
+				builder.Append('\t').Append(enumDisplayNameFn).Append(",\r\n");
+				builder.Append('\t').AppendUTF8LiteralString(enumObj.SourceName).Append(",\r\n");
+				builder.Append('\t').AppendUTF8LiteralString(enumObj.CppType).Append(",\r\n");
+				if (enumObj.EnumValues.Count > 0)
 				{
-					builder.Append("\t\tVerse::CodeGen::Private::ConstructUVerseEnum(").Append(registrationName).Append(".InnerSingleton, ").Append(staticsName).Append("::EnumParams);\r\n");
+					builder.Append('\t').Append(staticsName).Append("::Enumerators,\r\n");
+					builder.Append('\t').Append(ObjectFlags).Append(",\r\n");
+					builder.Append("\tUE_ARRAY_COUNT(").Append(staticsName).Append("::Enumerators),\r\n");
 				}
 				else
 				{
-					builder.Append("\t\tUECodeGen_Private::ConstructUEnum(").Append(registrationName).Append(".InnerSingleton, ").Append(staticsName).Append("::EnumParams);\r\n");
+					builder.Append('\t').Append("nullptr,\r\n");
+					builder.Append('\t').Append(ObjectFlags).Append(",\r\n");
+					builder.Append("\t0,\r\n");
 				}
-				builder.Append("\t}\r\n");
-				builder.Append("\treturn ").Append(registrationName).Append(".InnerSingleton;\r\n");
-				builder.Append("}\r\n");
+				builder.Append('\t').Append(enumObj.EnumFlags.HasAnyFlags(EEnumFlags.Flags) ? "EEnumFlags::Flags" : "EEnumFlags::None").Append(",\r\n");
+				builder.Append("\t(uint8)UEnum::ECppForm::").Append(enumObj.CppForm.ToString()).Append(",\r\n");
+				builder.Append('\t').AppendMetaDataParams(enumObj, staticsName, MetaDataParamsName).Append("\r\n");
+				builder.Append("};\r\n");
+			}
 
-				{
-					using UhtRentedPoolBuffer<char> borrowBuffer = builder.RentPoolBuffer(hashCodeBlockStart, builder.Length - hashCodeBlockStart);
-					ObjectInfos[enumObj.ObjectTypeIndex].Hash = UhtHash.GenenerateTextHash(borrowBuffer.Buffer.Memory.Span);
-				}
+			// Registration singleton
+			builder.Append("UEnum* ").Append(singletonName).Append("()\r\n");
+			builder.Append("{\r\n");
+			builder.Append("\tif (!").Append(registrationName).Append(".InnerSingleton)\r\n");
+			builder.Append("\t{\r\n");
+			if (enumObj.IsVerseField)
+			{
+				builder.Append("\t\tVerse::CodeGen::Private::ConstructUVerseEnum(").Append(registrationName).Append(".InnerSingleton, ").Append(staticsName).Append("::EnumParams);\r\n");
+			}
+			else
+			{
+				builder.Append("\t\tUECodeGen_Private::ConstructUEnum(").Append(registrationName).Append(".InnerSingleton, ").Append(staticsName).Append("::EnumParams);\r\n");
+			}
+			builder.Append("\t}\r\n");
+			builder.Append("\treturn ").Append(registrationName).Append(".InnerSingleton;\r\n");
+			builder.Append("}\r\n");
+
+			{
+				using UhtRentedPoolBuffer<char> borrowBuffer = builder.RentPoolBuffer(hashCodeBlockStart, builder.Length - hashCodeBlockStart);
+				ObjectInfos[enumObj.ObjectTypeIndex].Hash = UhtHash.GenenerateTextHash(borrowBuffer.Buffer.Memory.Span);
 			}
 			return builder;
 		}
diff --git a/Engine/Source/Programs/Shared/EpicGames.UHT/Exporters/CodeGen/UhtHeaderCodeGeneratorHFile.cs b/Engine/Source/Programs/Shared/EpicGames.UHT/Exporters/CodeGen/UhtHeaderCodeGeneratorHFile.cs
--- a/Engine/Source/Programs/Shared/EpicGames.UHT/Exporters/CodeGen/UhtHeaderCodeGeneratorHFile.cs
+++ b/Engine/Source/Programs/Shared/EpicGames.UHT/Exporters/CodeGen/UhtHeaderCodeGeneratorHFile.cs
@@ -120,12 +120,14 @@ namespace EpicGames.UHT.Exporters.CodeGen
 						if (scriptStruct.ScriptStructFlags.HasAnyFlags(EStructFlags.Native))
 						{
 							using UhtCodeBlockComment blockComment = new(builder, field);
+							using UhtMacroBlockEmitter macroBlockEmitter = new(builder, UhtDefineScopeNames.Standard, scriptStruct.DefineScope);
 							AppendScriptStruct(builder, scriptStruct);
 						}
 					}
 					else if (field is UhtFunction function)
 					{
 						using UhtCodeBlockComment blockComment = new(builder, field);
+						using UhtMacroBlockEmitter macroBlockEmitter = new(builder, UhtDefineScopeNames.Standard, function.DefineScope);
 						AppendDelegate(builder, function);
 					}
 					else if (field is UhtClass classObj)
@@ -133,6 +135,7 @@ namespace EpicGames.UHT.Exporters.CodeGen
 						if (!classObj.ClassFlags.HasAnyFlags(EClassFlags.Intrinsic))
 						{
 							using UhtCodeBlockComment blockComment = new(builder, field);
+							using UhtMacroBlockEmitter macroBlockEmitter = new(builder, UhtDefineScopeNames.Standard, classObj.DefineScope);
 							AppendClass(builder, classObj);
 						}
 					}
@@ -145,6 +148,7 @@ namespace EpicGames.UHT.Exporters.CodeGen
 						if (!classObj.ClassFlags.HasAnyFlags(EClassFlags.Intrinsic) && classObj.ClassType == UhtClassType.Interface && classObj.IsVerseField)
 						{
 							using UhtCodeBlockComment blockComment = new(builder, field, "Verse Interface Proxy");
+							using UhtMacroBlockEmitter macroBlockEmitter = new(builder, UhtDefineScopeNames.Standard, classObj.DefineScope);
 							AppendNativeInterfaceVerseProxy(builder, classObj);
 						}
 					}
@@ -159,6 +163,7 @@ namespace EpicGames.UHT.Exporters.CodeGen
 					if (field is UhtEnum enumObject)
 					{
 						using UhtCodeBlockComment blockComment = new(builder, field);
+						using UhtMacroBlockEmitter macroBlockEmitter = new(builder, UhtDefineScopeNames.Standard, field.DefineScope);
 						AppendEnum(builder, enumObject);
 					}
 				}
diff --git a/Engine/Source/Programs/Shared/EpicGames.UHT/Exporters/CodeGen/UhtUsedDefineScopes.cs b/Engine/Source/Programs/Shared/EpicGames.UHT/Exporters/CodeGen/UhtUsedDefineScopes.cs
--- a/Engine/Source/Programs/Shared/EpicGames.UHT/Exporters/CodeGen/UhtUsedDefineScopes.cs
+++ b/Engine/Source/Programs/Shared/EpicGames.UHT/Exporters/CodeGen/UhtUsedDefineScopes.cs
@@ -4,7 +4,9 @@ using System;
 using System.Collections.Generic;
 using System.Linq;
 using System.Text;
+using System.Xml.Linq;
 using EpicGames.Core;
+using EpicGames.UHT.Parsers;
 using EpicGames.UHT.Types;
 using EpicGames.UHT.Utils;
 
@@ -161,23 +163,6 @@ namespace EpicGames.UHT.Exporters.CodeGen
 		}
 	}
 
-	/// <summary>
-	/// Defines which set of names will be used for each scope
-	/// </summary>
-	internal enum UhtDefineScopeNames
-	{
-
-		/// <summary>
-		/// Names match the define
-		/// </summary>
-		Standard,
-
-		/// <summary>
-		/// Use WITH_EDITOR in place of WITH_EDITORONLY_DATA
-		/// </summary>
-		WithEditor,
-	}
-
 	/// <summary>
 	/// Helper extensions when working with defined scope collections.
 	/// 
@@ -749,7 +734,7 @@ namespace EpicGames.UHT.Exporters.CodeGen
 		}
 
 		/// <summary>
-		/// Append scope expression (i.e. WITH_X || WITH_Y || ...)
+		/// Append scope expression (i.e. WITH_X &amp;&amp; WITH_Y &amp;&amp; ...)
 		/// </summary>
 		/// <param name="builder">String builder</param>
 		/// <param name="defineScope">Scope</param>
@@ -759,23 +744,19 @@ namespace EpicGames.UHT.Exporters.CodeGen
 		{
 			if (defineScope != UhtDefineScope.None && defineScope != UhtDefineScope.Invalid)
 			{
-				bool needOr = false;
-				switch (defineScopeNames)
+				int scopes = (int)defineScope;
+				bool needAnd = false;
+				for (int mask = 1; mask <= scopes; mask <<= 1)
 				{
-					case UhtDefineScopeNames.Standard:
-						builder
-							.AppendName(defineScope, UhtDefineScope.EditorOnlyData, "WITH_EDITORONLY_DATA", ref needOr)
-							.AppendName(defineScope, UhtDefineScope.VerseVM, "WITH_VERSE_VM", ref needOr)
-							;
-						break;
-					case UhtDefineScopeNames.WithEditor:
-						builder
-							.AppendName(defineScope, UhtDefineScope.EditorOnlyData, "WITH_EDITOR", ref needOr)
-							.AppendName(defineScope, UhtDefineScope.VerseVM, "WITH_VERSE_VM", ref needOr)
-							;
-						break;
-					default:
-						throw new UhtIceException("Unexpected define scope names value");
+					if ((mask & scopes) != 0)
+					{
+						if (needAnd)
+						{
+							builder.Append(" && ");
+						}
+						builder.Append(((UhtDefineScope)(mask & scopes)).GetCompilerDirective(defineScopeNames).GetCompilerDirectiveText());
+						needAnd = true;
+					}
 				}
 			}
 			return builder;
@@ -789,19 +770,5 @@ namespace EpicGames.UHT.Exporters.CodeGen
 			}
 			return builder.Append(arrayName);
 		}
-
-		private static StringBuilder AppendName(this StringBuilder builder, UhtDefineScope defineScope, UhtDefineScope testScope, string name, ref bool needOr)
-		{
-			if (defineScope.HasAnyFlags(testScope))
-			{
-				if (needOr)
-				{
-					builder.Append("|| ");
-				}
-				builder.Append(name);
-				needOr = true;
-			}
-			return builder;
-		}
 	}
 }
diff --git a/Engine/Source/Programs/Shared/EpicGames.UHT/Parsers/UhtClassParser.cs b/Engine/Source/Programs/Shared/EpicGames.UHT/Parsers/UhtClassParser.cs
--- a/Engine/Source/Programs/Shared/EpicGames.UHT/Parsers/UhtClassParser.cs
+++ b/Engine/Source/Programs/Shared/EpicGames.UHT/Parsers/UhtClassParser.cs
@@ -112,6 +112,12 @@ namespace EpicGames.UHT.Parsers
 
 					specifiers.ParseDeferred();
 
+					UhtCompilerDirective compilerDirective = topScope.HeaderParser.GetCurrentCompositeCompilerDirective();
+					if ((compilerDirective & UhtCompilerDirective.WithTests) != 0)
+					{
+						classObj.DefineScope |= UhtDefineScope.Tests;
+					}
+
 					classObj.Outer?.AddChild(classObj);
 
 					topScope.AddModuleRelativePathToMetaData();
diff --git a/Engine/Source/Programs/Shared/EpicGames.UHT/Parsers/UhtEnumParser.cs b/Engine/Source/Programs/Shared/EpicGames.UHT/Parsers/UhtEnumParser.cs
--- a/Engine/Source/Programs/Shared/EpicGames.UHT/Parsers/UhtEnumParser.cs
+++ b/Engine/Source/Programs/Shared/EpicGames.UHT/Parsers/UhtEnumParser.cs
@@ -69,11 +69,17 @@ namespace EpicGames.UHT.Parsers
 
 					enumObject.Outer?.AddChild(enumObject);
 
-					if ((topScope.HeaderParser.GetCurrentCompositeCompilerDirective() & UhtCompilerDirective.WithEditorOnlyData) != 0)
+					UhtCompilerDirective compilerDirective = topScope.HeaderParser.GetCurrentCompositeCompilerDirective();
+					if ((compilerDirective & UhtCompilerDirective.WithEditorOnlyData) != 0)
 					{
 						enumObject.DefineScope |= UhtDefineScope.EditorOnlyData;
 					}
 
+					if ((compilerDirective & UhtCompilerDirective.WithTests) != 0)
+					{
+						enumObject.DefineScope |= UhtDefineScope.Tests;
+					}
+
 					// Read base for enum class
 					if (enumObject.CppForm == UhtEnumCppForm.EnumClass)
 					{
diff --git a/Engine/Source/Programs/Shared/EpicGames.UHT/Parsers/UhtFunctionParser.cs b/Engine/Source/Programs/Shared/EpicGames.UHT/Parsers/UhtFunctionParser.cs
--- a/Engine/Source/Programs/Shared/EpicGames.UHT/Parsers/UhtFunctionParser.cs
+++ b/Engine/Source/Programs/Shared/EpicGames.UHT/Parsers/UhtFunctionParser.cs
@@ -355,12 +355,18 @@ namespace EpicGames.UHT.Parsers
 						function.FunctionFlags |= EFunctionFlags.Event;
 					}
 
-					if ((topScope.HeaderParser.GetCurrentCompositeCompilerDirective() & UhtCompilerDirective.WithEditor) != 0)
+					UhtCompilerDirective compilerDirective = topScope.HeaderParser.GetCurrentCompositeCompilerDirective();
+					if ((compilerDirective & UhtCompilerDirective.WithEditor) != 0)
 					{
 						function.FunctionFlags |= EFunctionFlags.EditorOnly;
 						function.DefineScope |= UhtDefineScope.EditorOnlyData;
 					}
 
+					if ((compilerDirective & UhtCompilerDirective.WithTests) != 0)
+					{
+						function.DefineScope |= UhtDefineScope.Tests;
+					}
+
 					specifierParser.ParseDeferred();
 					FinalizeFunctionSpecifiers(function);
 
diff --git a/Engine/Source/Programs/Shared/EpicGames.UHT/Parsers/UhtHeaderFileParser.cs b/Engine/Source/Programs/Shared/EpicGames.UHT/Parsers/UhtHeaderFileParser.cs
--- a/Engine/Source/Programs/Shared/EpicGames.UHT/Parsers/UhtHeaderFileParser.cs
+++ b/Engine/Source/Programs/Shared/EpicGames.UHT/Parsers/UhtHeaderFileParser.cs
@@ -102,6 +102,11 @@ namespace EpicGames.UHT.Parsers
 		/// </summary>
 		Unrecognized = 1 << 10,
 
+		/// <summary>
+		/// This indicates we are in a "#if WITH_TESTS" block
+		/// </summary>
+		WithTests = 1 << 11,
+
 		/// <summary>
 		/// The following flags are always ignored when keywords test for allowed conditional blocks
 		/// </summary>
@@ -110,7 +115,7 @@ namespace EpicGames.UHT.Parsers
 		/// <summary>
 		/// Default compiler directives to be allowed
 		/// </summary>
-		DefaultAllowedCheck = WithEditor | WithEditorOnlyData | WithVerseVM,
+		DefaultAllowedCheck = WithEditor | WithEditorOnlyData | WithVerseVM | WithTests,
 
 		/// <summary>
 		/// All flags are allowed
@@ -170,12 +175,13 @@ namespace EpicGames.UHT.Parsers
 				case UhtCompilerDirective.NotCPPBlock: return "!CPP";
 				case UhtCompilerDirective.ZeroBlock: return "0";
 				case UhtCompilerDirective.OneBlock: return "1";
-				case UhtCompilerDirective.WithHotReload: return "WITH_HOT_RELOAD";
-				case UhtCompilerDirective.WithEditor: return "WITH_EDITOR";
-				case UhtCompilerDirective.WithEditorOnlyData: return "WITH_EDITORONLY_DATA";
-				case UhtCompilerDirective.WithEngine: return "WITH_ENGINE";
-				case UhtCompilerDirective.WithCoreUObject: return "WITH_COREUOBJECT";
-				case UhtCompilerDirective.WithVerseVM: return "WITH_VERSE_VM";
+				case UhtCompilerDirective.WithHotReload: return UhtNames.WithHotReload;
+				case UhtCompilerDirective.WithEditor: return UhtNames.WithEditor;
+				case UhtCompilerDirective.WithEditorOnlyData: return UhtNames.WithEditorOnlyData;
+				case UhtCompilerDirective.WithEngine: return UhtNames.WithEngine;
+				case UhtCompilerDirective.WithCoreUObject: return UhtNames.WithCoreUObject;
+				case UhtCompilerDirective.WithVerseVM: return UhtNames.WithVerseVM;
+				case UhtCompilerDirective.WithTests: return UhtNames.WithTests;
 				default: return "<unrecognized>";
 			}
 		}
@@ -625,7 +631,7 @@ namespace EpicGames.UHT.Parsers
 						UhtCompilerDirective currentDirective = topScope.HeaderParser.GetCurrentCompositeCompilerDirective() & ~UhtCompilerDirective.AllowedCheckIgnoredFlags;
 						if (currentDirective.HasAnyFlags(~keywordInfo.AllowedCompilerDirectives))
 						{
-							List<string> strings = (keywordInfo.AllowedCompilerDirectives).ToStringList();
+							List<string> strings = keywordInfo.AllowedCompilerDirectives.ToStringList();
 							string directives = UhtUtilities.MergeTypeNames(strings, "or", false);
 							topScope.TokenReader.LogError($"'{token.Value}' must not be inside preprocessor blocks, except for {directives}");
 						}
@@ -956,30 +962,34 @@ namespace EpicGames.UHT.Parsers
 					break;
 
 				case UhtTokenType.Identifier:
-					if (define.IsValue("WITH_EDITORONLY_DATA"))
+					if (define.IsValue(UhtNames.WithEditorOnlyData))
 					{
 						return notPresent ? UhtCompilerDirective.Unrecognized : UhtCompilerDirective.WithEditorOnlyData;
 					}
-					else if (define.IsValue("WITH_EDITOR"))
+					else if (define.IsValue(UhtNames.WithEditor))
 					{
 						return notPresent ? UhtCompilerDirective.Unrecognized : UhtCompilerDirective.WithEditor;
 					}
-					else if (define.IsValue("WITH_HOT_RELOAD"))
+					else if (define.IsValue(UhtNames.WithHotReload))
 					{
 						return UhtCompilerDirective.WithHotReload;
 					}
-					else if (define.IsValue("WITH_ENGINE"))
+					else if (define.IsValue(UhtNames.WithEngine))
 					{
 						return notPresent ? UhtCompilerDirective.Unrecognized : UhtCompilerDirective.WithEngine;
 					}
-					else if (define.IsValue("WITH_COREUOBJECT"))
+					else if (define.IsValue(UhtNames.WithCoreUObject))
 					{
 						return notPresent ? UhtCompilerDirective.Unrecognized : UhtCompilerDirective.WithCoreUObject;
 					}
-					else if (define.IsValue("WITH_VERSE_VM"))
+					else if (define.IsValue(UhtNames.WithVerseVM))
 					{
 						return notPresent ? UhtCompilerDirective.Unrecognized : UhtCompilerDirective.WithVerseVM;
 					}
+					else if (define.IsValue(UhtNames.WithTests))
+					{
+						return notPresent ? UhtCompilerDirective.Unrecognized : UhtCompilerDirective.WithTests;
+					}
 					else if (define.IsValue("CPP"))
 					{
 						return notPresent ? UhtCompilerDirective.NotCPPBlock : UhtCompilerDirective.CPPBlock;
@@ -1061,6 +1071,7 @@ namespace EpicGames.UHT.Parsers
 				case UhtCompilerDirective.WithCoreUObject:
 				case UhtCompilerDirective.WithHotReload:
 				case UhtCompilerDirective.WithVerseVM:
+				case UhtCompilerDirective.WithTests:
 					return true;
 
 				default:
@@ -1085,6 +1096,7 @@ namespace EpicGames.UHT.Parsers
 				case UhtCompilerDirective.WithCoreUObject:
 				case UhtCompilerDirective.WithHotReload:
 				case UhtCompilerDirective.WithVerseVM:
+				case UhtCompilerDirective.WithTests:
 					return true;
 
 				default:
@@ -1100,7 +1112,8 @@ namespace EpicGames.UHT.Parsers
 				compilerDirective == UhtCompilerDirective.WithHotReload ||
 				compilerDirective == UhtCompilerDirective.WithEngine ||
 				compilerDirective == UhtCompilerDirective.WithCoreUObject ||
-				compilerDirective == UhtCompilerDirective.WithVerseVM;
+				compilerDirective == UhtCompilerDirective.WithVerseVM ||
+				compilerDirective == UhtCompilerDirective.WithTests;
 		}
 
 		private static void SkipVirtualAndAPI(IUhtTokenReader replayReader)
diff --git a/Engine/Source/Programs/Shared/EpicGames.UHT/Parsers/UhtPropertyParser.cs b/Engine/Source/Programs/Shared/EpicGames.UHT/Parsers/UhtPropertyParser.cs
--- a/Engine/Source/Programs/Shared/EpicGames.UHT/Parsers/UhtPropertyParser.cs
+++ b/Engine/Source/Programs/Shared/EpicGames.UHT/Parsers/UhtPropertyParser.cs
@@ -688,10 +688,14 @@ namespace EpicGames.UHT.Parsers
 						tokenReader.LogError("UProperties should not be wrapped by WITH_EDITOR, use WITH_EDITORONLY_DATA instead.");
 					}
 				}
-				else if (compilerDirective.HasAnyFlags(UhtCompilerDirective.WithVerseVM))
+				if (compilerDirective.HasAnyFlags(UhtCompilerDirective.WithVerseVM))
 				{
 					propertySettings.DefineScope |= UhtDefineScope.VerseVM;
 				}
+				if (compilerDirective.HasAnyFlags(UhtCompilerDirective.WithTests))
+				{
+					propertySettings.DefineScope |= UhtDefineScope.Tests;
+				}
 			}
 
 			// Parse type information including UPARAM that might appear in template arguments
diff --git a/Engine/Source/Programs/Shared/EpicGames.UHT/Parsers/UhtScriptStructParser.cs b/Engine/Source/Programs/Shared/EpicGames.UHT/Parsers/UhtScriptStructParser.cs
--- a/Engine/Source/Programs/Shared/EpicGames.UHT/Parsers/UhtScriptStructParser.cs
+++ b/Engine/Source/Programs/Shared/EpicGames.UHT/Parsers/UhtScriptStructParser.cs
@@ -136,6 +136,12 @@ namespace EpicGames.UHT.Parsers
 					// Process the deferred specifiers
 					specifiers.ParseDeferred();
 
+					UhtCompilerDirective compilerDirective = topScope.HeaderParser.GetCurrentCompositeCompilerDirective();
+					if ((compilerDirective & UhtCompilerDirective.WithTests) != 0)
+					{
+						scriptStruct.DefineScope |= UhtDefineScope.Tests;
+					}
+
 					scriptStruct.Outer?.AddChild(scriptStruct);
 
 					topScope.HeaderParser.ParseStatements('{', '}', true);
diff --git a/Engine/Source/Programs/Shared/EpicGames.UHT/Types/UhtFunction.cs b/Engine/Source/Programs/Shared/EpicGames.UHT/Types/UhtFunction.cs
--- a/Engine/Source/Programs/Shared/EpicGames.UHT/Types/UhtFunction.cs
+++ b/Engine/Source/Programs/Shared/EpicGames.UHT/Types/UhtFunction.cs
@@ -611,6 +611,11 @@ namespace EpicGames.UHT.Types
 
 			ValidateSpecifierFlags();
 
+			if (DefineScope.HasAnyFlags(UhtDefineScope.Tests) && !Outer!.DefineScope.HasAnyFlags(UhtDefineScope.Tests))
+			{
+				this.LogError($"Functions are not supported within a {UhtNames.WithTests} scope.");
+			}
+
 			// Some parameter test are optionally done
 			bool checkForBlueprint = false;
 
diff --git a/Engine/Source/Programs/Shared/EpicGames.UHT/Types/UhtType.cs b/Engine/Source/Programs/Shared/EpicGames.UHT/Types/UhtType.cs
--- a/Engine/Source/Programs/Shared/EpicGames.UHT/Types/UhtType.cs
+++ b/Engine/Source/Programs/Shared/EpicGames.UHT/Types/UhtType.cs
@@ -7,6 +7,7 @@ using System.Diagnostics.CodeAnalysis;
 using System.Text;
 using System.Text.Json.Serialization;
 using System.Threading;
+using EpicGames.UHT.Parsers;
 using EpicGames.UHT.Tables;
 using EpicGames.UHT.Tokenizer;
 using EpicGames.UHT.Utils;
@@ -520,6 +521,23 @@ namespace EpicGames.UHT.Types
 		}
 	}
 
+	/// <summary>
+	/// Defines which set of names will be used for each scope
+	/// </summary>
+	public enum UhtDefineScopeNames
+	{
+
+		/// <summary>
+		/// Names match the define
+		/// </summary>
+		Standard,
+
+		/// <summary>
+		/// Use WITH_EDITOR in place of WITH_EDITORONLY_DATA
+		/// </summary>
+		WithEditor,
+	}
+
 	/// <summary>
 	/// #define scope where the type was defined.  This only includes macros such as WITH_EDITOR, WITH_EDITORONLY_DATA, and WITH_VERSE_VM.
 	/// Not all types support the given options.
@@ -549,10 +567,15 @@ namespace EpicGames.UHT.Types
 		/// </summary>
 		VerseVM = 1 << 1,
 
+		/// <summary>
+		/// Exists inside of a WITH_TESTS block
+		/// </summary>
+		Tests = 1 << 2,
+
 		/// <summary>
 		/// Number of unique scope combinations
 		/// </summary>
-		ScopeCount = 1 << 2,
+		ScopeCount = 1 << 3,
 	}
 
 	/// <summary>
@@ -595,6 +618,27 @@ namespace EpicGames.UHT.Types
 		{
 			return (inFlags & testFlags) == matchFlags;
 		}
+
+		/// <summary>
+		/// Return the compiler directive the matches the given scope
+		/// </summary>
+		/// <param name="defineScope">Scope to return compiler directive</param>
+		/// <param name="defineScopeNames">Used to optionally select WITH_EDITOR instead of WITH_EDITORONLYDATA</param>
+		/// <returns></returns>
+		public static UhtCompilerDirective GetCompilerDirective(this UhtDefineScope defineScope, UhtDefineScopeNames defineScopeNames)
+		{
+			switch (defineScope)
+			{
+				case UhtDefineScope.EditorOnlyData:
+					return defineScopeNames == UhtDefineScopeNames.WithEditor ? UhtCompilerDirective.WithEditor : UhtCompilerDirective.WithEditorOnlyData;
+				case UhtDefineScope.VerseVM:
+					return UhtCompilerDirective.WithVerseVM;
+				case UhtDefineScope.Tests:
+					return UhtCompilerDirective.WithTests;
+				default:
+					throw new UhtIceException("A single valid define scope must be passed to GetCompilerDirective()");
+			}
+		}
 	}
 
 	/// <summary>
diff --git a/Engine/Source/Programs/Shared/EpicGames.UHT/Utils/UhtNames.cs b/Engine/Source/Programs/Shared/EpicGames.UHT/Utils/UhtNames.cs
--- a/Engine/Source/Programs/Shared/EpicGames.UHT/Utils/UhtNames.cs
+++ b/Engine/Source/Programs/Shared/EpicGames.UHT/Utils/UhtNames.cs
@@ -67,7 +67,16 @@ namespace EpicGames.UHT.Utils
 		public const string UIMax = "UIMax";
 		public const string Visible = "Visible";
 
-		// Other names
+		// Verse names
 		public const string VerseProxySuffix = "__Proxy";
+
+		// Compiler directives
+		public const string WithHotReload = "WITH_HOT_RELOAD";
+		public const string WithEditor = "WITH_EDITOR";
+		public const string WithEditorOnlyData = "WITH_EDITORONLY_DATA";
+		public const string WithEngine = "WITH_ENGINE";
+		public const string WithCoreUObject = "WITH_COREUOBJECT";
+		public const string WithVerseVM = "WITH_VERSE_VM";
+		public const string WithTests = "WITH_TESTS";
 	}
 }
