UnrealBuildTool: Support path hashes for UBA. This is used to reduce size of cache entries in uba cache server

#rnx

diff --git a/Engine/Source/Programs/Shared/EpicGames.UBA/ICacheClient.cs b/Engine/Source/Programs/Shared/EpicGames.UBA/ICacheClient.cs
--- a/Engine/Source/Programs/Shared/EpicGames.UBA/ICacheClient.cs
+++ b/Engine/Source/Programs/Shared/EpicGames.UBA/ICacheClient.cs
@@ -25,6 +25,13 @@ namespace EpicGames.UBA
 		/// <returns>True if successful</returns>
 		public abstract bool Connect(string host, int port);
 
+		/// <summary>
+		/// Register path with string that will be hashed. All files under this path will be ignored and instead refer to this path hash
+		/// </summary>
+		/// <param name="path">Path that will represent all files under that path</param>
+		/// <param name="hash">String that is unique to the content of data under path</param>
+		public abstract bool RegisterPathHash(string path, string hash);
+
 		/// <summary>
 		/// Write to cache
 		/// </summary>
diff --git a/Engine/Source/Programs/Shared/EpicGames.UBA/Impl/CacheClientImpl.cs b/Engine/Source/Programs/Shared/EpicGames.UBA/Impl/CacheClientImpl.cs
--- a/Engine/Source/Programs/Shared/EpicGames.UBA/Impl/CacheClientImpl.cs
+++ b/Engine/Source/Programs/Shared/EpicGames.UBA/Impl/CacheClientImpl.cs
@@ -46,6 +46,10 @@ namespace EpicGames.UBA.Impl
 
 		[DllImport("UbaHost", CharSet = CharSet.Auto)]
 		static extern void CacheClient_Destroy(nint cacheClient);
+
+		[DllImport("UbaHost", CharSet = CharSet.Auto)]
+		static extern byte CacheClient_RegisterPathHash(nint cacheClient, string path, string hashString);
+
 		#endregion
 
 		#region IDisposable
@@ -79,7 +83,22 @@ namespace EpicGames.UBA.Impl
 
 		public bool Connect(string host, int port)
 		{
-			return CacheClient_Connect(_handle, host, port) != 0;
+			if (CacheClient_Connect(_handle, host, port) == 0)
+			{
+				return false;
+			}
+
+			foreach (KeyValuePair<string, string> item in Utils.PathHashes)
+			{
+				CacheClient_RegisterPathHash(_handle, item.Key, item.Value);
+			}
+
+			return true;
+		}
+
+		public bool RegisterPathHash(string path, string hash)
+		{
+			return CacheClient_RegisterPathHash(_handle, path, hash) != 0;
 		}
 
 		public bool WriteToCache(uint bucket, IProcess process, byte[] inputs, uint inputsSize, byte[] outputs, uint outputsSize)
diff --git a/Engine/Source/Programs/Shared/EpicGames.UBA/Impl/StorageServerImpl.cs b/Engine/Source/Programs/Shared/EpicGames.UBA/Impl/StorageServerImpl.cs
--- a/Engine/Source/Programs/Shared/EpicGames.UBA/Impl/StorageServerImpl.cs
+++ b/Engine/Source/Programs/Shared/EpicGames.UBA/Impl/StorageServerImpl.cs
@@ -31,7 +31,7 @@ namespace EpicGames.UBA.Impl
 			_server = server;
 			_logger = logger;
 			_handle = StorageServer_Create(_server.GetHandle(), info.RootDirectory, info.CapacityBytes, info.StoreCompressed, _logger.GetHandle(), info.Zone);
-			Utils.DisallowedPaths().ToList().ForEach(x => RegisterDisallowedPath(x));
+			Utils.DisallowedPaths.ToList().ForEach(RegisterDisallowedPath);
 		}
 
 		#region IDisposable
diff --git a/Engine/Source/Programs/Shared/EpicGames.UBA/Utils.cs b/Engine/Source/Programs/Shared/EpicGames.UBA/Utils.cs
--- a/Engine/Source/Programs/Shared/EpicGames.UBA/Utils.cs
+++ b/Engine/Source/Programs/Shared/EpicGames.UBA/Utils.cs
@@ -25,16 +25,22 @@ namespace EpicGames.UBA
 		/// Paths that are not allowed to be transferred over the network for UBA remote agents.
 		/// </summary>
 		/// <returns>enumerable of disallowed paths</returns>
-		public static IEnumerable<string> DisallowedPaths() => s_disallowedPaths.Distinct().Order();
+		public static IEnumerable<string> DisallowedPaths => s_disallowedPaths.Distinct().Order();
 		static readonly ConcurrentBag<string> s_disallowedPaths = [];
 
 		/// <summary>
 		/// Mapping of binary paths for cross architecture host binaries, to allow for using helpers of a different architecture.
 		/// </summary>
 		/// <returns>dictionary of binary mappings, where the key is the binary for the current host architecture</returns>
-		public static IDictionary<string, string> CrossArchitecturePaths => s_crossArchitecturePaths;
+		public static IReadOnlyDictionary<string, string> CrossArchitecturePaths => s_crossArchitecturePaths;
 		static readonly ConcurrentDictionary<string, string> s_crossArchitecturePaths = [];
 
+		/// <summary>
+		/// Mapping of a folder path to a single hash, to allow for hashing an entire folder so each individual file does not need to be processed.
+		/// </summary>
+		public static IReadOnlyDictionary<string, string> PathHashes => s_pathHashes;
+		static readonly ConcurrentDictionary<string, string> s_pathHashes = [];
+
 		/// <summary>
 		/// Registers a path that is not allowed to be transferred over the network for UBA remote agents.
 		/// </summary>
@@ -43,7 +49,10 @@ namespace EpicGames.UBA
 		{
 			foreach (string path in paths)
 			{
-				s_disallowedPaths.Add(path);
+				if (!s_disallowedPaths.Contains(path))
+				{
+					s_disallowedPaths.Add(path);
+				}
 			}
 			DisallowedPathRegistered?.Invoke(s_disallowedPaths, new(paths));
 		}
@@ -55,8 +64,23 @@ namespace EpicGames.UBA
 		/// <param name="crossPath">cross architecture path</param>
 		public static void RegisterCrossArchitecturePath(string path, string crossPath)
 		{
-			s_crossArchitecturePaths.TryAdd(path, crossPath);
-			CrossArchitecturePathRegistered?.Invoke(s_crossArchitecturePaths, new(path, crossPath));
+			if (s_crossArchitecturePaths.TryAdd(path, crossPath))
+			{
+				CrossArchitecturePathRegistered?.Invoke(s_crossArchitecturePaths, new(path, crossPath));
+			}
+		}
+
+		/// <summary>
+		/// Registers a hash for a path
+		/// </summary>
+		/// <param name="path">The path string</param>
+		/// <param name="hash">The hash string</param>
+		public static void RegisterPathHash(string path, string hash)
+		{
+			if (s_pathHashes.TryAdd(path, hash))
+			{
+				PathHashRegistered?.Invoke(s_pathHashes, new(path, hash));
+			}
 		}
 
 		/// <summary>
@@ -64,14 +88,21 @@ namespace EpicGames.UBA
 		/// </summary>
 		/// <param name="sender">collection that is being changed</param>
 		/// <param name="e">event args containing which paths were added</param>
-		public delegate void DisallowedPathRegisteredEventHandler(ConcurrentBag<string> sender, DisallowedPathRegisteredEventArgs e);
+		public delegate void DisallowedPathRegisteredEventHandler(IReadOnlyCollection<string> sender, DisallowedPathRegisteredEventArgs e);
 
 		/// <summary>
 		/// Delegate for registering a cross architecture path
 		/// </summary>
 		/// <param name="sender">collection that is being changed</param>
 		/// <param name="e">event args containing which path was added</param>
-		public delegate void CrossArchitecturePathRegisteredEventHandler(ConcurrentDictionary<string, string> sender, CrossArchitecturePathRegisteredEventArgs e);
+		public delegate void CrossArchitecturePathRegisteredEventHandler(IReadOnlyDictionary<string, string> sender, CrossArchitecturePathRegisteredEventArgs e);
+
+		/// <summary>
+		/// Delegate for registering a path hash
+		/// </summary>
+		/// <param name="sender">collection that is being changed</param>
+		/// <param name="e">event args containing which path was added</param>
+		public delegate void PathHashRegisteredEventHandler(IReadOnlyDictionary<string, string> sender, PathHashRegisteredEventArgs e);
 
 		/// <summary>
 		/// Remote disallowed path registered event handler
@@ -83,6 +114,11 @@ namespace EpicGames.UBA
 		/// </summary>
 		public static event CrossArchitecturePathRegisteredEventHandler? CrossArchitecturePathRegistered;
 
+		/// <summary>
+		/// Remote disallowed path registered event handler
+		/// </summary>
+		public static event PathHashRegisteredEventHandler? PathHashRegistered;
+
 		/// <summary>
 		/// Get the path to the p/invoke library that would be loaded 
 		/// </summary>
@@ -112,7 +148,7 @@ namespace EpicGames.UBA
 	/// <summary>
 	/// Event args for registering a remote disallowed path
 	/// </summary>
-	public class DisallowedPathRegisteredEventArgs(params string[] paths) : EventArgs
+	public sealed class DisallowedPathRegisteredEventArgs(params string[] paths) : EventArgs
 	{
 		/// <summary>
 		/// The paths being registered
@@ -123,16 +159,32 @@ namespace EpicGames.UBA
 	/// <summary>
 	/// Event args for registering a cross architecture path
 	/// </summary>
-	public class CrossArchitecturePathRegisteredEventArgs(string path, string crossPath) : EventArgs
+	public sealed class CrossArchitecturePathRegisteredEventArgs(string path, string crossPath) : EventArgs
 	{
 		/// <summary>
 		/// The host architecture path being registered
 		/// </summary>
-		public string Paths { get; } = path;
+		public string Path { get; } = path;
 
 		/// <summary>
 		/// The cross architecture path being registered
 		/// </summary>
 		public string CrossPath { get; } = crossPath;
 	}
+
+	/// <summary>
+	/// Event args for registering a path hash
+	/// </summary>
+	public sealed class PathHashRegisteredEventArgs(string path, string hash) : EventArgs
+	{
+		/// <summary>
+		/// The path being registered
+		/// </summary>
+		public string Path { get; } = path;
+
+		/// <summary>
+		/// The hash for the path being registered
+		/// </summary>
+		public string Hash { get; } = hash;
+	}
 }
diff --git a/Engine/Source/Programs/UnrealBuildTool/Executors/UnrealBuildAccelerator/UBAAgentCoordinatorHorde.cs b/Engine/Source/Programs/UnrealBuildTool/Executors/UnrealBuildAccelerator/UBAAgentCoordinatorHorde.cs
--- a/Engine/Source/Programs/UnrealBuildTool/Executors/UnrealBuildAccelerator/UBAAgentCoordinatorHorde.cs
+++ b/Engine/Source/Programs/UnrealBuildTool/Executors/UnrealBuildAccelerator/UBAAgentCoordinatorHorde.cs
@@ -570,7 +570,7 @@ namespace UnrealBuildTool
 						//arguments.Add("-NoStore");
 						//arguments.Add("-KillRandom"); // For debugging
 
-						if (OperatingSystem.IsMacOS() && EpicGames.UBA.Utils.DisallowedPaths().Any())
+						if (OperatingSystem.IsMacOS() && EpicGames.UBA.Utils.DisallowedPaths.Any())
 						{
 							// we need to populate the cas with all known xcodes so we can serve all the ones we have installed
 							string xcodeVersion = Utils.RunLocalProcessAndReturnStdOut("/bin/sh", "-c '/usr/bin/defaults read $(xcode-select -p)/../version.plist ProductBuildVersion");
diff --git a/Engine/Source/Programs/UnrealBuildTool/Platform/Windows/VCToolChain.cs b/Engine/Source/Programs/UnrealBuildTool/Platform/Windows/VCToolChain.cs
--- a/Engine/Source/Programs/UnrealBuildTool/Platform/Windows/VCToolChain.cs
+++ b/Engine/Source/Programs/UnrealBuildTool/Platform/Windows/VCToolChain.cs
@@ -3470,6 +3470,14 @@ namespace UnrealBuildTool
 			}
 			GlobalCompileEnvironment.RootPaths[CppRootPathFolder.WinSDK] = EnvVars.WindowsSdkDir;
 			GlobalLinkEnvironment.RootPaths[CppRootPathFolder.WinSDK] = EnvVars.WindowsSdkDir;
+
+			// This is used by uba to create one input key for all files under dirs
+			EpicGames.UBA.Utils.RegisterPathHash(EnvVars.WindowsSdkDir.FullName, EnvVars.WindowsSdkVersion.ToString());
+			EpicGames.UBA.Utils.RegisterPathHash(EnvVars.CompilerDir.FullName, EnvVars.CompilerVersion.ToString());
+			if (EnvVars.Compiler != EnvVars.ToolChain)
+			{
+				EpicGames.UBA.Utils.RegisterPathHash(EnvVars.ToolChainDir.FullName, EnvVars.ToolChainVersion.ToString());
+			}
 		}
 
 		public override void ModifyBuildProducts(ReadOnlyTargetRules Target, UEBuildBinary Binary, IEnumerable<string> Libraries, IEnumerable<UEBuildBundleResource> BundleResources, Dictionary<FileReference, BuildProductType> BuildProducts)
diff --git a/Engine/Source/Programs/UnrealBuildTool/System/TargetMakefile.cs b/Engine/Source/Programs/UnrealBuildTool/System/TargetMakefile.cs
--- a/Engine/Source/Programs/UnrealBuildTool/System/TargetMakefile.cs
+++ b/Engine/Source/Programs/UnrealBuildTool/System/TargetMakefile.cs
@@ -361,12 +361,27 @@ namespace UnrealBuildTool
 		/// <summary>
 		/// Set of external (ie. user-owned) files which will cause the makefile to be invalidated if modified after
 		/// </summary>
-		public HashSet<FileItem> ExternalDependencies = new HashSet<FileItem>();
+		public HashSet<FileItem> ExternalDependencies;
 
 		/// <summary>
 		/// Set of internal (eg. response files, unity files) which will cause the makefile to be invalidated if modified.
 		/// </summary>
-		public HashSet<FileItem> InternalDependencies = new HashSet<FileItem>();
+		public HashSet<FileItem> InternalDependencies;
+
+		/// <summary>
+		/// Set of paths that are disallowed to be transferred with UBA to remote machines
+		/// </summary>
+		public HashSet<DirectoryItem> UbaDisallowedPaths;
+
+		/// <summary>
+		/// Mapping of binary paths for cross architecture host binaries, to allow for using UBA helpers of a different architecture
+		/// </summary>
+		public Dictionary<DirectoryItem, DirectoryItem> UbaCrossArchitecturePaths;
+
+		/// <summary>
+		/// Set of path hashes that are registered with UBA, to prevent needing to individually hash each file
+		/// </summary>
+		public Dictionary<DirectoryItem, string> UbaPathHashes;
 
 		/// <summary>
 		/// TargetRules-set memory estimate per action. Used to control the number of parallel actions that are spawned.
@@ -445,6 +460,12 @@ namespace UnrealBuildTool
 			PluginFiles = new HashSet<FileItem>();
 			ExternalDependencies = new HashSet<FileItem>();
 			InternalDependencies = new HashSet<FileItem>();
+
+			UbaDisallowedPaths = [.. EpicGames.UBA.Utils.DisallowedPaths.Select(DirectoryItem.GetItemByPath)];
+			UbaCrossArchitecturePaths = new(EpicGames.UBA.Utils.CrossArchitecturePaths.Select(x => new KeyValuePair<DirectoryItem, DirectoryItem>(DirectoryItem.GetItemByPath(x.Key), DirectoryItem.GetItemByPath(x.Value))));
+			UbaPathHashes = new(EpicGames.UBA.Utils.PathHashes.Select(x => new KeyValuePair<DirectoryItem, string>(DirectoryItem.GetItemByPath(x.Key), x.Value)));
+
+			RegisterUBAHandlers();
 		}
 
 		/// <summary>
@@ -492,7 +513,29 @@ namespace UnrealBuildTool
 			PluginFiles = Reader.ReadHashSet(() => Reader.ReadFileItem())!;
 			ExternalDependencies = Reader.ReadHashSet(() => Reader.ReadFileItem())!;
 			InternalDependencies = Reader.ReadHashSet(() => Reader.ReadFileItem())!;
+			UbaDisallowedPaths = Reader.ReadHashSet(() => Reader.ReadDirectoryItem())!;
+			UbaCrossArchitecturePaths = Reader.ReadDictionary(() => Reader.ReadDirectoryItem()!, () => Reader.ReadDirectoryItem())!;
+			UbaPathHashes = Reader.ReadDictionary(() => Reader.ReadDirectoryItem()!, () => Reader.ReadString())!;
 			MemoryPerActionGB = Reader.ReadDouble();
+
+			EpicGames.UBA.Utils.RegisterDisallowedPaths(UbaDisallowedPaths.Select(x => x.FullName).ToArray());
+			foreach (KeyValuePair<DirectoryItem, DirectoryItem> item in UbaCrossArchitecturePaths)
+			{
+				EpicGames.UBA.Utils.RegisterCrossArchitecturePath(item.Key.FullName, item.Value.FullName);
+			}
+			foreach (KeyValuePair<DirectoryItem, string> item in UbaPathHashes)
+			{
+				EpicGames.UBA.Utils.RegisterPathHash(item.Key.FullName, item.Value);
+			}
+
+			RegisterUBAHandlers();
+		}
+
+		void RegisterUBAHandlers()
+		{
+			EpicGames.UBA.Utils.DisallowedPathRegistered += (sender, e) => UbaDisallowedPaths.UnionWith(e.Paths.Select(DirectoryItem.GetItemByPath));
+			EpicGames.UBA.Utils.CrossArchitecturePathRegistered += (sender, e) => UbaCrossArchitecturePaths.Add(DirectoryItem.GetItemByPath(e.Path), DirectoryItem.GetItemByPath(e.CrossPath));
+			EpicGames.UBA.Utils.PathHashRegistered += (sender, e) => UbaPathHashes.Add(DirectoryItem.GetItemByPath(e.Path), e.Hash);
 		}
 
 		/// <summary>
@@ -542,6 +585,9 @@ namespace UnrealBuildTool
 			Writer.WriteHashSet(PluginFiles, x => Writer.WriteFileItem(x));
 			Writer.WriteHashSet(ExternalDependencies, x => Writer.WriteFileItem(x));
 			Writer.WriteHashSet(InternalDependencies, x => Writer.WriteFileItem(x));
+			Writer.WriteHashSet(UbaDisallowedPaths, Writer.WriteDirectoryItem);
+			Writer.WriteDictionary(UbaCrossArchitecturePaths, Writer.WriteDirectoryItem, Writer.WriteDirectoryItem);
+			Writer.WriteDictionary(UbaPathHashes, Writer.WriteDirectoryItem, Writer.WriteString);
 			Writer.WriteDouble(MemoryPerActionGB);
 		}
 
@@ -871,6 +917,17 @@ namespace UnrealBuildTool
 					}
 				}
 
+				IEnumerable<DirectoryItem> UbaPaths = [.. Makefile.UbaDisallowedPaths, .. Makefile.UbaPathHashes.Keys, .. Makefile.UbaCrossArchitecturePaths.Keys, .. Makefile.UbaCrossArchitecturePaths.Values];
+				foreach (DirectoryItem item in UbaPaths.Distinct())
+				{
+					if (!item.Exists)
+					{
+						Logger.LogDebug("{Directory} has been deleted since makefile was written.", item.Location);
+						ReasonNotLoaded = String.Format("{0} deleted", item.Location);
+						return false;
+					}
+				}
+
 				// Check that no new plugins have been added
 				foreach (FileReference PluginFile in PluginsBase.EnumeratePlugins(ProjectFile))
 				{
diff --git a/Engine/Source/Programs/UnrealBuildTool/ToolChain/ClangToolChain.cs b/Engine/Source/Programs/UnrealBuildTool/ToolChain/ClangToolChain.cs
--- a/Engine/Source/Programs/UnrealBuildTool/ToolChain/ClangToolChain.cs
+++ b/Engine/Source/Programs/UnrealBuildTool/ToolChain/ClangToolChain.cs
@@ -246,6 +246,8 @@ namespace UnrealBuildTool
 				GlobalLinkEnvironment.RootPaths.AddFolderName(CppRootPathFolder.PlatformSDK, $"{Target.Platform}SDK");
 				GlobalCompileEnvironment.RootPaths[CppRootPathFolder.PlatformSDK] = LazyInfo.Value.BasePath;
 				GlobalLinkEnvironment.RootPaths[CppRootPathFolder.PlatformSDK] = LazyInfo.Value.BasePath;
+
+				EpicGames.UBA.Utils.RegisterPathHash(LazyInfo.Value.BasePath.FullName, LazyInfo.Value.ClangVersion.ToString());
 			}
 
 			PreprocessDepends = Target.bPreprocessDepends;
