[UBA]
* Changed bool to byte in all pinvoke calls because it seems like macos does not like bool conversions

diff --git a/Engine/Source/Programs/Shared/EpicGames.UBA/Impl/CacheClientImpl.cs b/Engine/Source/Programs/Shared/EpicGames.UBA/Impl/CacheClientImpl.cs
--- a/Engine/Source/Programs/Shared/EpicGames.UBA/Impl/CacheClientImpl.cs
+++ b/Engine/Source/Programs/Shared/EpicGames.UBA/Impl/CacheClientImpl.cs
@@ -15,13 +15,13 @@ namespace EpicGames.UBA.Impl
 
 		#region DllImport
 		[DllImport("UbaHost", CharSet = CharSet.Auto)]
-		static extern nint ProcessStartInfo_Create2(string application, string arguments, string workingDir, string description, uint priorityClass, ulong rootsHandle, bool trackInputs, string logFile, ExitCallback? exit);
+		static extern nint ProcessStartInfo_Create2(string application, string arguments, string workingDir, string description, uint priorityClass, ulong rootsHandle, byte trackInputs, string logFile, ExitCallback? exit);
 
 		[DllImport("UbaHost", CharSet = CharSet.Auto)]
 		static extern void ProcessStartInfo_Destroy(nint server);
 
 		[DllImport("UbaHost", CharSet = CharSet.Auto)]
-		static extern nint CacheClient_Create(nint session, bool reportMissReason, string crypto, string hint);
+		static extern nint CacheClient_Create(nint session, byte reportMissReason, string crypto, string hint);
 
 		[DllImport("UbaHost", CharSet = CharSet.Auto)]
 		static extern byte CacheClient_Connect(nint cacheClient, string host, int port);
@@ -78,7 +78,7 @@ namespace EpicGames.UBA.Impl
 		public CacheClientImpl(ISessionServer server, bool reportMissReason, string crypto, string hint)
 		{
 			_sessionServer = server;
-			_handle = CacheClient_Create(_sessionServer.GetHandle(), reportMissReason, crypto, hint);
+			_handle = CacheClient_Create(_sessionServer.GetHandle(), (byte)(reportMissReason?1:0), crypto, hint);
 		}
 
 		public bool Connect(string host, int port)
@@ -108,7 +108,7 @@ namespace EpicGames.UBA.Impl
 
 		public FetchFromCacheResult FetchFromCache(ulong rootsHandle, uint bucket, ProcessStartInfo info)
 		{
-			nint si = ProcessStartInfo_Create2(info.Application, info.Arguments, info.WorkingDirectory, info.Description, (uint)info.Priority, rootsHandle, info.TrackInputs, info.LogFile ?? String.Empty, null);
+			nint si = ProcessStartInfo_Create2(info.Application, info.Arguments, info.WorkingDirectory, info.Description, (uint)info.Priority, rootsHandle, (byte)(info.TrackInputs?1:0), info.LogFile ?? String.Empty, null);
 			nint result = CacheClient_FetchFromCache3(_handle, rootsHandle, bucket, si);
 			ProcessStartInfo_Destroy(si);
 
diff --git a/Engine/Source/Programs/Shared/EpicGames.UBA/Impl/ProcessImpl.cs b/Engine/Source/Programs/Shared/EpicGames.UBA/Impl/ProcessImpl.cs
--- a/Engine/Source/Programs/Shared/EpicGames.UBA/Impl/ProcessImpl.cs
+++ b/Engine/Source/Programs/Shared/EpicGames.UBA/Impl/ProcessImpl.cs
@@ -16,7 +16,7 @@ namespace EpicGames.UBA.Impl
 
 		#region DllImport
 		[DllImport("UbaHost", CharSet = CharSet.Auto)]
-		static extern nint ProcessStartInfo_Create2(string application, string arguments, string workingDir, string description, uint priorityClass, ulong rootsHandle, bool trackInputs, string logFile, ExitCallback? exit);
+		static extern nint ProcessStartInfo_Create2(string application, string arguments, string workingDir, string description, uint priorityClass, ulong rootsHandle, byte trackInputs, string logFile, ExitCallback? exit);
 
 		[DllImport("UbaHost", CharSet = CharSet.Auto)]
 		static extern void ProcessStartInfo_Destroy(nint server);
@@ -29,7 +29,7 @@ namespace EpicGames.UBA.Impl
 				ExitCallbackDelegate = RaiseExited;
 			}
 
-			_handle = ProcessStartInfo_Create2(info.Application, info.Arguments, info.WorkingDirectory, info.Description, (uint)info.Priority, info.RootsHandle, info.TrackInputs, info.LogFile ?? String.Empty, ExitCallbackDelegate);
+			_handle = ProcessStartInfo_Create2(info.Application, info.Arguments, info.WorkingDirectory, info.Description, (uint)info.Priority, info.RootsHandle, (byte)(info.TrackInputs?1:0), info.LogFile ?? String.Empty, ExitCallbackDelegate);
 		}
 
 		void RaiseExited(nint userData, nint handle)
@@ -96,7 +96,7 @@ namespace EpicGames.UBA.Impl
 		static extern ulong ProcessHandle_GetHash(nint process);
 
 		[DllImport("UbaHost", CharSet = CharSet.Auto)]
-		static extern void ProcessHandle_Cancel(nint process, bool terminate);
+		static extern void ProcessHandle_Cancel(nint process, byte terminate);
 
 		[DllImport("UbaHost", CharSet = CharSet.Auto)]
 		static extern void DestroyProcessHandle(nint process);
@@ -168,7 +168,7 @@ namespace EpicGames.UBA.Impl
 
 		public object? UserData { get; }
 
-		public void Cancel(bool terminate) => ProcessHandle_Cancel(_handle, terminate);
+		public void Cancel(bool terminate) => ProcessHandle_Cancel(_handle, (byte)(terminate?1:0));
 		#endregion
 
 		internal void RaiseExited()
diff --git a/Engine/Source/Programs/Shared/EpicGames.UBA/Impl/ServerImpl.cs b/Engine/Source/Programs/Shared/EpicGames.UBA/Impl/ServerImpl.cs
--- a/Engine/Source/Programs/Shared/EpicGames.UBA/Impl/ServerImpl.cs
+++ b/Engine/Source/Programs/Shared/EpicGames.UBA/Impl/ServerImpl.cs
@@ -12,23 +12,23 @@ namespace EpicGames.UBA.Impl
 
 		#region DllImport
 		[DllImport("UbaHost", CharSet = CharSet.Auto)]
-		static extern nint NetworkServer_Create(nint logger, int workerCount, int sendSize, int receiveTimeoutSeconds, bool useQuic);
+		static extern nint NetworkServer_Create(nint logger, int workerCount, int sendSize, int receiveTimeoutSeconds, byte useQuic);
 
 		[DllImport("UbaHost", CharSet = CharSet.Auto)]
 		static extern void NetworkServer_Destroy(nint server);
 
 		// Server Imports
 		[DllImport("UbaHost", CharSet = CharSet.Auto)]
-		static extern bool NetworkServer_StartListen(nint server, int port, string ip, string crypto);
+		static extern byte NetworkServer_StartListen(nint server, int port, string ip, string crypto);
 
 		[DllImport("UbaHost", CharSet = CharSet.Auto)]
-		static extern bool NetworkServer_AddClient(nint server, string ip, int port, string crypto);
+		static extern byte NetworkServer_AddClient(nint server, string ip, int port, string crypto);
 
 		[DllImport("UbaHost", CharSet = CharSet.Auto)]
-		static extern bool NetworkServer_AddNamedConnection(nint server, string name);
+		static extern byte NetworkServer_AddNamedConnection(nint server, string name);
 
 		[DllImport("UbaHost", CharSet = CharSet.Auto)]
-		static extern bool NetworkServer_Stop(nint server);
+		static extern byte NetworkServer_Stop(nint server);
 
 		const int DEFAULT_PORT = 1345;
 		#endregion
@@ -36,7 +36,7 @@ namespace EpicGames.UBA.Impl
 		public ServerImpl(int maxWorkers, int sendSize, ILogger logger, bool useQuic)
 		{
 			_logger = logger;
-			_handle = NetworkServer_Create(_logger.GetHandle(), maxWorkers, sendSize, 60, useQuic); // We know this code is using 5s pings.. so 60 seconds timeout on recv socket is good
+			_handle = NetworkServer_Create(_logger.GetHandle(), maxWorkers, sendSize, 60, (byte)(useQuic?1:0)); // We know this code is using 5s pings.. so 60 seconds timeout on recv socket is good
 		}
 
 		#region IDisposable
@@ -68,16 +68,16 @@ namespace EpicGames.UBA.Impl
 		public nint GetHandle() => _handle;
 
 		/// <inheritdoc/>
-		public bool StartServer(string ip = "", int port = -1, string crypto = "") => NetworkServer_StartListen(_handle, port > 0 ? port : DEFAULT_PORT, ip, crypto);
+		public bool StartServer(string ip = "", int port = -1, string crypto = "") => NetworkServer_StartListen(_handle, port > 0 ? port : DEFAULT_PORT, ip, crypto) != 0;
 
 		/// <inheritdoc/>
 		public void StopServer() => NetworkServer_Stop(_handle);
 
 		/// <inheritdoc/>
-		public bool AddClient(string ip, int port, string crypto = "") => NetworkServer_AddClient(_handle, ip, port, crypto);
+		public bool AddClient(string ip, int port, string crypto = "") => NetworkServer_AddClient(_handle, ip, port, crypto) != 0;
 
 		/// <inheritdoc/>
-		public bool AddNamedConnection(string name) => NetworkServer_AddNamedConnection(_handle, name);
+		public bool AddNamedConnection(string name) => NetworkServer_AddNamedConnection(_handle, name) != 0;
 		#endregion
 	}
 }
diff --git a/Engine/Source/Programs/Shared/EpicGames.UBA/Impl/SessionServerImpl.cs b/Engine/Source/Programs/Shared/EpicGames.UBA/Impl/SessionServerImpl.cs
--- a/Engine/Source/Programs/Shared/EpicGames.UBA/Impl/SessionServerImpl.cs
+++ b/Engine/Source/Programs/Shared/EpicGames.UBA/Impl/SessionServerImpl.cs
@@ -19,7 +19,7 @@ namespace EpicGames.UBA.Impl
 		#region DllImport
 		[DllImport("UbaHost", CharSet = CharSet.Auto)]
 		static extern nint SessionServerCreateInfo_Create(nint storage, nint client, nint logger, string rootDir, string traceOutputFile,
-			bool disableCustomAllocator, bool launchVisualizer, bool resetCas, bool writeToDisk, bool detailedTrace, bool allowWaitOnMem, bool allowKillOnMem, bool storeObjFilesCompressed);
+			byte disableCustomAllocator, byte launchVisualizer, byte resetCas, byte writeToDisk, byte detailedTrace, byte allowWaitOnMem, byte allowKillOnMem, byte storeObjFilesCompressed);
 
 		[DllImport("UbaHost", CharSet = CharSet.Auto)]
 		static extern void SessionServerCreateInfo_Destroy(nint server);
@@ -30,7 +30,7 @@ namespace EpicGames.UBA.Impl
 			_storage = storage;
 			_client = client;
 			_logger = logger;
-			_handle = SessionServerCreateInfo_Create(_storage.GetHandle(), _client.GetHandle(), _logger.GetHandle(), info.RootDirectory, info.TraceOutputFile, info.DisableCustomAllocator, info.LaunchVisualizer, info.ResetCas, info.WriteToDisk, info.DetailedTrace, info.AllowWaitOnMem, info.AllowKillOnMem, info.StoreObjFilesCompressed);
+			_handle = SessionServerCreateInfo_Create(_storage.GetHandle(), _client.GetHandle(), _logger.GetHandle(), info.RootDirectory, info.TraceOutputFile, (byte)(info.DisableCustomAllocator?1:0), (byte)(info.LaunchVisualizer?1:0), (byte)(info.ResetCas?1:0), (byte)(info.WriteToDisk?1:0), (byte)(info.DetailedTrace?1:0), (byte)(info.AllowWaitOnMem?1:0), (byte)(info.AllowKillOnMem?1:0), (byte)(info.StoreObjFilesCompressed?1:0));
 		}
 
 		#region IDisposable
diff --git a/Engine/Source/Programs/Shared/EpicGames.UBA/Impl/StorageServerImpl.cs b/Engine/Source/Programs/Shared/EpicGames.UBA/Impl/StorageServerImpl.cs
--- a/Engine/Source/Programs/Shared/EpicGames.UBA/Impl/StorageServerImpl.cs
+++ b/Engine/Source/Programs/Shared/EpicGames.UBA/Impl/StorageServerImpl.cs
@@ -14,7 +14,7 @@ namespace EpicGames.UBA.Impl
 
 		#region DllImport
 		[DllImport("UbaHost", CharSet = CharSet.Auto)]
-		static extern nint StorageServer_Create(nint server, string rootDir, ulong casCapacityBytes, bool storeCompressed, nint logger, string zone);
+		static extern nint StorageServer_Create(nint server, string rootDir, ulong casCapacityBytes, byte storeCompressed, nint logger, string zone);
 
 		[DllImport("UbaHost", CharSet = CharSet.Auto)]
 		static extern void StorageServer_Destroy(nint server);
@@ -30,7 +30,7 @@ namespace EpicGames.UBA.Impl
 		{
 			_server = server;
 			_logger = logger;
-			_handle = StorageServer_Create(_server.GetHandle(), info.RootDirectory, info.CapacityBytes, info.StoreCompressed, _logger.GetHandle(), info.Zone);
+			_handle = StorageServer_Create(_server.GetHandle(), info.RootDirectory, info.CapacityBytes, (byte)(info.StoreCompressed?1:0), _logger.GetHandle(), info.Zone);
 			Utils.DisallowedPaths.ToList().ForEach(RegisterDisallowedPath);
 		}
 
