UBT: Fix static analyzer errors in DataRouterTelemetryService.cs (for changes in cl 41100779)

#rnx

diff --git a/Engine/Source/Programs/Shared/EpicGames.Core/Telemetry/DataRouterTelemetryService.cs b/Engine/Source/Programs/Shared/EpicGames.Core/Telemetry/DataRouterTelemetryService.cs
--- a/Engine/Source/Programs/Shared/EpicGames.Core/Telemetry/DataRouterTelemetryService.cs
+++ b/Engine/Source/Programs/Shared/EpicGames.Core/Telemetry/DataRouterTelemetryService.cs
@@ -353,8 +353,6 @@ namespace EpicGames.Core.Telemetry
 
 		private readonly Timer _autoFlushTimer = new Timer();
 
-		private Task<HttpResponseMessage> InitialHeadRequest;
-
 		/// <summary>
 		/// Length of an empty data router post when serialized to json
 		/// </summary>
@@ -411,7 +409,7 @@ namespace EpicGames.Core.Telemetry
 			HttpClient = inHttpClient;
 			BaseAddress = inBaseAddress;
 			SetHttpClientHeaders();
-			InitialHeadRequest = HeadRequestAsync();
+			Task<HttpResponseMessage> initializeConnectionFireAndForget = HeadRequestAsync();
 
 			_autoFlushTimer.Interval = AutoFlushIntervalMilliseconds;
 			_autoFlushTimer.Elapsed += AutoFlushEventsAsync;
