EpicGames.UHT: Don't cull output directories for modules that are not set to save exported headers (due to being read-only / precompiled)

#rnx

diff --git a/Engine/Source/Programs/Shared/EpicGames.UHT/Utils/UhtSession.cs b/Engine/Source/Programs/Shared/EpicGames.UHT/Utils/UhtSession.cs
--- a/Engine/Source/Programs/Shared/EpicGames.UHT/Utils/UhtSession.cs
+++ b/Engine/Source/Programs/Shared/EpicGames.UHT/Utils/UhtSession.cs
@@ -413,7 +413,7 @@ namespace EpicGames.UHT.Utils
 
 			// These outputs are used to cull old outputs from the directories
 			Dictionary<string, HashSet<string>> outputsByDirectory = new(
-				Session.Modules.Select(x => new KeyValuePair<string, HashSet<string>>(x.Module.OutputDirectory, new HashSet<string>(StringComparer.OrdinalIgnoreCase))),
+				Session.Modules.Where(x => x.Module.SaveExportedHeaders).Select(x => new KeyValuePair<string, HashSet<string>>(x.Module.OutputDirectory, new HashSet<string>(StringComparer.OrdinalIgnoreCase))),
 				StringComparer.OrdinalIgnoreCase
 			);
 
