Correct sort order in changelists returned from Unreal Cloud.

[RN] none

diff --git a/Engine/Source/Programs/Shared/EpicGames.ProjectStore/ProjectStoreTypes.cs b/Engine/Source/Programs/Shared/EpicGames.ProjectStore/ProjectStoreTypes.cs
--- a/Engine/Source/Programs/Shared/EpicGames.ProjectStore/ProjectStoreTypes.cs
+++ b/Engine/Source/Programs/Shared/EpicGames.ProjectStore/ProjectStoreTypes.cs
@@ -888,6 +888,8 @@ namespace EpicGames.ProjectStore
 				throw new BuildIndexException("No snapshots found");
 			}
 
+			builds.SortBy(b => b.Metadata.Changelist);
+
 			if (!getClosestPreceding && builds.Last().Metadata.Changelist != changelist)
 			{
 				throw new BuildIndexException("Exact match for target changelist requested, but no snapshots matching that changelist were found");
