PR #13069: Enable setting exporter status from the command line for all UBT plugins

#rnx

diff --git a/Engine/Source/Programs/Shared/EpicGames.UHT/Utils/UhtSession.cs b/Engine/Source/Programs/Shared/EpicGames.UHT/Utils/UhtSession.cs
--- a/Engine/Source/Programs/Shared/EpicGames.UHT/Utils/UhtSession.cs
+++ b/Engine/Source/Programs/Shared/EpicGames.UHT/Utils/UhtSession.cs
@@ -1052,7 +1052,8 @@ namespace EpicGames.UHT.Utils
 		/// Run UHT on the given manifest.  Use the bHasError property to see if process was successful.
 		/// </summary>
 		/// <param name="manifestFilePath">Path to the manifest file</param>
-		public void Run(string manifestFilePath)
+		/// <param name="arguments">Any command line arguments</param>
+		public void Run(string manifestFilePath, CommandLineArguments? arguments = null)
 		{
 			if (FileManager == null)
 			{
@@ -1124,6 +1125,22 @@ namespace EpicGames.UHT.Utils
 			}
 
 			StepReadManifestFile(manifestFilePath);
+
+			if (arguments != null)
+			{
+				foreach (UhtExporter Exporter in ExporterTable)
+				{
+					if (arguments.HasOption($"-{Exporter.Name}"))
+					{
+						SetExporterStatus(Exporter.Name, true);
+					}
+					else if (arguments.HasOption($"-no{Exporter.Name}"))
+					{
+						SetExporterStatus(Exporter.Name, false);
+					}
+				}
+			}
+
 			StepPrepareModules();
 			StepPrepareHeaders();
 			StepParseHeaders();
diff --git a/Engine/Source/Programs/UnrealBuildTool/Modes/UnrealHeaderToolMode.cs b/Engine/Source/Programs/UnrealBuildTool/Modes/UnrealHeaderToolMode.cs
--- a/Engine/Source/Programs/UnrealBuildTool/Modes/UnrealHeaderToolMode.cs
+++ b/Engine/Source/Programs/UnrealBuildTool/Modes/UnrealHeaderToolMode.cs
@@ -1313,20 +1313,8 @@ namespace UnrealBuildTool.Modes
 					Session.ReferenceMode = UhtReferenceMode.Verify;
 				}
 
-				foreach (UhtExporter Exporter in Session.ExporterTable)
-				{
-					if (Arguments.HasOption($"-{Exporter.Name}"))
-					{
-						Session.SetExporterStatus(Exporter.Name, true);
-					}
-					else if (Arguments.HasOption($"-no{Exporter.Name}"))
-					{
-						Session.SetExporterStatus(Exporter.Name, false);
-					}
-				}
-
 				// Read and parse
-				Session.Run(ManifestPath!);
+				Session.Run(ManifestPath!, Arguments);
 				Session.LogMessages();
 				return (int)(Session.HasErrors ? CompilationResult.OtherCompilationError : CompilationResult.Succeeded);
 			}
