Introduce a second tier of event ids filtering from Gauntlet events
* Second tier(test focus) will prune other events if reported
* Prune warnings if errors are found
* Redact path and http link more broadly
* Limit the number of lines to consider for hashing

#rnx

diff --git a/Engine/Source/Programs/Horde/Plugins/Build/HordeServer.Build.Tests/Issues/IssueServiceTests.cs b/Engine/Source/Programs/Horde/Plugins/Build/HordeServer.Build.Tests/Issues/IssueServiceTests.cs
--- a/Engine/Source/Programs/Horde/Plugins/Build/HordeServer.Build.Tests/Issues/IssueServiceTests.cs
+++ b/Engine/Source/Programs/Horde/Plugins/Build/HordeServer.Build.Tests/Issues/IssueServiceTests.cs
@@ -2178,6 +2178,59 @@ namespace HordeServer.Tests.Issues
 				Assert.AreEqual("Automation build drop errors in Update Version Files", issues[0].Summary);
 				Assert.AreEqual("Automation fatal errors in Update Version Files", issues[1].Summary);
 			}
+			// #6
+			// Scenario: 2 Test events + Gauntlet Framework + Device events
+			// Expected: Test events are treated as test errors and are the only 2 issues that are kept
+			{
+				IJob job = CreateJob(_mainStreamId, 160, "Test Build", _graph);
+				await using (TestJsonLogger logger = await CreateLoggerAsync(job, 0, 0))
+				{
+					logger.LogError(KnownLogEvents.Gauntlet_TestEvent, "Test encountered an error at some point");
+					logger.LogWarning(KnownLogEvents.Gauntlet_DeviceEvent, "Failed to setup device");
+					logger.LogError(KnownLogEvents.Gauntlet, "Generic Gauntlet failure");
+				}
+				IJob job2 = CreateJob(_mainStreamId, 170, "Test Build", _graph);
+				await using (TestJsonLogger logger = await CreateLoggerAsync(job2, 0, 0))
+				{
+					logger.LogError(KnownLogEvents.Gauntlet_TestEvent, "Test encountered an error at some point");
+					logger.LogError(KnownLogEvents.Gauntlet_UnrealEngineTestEvent, "Unreal Test encountered an error");
+					logger.LogWarning(KnownLogEvents.Gauntlet_DeviceEvent, "Failed to setup device");
+					logger.LogError(KnownLogEvents.Gauntlet, "Generic Gauntlet failure");
+				}
+				await UpdateCompleteStepAsync(job, 0, 0, JobStepOutcome.Failure);
+				await UpdateCompleteStepAsync(job2, 0, 0, JobStepOutcome.Failure);
+
+				IReadOnlyList<IIssue> issues = await IssueCollection.FindIssuesAsync();
+				Assert.AreEqual(2, issues.Count);
+				Assert.AreEqual("Gauntlet:test", issues[0].Fingerprints[0].Type);
+				Assert.AreEqual("Gauntlet:test", issues[1].Fingerprints[0].Type);
+				Assert.AreEqual("Automation test errors in Update Version Files", issues[0].Summary);
+				Assert.AreEqual("Automation test errors in Update Version Files", issues[1].Summary);
+			}
+			// #7
+			// Scenario: Gauntlet Framework + Device events
+			// Expected: Only one issue is generated
+			{
+				IJob job = CreateJob(_mainStreamId, 180, "Test Build", _graph);
+				await using (TestJsonLogger logger = await CreateLoggerAsync(job, 0, 0))
+				{
+					logger.LogWarning(KnownLogEvents.Gauntlet_DeviceEvent, "Failed to setup device");
+					logger.LogError(KnownLogEvents.Gauntlet, "Generic Gauntlet failure");
+				}
+				IJob job2 = CreateJob(_mainStreamId, 190, "Test Build", _graph);
+				await using (TestJsonLogger logger = await CreateLoggerAsync(job2, 0, 0))
+				{
+					logger.LogWarning(KnownLogEvents.Gauntlet_DeviceEvent, "Failed to setup device");
+					logger.LogError(KnownLogEvents.Gauntlet, "Generic Gauntlet failure");
+				}
+				await UpdateCompleteStepAsync(job, 0, 0, JobStepOutcome.Failure);
+				await UpdateCompleteStepAsync(job2, 0, 0, JobStepOutcome.Failure);
+
+				IReadOnlyList<IIssue> issues = await IssueCollection.FindIssuesAsync();
+				Assert.AreEqual(1, issues.Count);
+				Assert.AreEqual("Gauntlet:test framework", issues[0].Fingerprints[0].Type);
+				Assert.AreEqual("Automation test framework errors in Update Version Files", issues[0].Summary);
+			}
 		}
 
 		[TestMethod]
diff --git a/Engine/Source/Programs/Shared/EpicGames.Horde/Issues/Handlers/GauntletIssueHandler.cs b/Engine/Source/Programs/Shared/EpicGames.Horde/Issues/Handlers/GauntletIssueHandler.cs
--- a/Engine/Source/Programs/Shared/EpicGames.Horde/Issues/Handlers/GauntletIssueHandler.cs
+++ b/Engine/Source/Programs/Shared/EpicGames.Horde/Issues/Handlers/GauntletIssueHandler.cs
@@ -1,5 +1,6 @@
 // Copyright Epic Games, Inc. All Rights Reserved.
 
+using System;
 using System.Collections.Generic;
 using System.Linq;
 using System.Text;
@@ -51,6 +52,11 @@ namespace EpicGames.Horde.Issues.Handlers
 		/// </summary>
 		const string SummaryLogType = "Summary";
 
+		/// <summary>
+		/// Max Number of lines to consider to hash
+		/// </summary>
+		const int MaxLines = 6;
+
 		/// <summary>
 		/// Max Message Length to hash
 		/// </summary>
@@ -61,6 +67,16 @@ namespace EpicGames.Horde.Issues.Handlers
 		/// </summary>
 		bool _wasSevereEventReported = false;
 
+		/// <summary>
+		/// Whether or not a test event was reported
+		/// </summary>
+		bool _wasTestEventReported = false;
+
+		/// <summary>
+		/// Whether or not at least one error was reported
+		/// </summary>
+		bool _wasErrorEventReported = false;
+
 		readonly IssueHandlerContext _context;
 		readonly List<IssueEventGroup> _issues = new List<IssueEventGroup>();
 
@@ -86,6 +102,15 @@ namespace EpicGames.Horde.Issues.Handlers
 			KnownLogEvents.Gauntlet_FatalEvent
 		};
 
+		/// <summary>
+		///  Known Gauntlet events associated with test context
+		/// </summary>
+		static readonly HashSet<EventId> s_knownTestGauntletEvents = new HashSet<EventId>
+		{
+			KnownLogEvents.Gauntlet_TestEvent,
+			KnownLogEvents.Gauntlet_UnrealEngineTestEvent
+		};
+
 		/// <inheritdoc/>
 		public override int Priority => 10;
 
@@ -148,10 +173,10 @@ namespace EpicGames.Horde.Issues.Handlers
 		{
 			// Use only the summary if one is found instead of the full callstack
 			string sanitized = GetSummaryProperty(issueEvent) ?? issueEvent.Render();
+			sanitized = String.Join("\n", sanitized.Split("\n", MaxLines + 1).Take(MaxLines)); // Limit the number of lines to consider
 			sanitized = sanitized.Length > MaxMessageLength ? sanitized.Substring(0, MaxMessageLength) : sanitized;
 			sanitized = sanitized.Trim().ToUpperInvariant();
-			sanitized = Regex.Replace(sanitized, @"(?<![A-Z])(?:[A-Z]:|/)[^ :]+[/\\]SYNC[/\\]", "{root}/"); // Redact things that look like workspace roots; may be different between agents
-			sanitized = Regex.Replace(sanitized, @"(?:^|\s)/[\w/.:-]+", "{path}"); // Redact path like reference
+			sanitized = Regex.Replace(sanitized, @"(?:(?<![A-Z])|HTTPS*:)(?:[A-Z]:|[/\\])[\w/\\.:+?=%-]+", "{path}"); // Redact path like reference
 			sanitized = Regex.Replace(sanitized, @"0X[0-9A-F]+", "H"); // Redact hex strings
 			sanitized = Regex.Replace(sanitized, @"\d[\d.,:]*", "n"); // Redact numbers and timestamp like things
 
@@ -232,7 +257,11 @@ namespace EpicGames.Horde.Issues.Handlers
 			if (issueEvent.EventId != null && IsMatchingEventId(issueEvent.EventId.Value))
 			{
 				bool isSevereEvent = s_knownSevereGauntletEvents.Contains(issueEvent.EventId.Value);
-				if (_wasSevereEventReported && !isSevereEvent)
+				bool isTestEvent = !isSevereEvent && s_knownTestGauntletEvents.Contains(issueEvent.EventId.Value);
+				bool isErrorEvent = issueEvent.Severity >= LogLevel.Error;
+				if ((_wasSevereEventReported && !isSevereEvent)
+					|| (_wasTestEventReported && !isTestEvent && !isSevereEvent)
+					|| (_wasErrorEventReported && !isErrorEvent))
 				{
 					return true;
 				}
@@ -248,7 +277,14 @@ namespace EpicGames.Horde.Issues.Handlers
 					string? hash = issue.Metadata.FindValues("Hash").FirstOrDefault();
 					issue.Type = $"{issue.Type}:with-callstack:{hash}";
 				}
-				_issues.Add(issue);
+
+				if (isErrorEvent && !_wasErrorEventReported)
+				{
+					// We've encountered an error event;
+					// We can ignore other issues to prevent superfluous issues from being created
+					_issues.RemoveAll((issue) => issue.Events.Any(x => x.Severity <= LogLevel.Warning));
+					_wasErrorEventReported = true;
+				}
 
 				if (isSevereEvent)
 				{
@@ -257,6 +293,15 @@ namespace EpicGames.Horde.Issues.Handlers
 					_issues.RemoveAll((issue) => !issue.Events.Any(x => s_knownSevereGauntletEvents.Contains(x.EventId!.Value)));
 					_wasSevereEventReported = true;
 				}
+				else if (isTestEvent && !_wasTestEventReported)
+				{
+					// We've encountered a test event;
+					// We can ignore other issues to prevent superfluous issues from being created
+					_issues.RemoveAll((issue) => !issue.Events.Any(x => s_knownTestGauntletEvents.Contains(x.EventId!.Value)));
+					_wasTestEventReported = true;
+				}
+
+				_issues.Add(issue);
 
 				return true;
 			}
