Implemented suspends support for Verse native callable functions.

#rnx

diff --git a/Engine/Source/Programs/Shared/EpicGames.UHT/Exporters/CodeGen/UhtCodeGenerator.cs b/Engine/Source/Programs/Shared/EpicGames.UHT/Exporters/CodeGen/UhtCodeGenerator.cs
--- a/Engine/Source/Programs/Shared/EpicGames.UHT/Exporters/CodeGen/UhtCodeGenerator.cs
+++ b/Engine/Source/Programs/Shared/EpicGames.UHT/Exporters/CodeGen/UhtCodeGenerator.cs
@@ -45,6 +45,7 @@ namespace EpicGames.UHT.Exporters.CodeGen
 			public bool NeedsVerseCodeGen { get; set; }
 			public bool NeedsVerseField { get; set; }
 			public bool NeedsVersePointer { get; set; }
+			public bool NeedsVerseTask { get; set; }
 		}
 		public HeaderInfo[] HeaderInfos { get; set; }
 
@@ -427,6 +428,10 @@ namespace EpicGames.UHT.Exporters.CodeGen
 					{
 						headerInfo.NeedsVersePointer = true;
 					}
+					if (childProperty.PropertyExportFlags.HasAnyFlags(UhtPropertyExportFlags.TVerseTask))
+					{
+						headerInfo.NeedsVerseTask = true;
+					}
 					if (!headerInfo.NeedsVerseValue && childProperty is UhtVerseValueProperty)
 					{
 						headerInfo.NeedsVerseValue = true;
diff --git a/Engine/Source/Programs/Shared/EpicGames.UHT/Exporters/CodeGen/UhtHeaderCodeGenerator.cs b/Engine/Source/Programs/Shared/EpicGames.UHT/Exporters/CodeGen/UhtHeaderCodeGenerator.cs
--- a/Engine/Source/Programs/Shared/EpicGames.UHT/Exporters/CodeGen/UhtHeaderCodeGenerator.cs
+++ b/Engine/Source/Programs/Shared/EpicGames.UHT/Exporters/CodeGen/UhtHeaderCodeGenerator.cs
@@ -718,19 +718,41 @@ namespace EpicGames.UHT.Exporters.CodeGen
 		public static StringBuilder AppendVerseNativeCallableReturnType(this StringBuilder builder, UhtFunction function)
 		{
 			UhtProperty? returnProperty = function.ReturnProperty;
-			// if (coroutine) "FVerseTask"
+			bool isIncompleteReturn = function.Session.IsIncompleteReturn(returnProperty);
+			bool isSuspends = function.FunctionExportFlags.HasAnyFlags(UhtFunctionExportFlags.VerseSuspends);
+			bool isIfNoRuntimeError = returnProperty != null && !isIncompleteReturn;
+
+			if (isSuspends)
+			{
+				builder.Append("TVerseTask<");
+			}
+			else if (isIfNoRuntimeError)
+			{
+				builder.Append("TIfNoRuntimeError<");
+			}
+
 			if (returnProperty == null)
 			{
-				builder.Append("void ");
+				builder.Append("void");
+			}
+			else if (isIncompleteReturn)
+			{
+				builder.Append("EVerseNativeCallResult");
 			}
-			else if (returnProperty is UhtEnumProperty enumProperty && enumProperty.Enum == function.Session.EVerseNativeCallResult)
+			else
+			{
+				builder.AppendTokens(returnProperty.TypeTokens.AllTokens);
+			}
+
+			if (isSuspends || isIfNoRuntimeError)
 			{
-				builder.Append("EVerseNativeCallResult ");
+				builder.Append("> ");
 			}
 			else
 			{
-				builder.Append("TIfNoRuntimeError<").AppendTokens(returnProperty.TypeTokens.AllTokens).Append("> ");
+				builder.Append(' ');
 			}
+
 			return builder;
 		}
 
@@ -765,14 +787,18 @@ namespace EpicGames.UHT.Exporters.CodeGen
 		public static StringBuilder AppendVerseNativeCallableTypeDef(this StringBuilder builder, UhtFunction function)
 		{
 			UhtProperty? returnProperty = function.ReturnProperty;
-			bool isIncompleteReturn = returnProperty is UhtEnumProperty enumProperty && enumProperty.Enum == function.Session.EVerseNativeCallResult;
+			bool isIncompleteReturn = function.Session.IsIncompleteReturn(returnProperty);
+			bool isSuspends = function.FunctionExportFlags.HasAnyFlags(UhtFunctionExportFlags.VerseSuspends);
 			if (!isIncompleteReturn && function.FunctionExportFlags.HasAnyFlags(UhtFunctionExportFlags.VerseDecides) && returnProperty is UhtOptionalProperty optionalReturn)
 			{
 				returnProperty = optionalReturn.ValueProperty;
 			}
 
 			builder.Append("TVerseFunction<");
-			// if (coroutine) "FVerseTask"
+			if (isSuspends)
+			{
+				builder.Append("FVerseResult(TVerseCall<");
+			}
 			if (returnProperty == null)
 			{
 				builder.Append("void");
@@ -785,14 +811,16 @@ namespace EpicGames.UHT.Exporters.CodeGen
 			{
 				builder.AppendTokens(returnProperty.TypeTokens.MyTokens);
 			}
-			builder.Append('(').AppendVerseFunctionArgs(function, false, true, property => builder.AppendTokens(property.TypeTokens.MyTokens), () => builder.Append("FDecidesContext")).Append(")>");
+
+
+			builder.Append(isSuspends ? '>' : '(').AppendVerseFunctionArgs(function, false, true, property => builder.AppendTokens(property.TypeTokens.MyTokens), () => builder.Append("FDecidesContext")).Append(")>");
 			return builder;
 		}
 
 		public static StringBuilder AppendVerseFunctionArgs(this StringBuilder builder, UhtFunction function, bool startWithSeparator, bool decidesBeforeReturn, Action<UhtProperty> propertyAction, Action decidesAction)
 		{
 			bool insertSeparator = startWithSeparator;
-			if (function.HasVerseIncompleteReturn)
+			if (function.Session.IsIncompleteReturn(function.ReturnProperty))
 			{
 				if (decidesBeforeReturn && function.FunctionExportFlags.HasAnyFlags(UhtFunctionExportFlags.VerseDecides))
 				{
diff --git a/Engine/Source/Programs/Shared/EpicGames.UHT/Exporters/CodeGen/UhtHeaderCodeGeneratorCppFile.cs b/Engine/Source/Programs/Shared/EpicGames.UHT/Exporters/CodeGen/UhtHeaderCodeGeneratorCppFile.cs
--- a/Engine/Source/Programs/Shared/EpicGames.UHT/Exporters/CodeGen/UhtHeaderCodeGeneratorCppFile.cs
+++ b/Engine/Source/Programs/Shared/EpicGames.UHT/Exporters/CodeGen/UhtHeaderCodeGeneratorCppFile.cs
@@ -1191,13 +1191,15 @@ namespace EpicGames.UHT.Exporters.CodeGen
 		private static StringBuilder AppendVerseNativeCallableFunctionStub(StringBuilder builder, UhtClass classObj, UhtFunction function)
 		{
 			UhtProperty? returnProperty = function.ReturnProperty;
+			bool isSuspends = function.FunctionExportFlags.HasAnyFlags(UhtFunctionExportFlags.VerseSuspends);
+
 			if (returnProperty != null)
 			{
 				builder.AppendTokens(returnProperty.TypeTokens.AllTokens).Append(' ');
 			}
 			else
 			{
-				builder.Append("void ");
+				builder.Append(isSuspends ? "TVerseTask<void> " : "void ");
 			}
 
 			builder.AppendClassSourceNameOrInterfaceName(classObj).Append("::").Append(function.SourceName).Append('(');
@@ -1216,7 +1218,7 @@ namespace EpicGames.UHT.Exporters.CodeGen
 			}
 			builder.Append(")\r\n");
 			builder.Append("{\r\n");
-			if (returnProperty != null)
+			if (returnProperty != null || isSuspends)
 			{
 				builder.Append("\treturn ");
 			}
@@ -1234,7 +1236,7 @@ namespace EpicGames.UHT.Exporters.CodeGen
 			}
 			if (returnProperty != null)
 			{
-				if (returnProperty is UhtEnumProperty enumProperty && enumProperty.Enum == function.Session.EVerseNativeCallResult)
+				if (function.Session.IsIncompleteReturn(returnProperty))
 				{
 					builder.Append(");\r\n");
 				}
@@ -1254,29 +1256,24 @@ namespace EpicGames.UHT.Exporters.CodeGen
 		private static StringBuilder AppendVerseNativeCallableFunction(StringBuilder builder, UhtClass classObj, UhtFunction function, string classNameSuffix)
 		{
 			UhtProperty? returnProperty = function.ReturnProperty;
-			bool isIncompleteReturn = returnProperty is UhtEnumProperty enumProperty && enumProperty.Enum == function.Session.EVerseNativeCallResult;
 
 			builder.AppendVerseNativeCallableSignature(classObj, function, classNameSuffix).Append("\r\n");
 			builder.Append("{\r\n");
 
 			// Result value declaration
-			if (isIncompleteReturn)
+			if (returnProperty != null || function.FunctionExportFlags.HasAnyFlags(UhtFunctionExportFlags.VerseSuspends))
 			{
-				builder.Append("\tEVerseNativeCallResult Result;\r\n");
-			}
-			else if (returnProperty != null)
-			{
-				builder.Append("\tTIfNoRuntimeError<").AppendTokens(returnProperty.TypeTokens.AllTokens).Append("> Result;\r\n");
+				builder.Append('\t').AppendVerseNativeCallableReturnType(function).Append("Result;\r\n");
 			}
 
 			// The bodies must differ based on the VM
 			builder.Append("#if WITH_VERSE_VM\r\n");
-			AppendVerseNativeCallableFunctionBody(builder, classObj, function, classNameSuffix, true);
+			AppendVerseNativeCallableFunctionBody(builder, classObj, function, classNameSuffix, returnProperty, true);
 			builder.Append("#else\r\n");
-			AppendVerseNativeCallableFunctionBody(builder, classObj, function, classNameSuffix, false);
+			AppendVerseNativeCallableFunctionBody(builder, classObj, function, classNameSuffix, returnProperty, false);
 			builder.Append("#endif\r\n");
 
-			if (returnProperty != null)
+			if (returnProperty != null || function.FunctionExportFlags.HasAnyFlags(UhtFunctionExportFlags.VerseSuspends))
 			{
 				builder.Append("\treturn Result;\r\n");
 			}
@@ -1284,7 +1281,7 @@ namespace EpicGames.UHT.Exporters.CodeGen
 			return builder;
 		}
 
-		private static StringBuilder AppendVerseNativeCallableFunctionBody(StringBuilder builder, UhtClass classObj, UhtFunction function, string classNameSuffix, bool isVerseVM)
+		private static StringBuilder AppendVerseNativeCallableFunctionBody(StringBuilder builder, UhtClass classObj, UhtFunction function, string classNameSuffix, UhtProperty? returnProperty, bool isVerseVM)
 		{
 			builder.Append(isVerseVM ? "\tAutoRTFM::Open([&] {\r\n" : "\t{\r\n");
 
@@ -1320,7 +1317,7 @@ namespace EpicGames.UHT.Exporters.CodeGen
 
 			// Invocation
 			builder.Append("\t\t");
-			if (function.ReturnProperty != null)
+			if (returnProperty != null || function.FunctionExportFlags.HasAnyFlags(UhtFunctionExportFlags.VerseSuspends))
 			{
 				builder.Append("Result = ");
 			}
diff --git a/Engine/Source/Programs/Shared/EpicGames.UHT/Exporters/CodeGen/UhtHeaderCodeGeneratorHFile.cs b/Engine/Source/Programs/Shared/EpicGames.UHT/Exporters/CodeGen/UhtHeaderCodeGeneratorHFile.cs
--- a/Engine/Source/Programs/Shared/EpicGames.UHT/Exporters/CodeGen/UhtHeaderCodeGeneratorHFile.cs
+++ b/Engine/Source/Programs/Shared/EpicGames.UHT/Exporters/CodeGen/UhtHeaderCodeGeneratorHFile.cs
@@ -90,6 +90,10 @@ namespace EpicGames.UHT.Exporters.CodeGen
 				{
 					builder.Append("#include \"VersePointer.h\"\r\n");
 				}
+				if (headerInfo.NeedsVerseTask)
+				{
+					builder.Append("#include \"VerseTask.h\"\r\n");
+				}
 
 				builder.Append("\r\n");
 				builder.Append(DisableDeprecationWarnings).Append("\r\n");
diff --git a/Engine/Source/Programs/Shared/EpicGames.UHT/Parsers/UhtFunctionParser.cs b/Engine/Source/Programs/Shared/EpicGames.UHT/Parsers/UhtFunctionParser.cs
--- a/Engine/Source/Programs/Shared/EpicGames.UHT/Parsers/UhtFunctionParser.cs
+++ b/Engine/Source/Programs/Shared/EpicGames.UHT/Parsers/UhtFunctionParser.cs
@@ -448,6 +448,10 @@ namespace EpicGames.UHT.Parsers
 								{
 									returnValueProperty = property;
 								}
+								if (property.PropertyExportFlags.HasAnyFlags(UhtPropertyExportFlags.TVerseTask))
+								{
+									function.FunctionExportFlags |= UhtFunctionExportFlags.VerseSuspends;
+								}
 							});
 
 						if (funcNameToken.Value.Length == 0)
@@ -459,7 +463,7 @@ namespace EpicGames.UHT.Parsers
 						function.SourceName = funcNameToken.Value.ToString();
 
 						// If the function is a verse function, then we need to apply the source name to the verse name
-						if (function.IsVerseField)
+						if (function.IsVerseField && String.IsNullOrEmpty(function.VerseName))
 						{
 							function.VerseName = function.SourceName;
 						}
diff --git a/Engine/Source/Programs/Shared/EpicGames.UHT/Parsers/UhtPropertyParser.cs b/Engine/Source/Programs/Shared/EpicGames.UHT/Parsers/UhtPropertyParser.cs
--- a/Engine/Source/Programs/Shared/EpicGames.UHT/Parsers/UhtPropertyParser.cs
+++ b/Engine/Source/Programs/Shared/EpicGames.UHT/Parsers/UhtPropertyParser.cs
@@ -1135,7 +1135,7 @@ namespace EpicGames.UHT.Parsers
 	[UnrealHeaderTool]
 	static class UhtVerseTypeParser
 	{
-		[UhtPropertyType(Keyword = "TVal")]
+		[UhtPropertyType(Keyword = "TVal", Options = UhtPropertyTypeOptions.Immediate)]
 		[SuppressMessage("CodeQuality", "IDE0051:Remove unused private members", Justification = "Attribute accessed method")]
 		[SuppressMessage("Style", "IDE0060:Remove unused parameter", Justification = "Attribute accessed method")]
 		private static UhtProperty? TValType(UhtPropertyResolveArgs args)
@@ -1148,11 +1148,15 @@ namespace EpicGames.UHT.Parsers
 			args.PropertySettings.PropertyExportFlags |= UhtPropertyExportFlags.TVal;
 			args.TokenReader.Require('<');
 			UhtProperty? inner = args.ParseWrappedType();
+			if (inner == null)
+			{
+				return null;
+			}
 			args.TokenReader.Require('>');
 			return inner;
 		}
 
-		[UhtPropertyType(Keyword = "verse::TPtr")]
+		[UhtPropertyType(Keyword = "verse::TPtr", Options = UhtPropertyTypeOptions.Immediate)]
 		[SuppressMessage("CodeQuality", "IDE0051:Remove unused private members", Justification = "Attribute accessed method")]
 		[SuppressMessage("Style", "IDE0060:Remove unused parameter", Justification = "Attribute accessed method")]
 		private static UhtProperty? VerseTPtrType(UhtPropertyResolveArgs args)
@@ -1165,6 +1169,31 @@ namespace EpicGames.UHT.Parsers
 			args.PropertySettings.PropertyExportFlags |= UhtPropertyExportFlags.TPtr;
 			args.TokenReader.Require('<');
 			UhtProperty? inner = args.ParseWrappedType();
+			if (inner == null)
+			{
+				return null;
+			}
+			args.TokenReader.Require('>');
+			return inner;
+		}
+
+		[UhtPropertyType(Keyword = "TVerseTask", Options = UhtPropertyTypeOptions.Immediate)]
+		[SuppressMessage("CodeQuality", "IDE0051:Remove unused private members", Justification = "Attribute accessed method")]
+		[SuppressMessage("Style", "IDE0060:Remove unused parameter", Justification = "Attribute accessed method")]
+		private static UhtProperty? VerseTaskType(UhtPropertyResolveArgs args)
+		{
+			if (!args.SkipExpectedType())
+			{
+				return null;
+			}
+
+			args.PropertySettings.PropertyExportFlags |= UhtPropertyExportFlags.TVerseTask;
+			args.TokenReader.Require('<');
+			UhtProperty? inner = args.ParseWrappedType();
+			if (inner == null)
+			{
+				return null;
+			}
 			args.TokenReader.Require('>');
 			return inner;
 		}
diff --git a/Engine/Source/Programs/Shared/EpicGames.UHT/Specifiers/UhtFunctionSpecifiers.cs b/Engine/Source/Programs/Shared/EpicGames.UHT/Specifiers/UhtFunctionSpecifiers.cs
--- a/Engine/Source/Programs/Shared/EpicGames.UHT/Specifiers/UhtFunctionSpecifiers.cs
+++ b/Engine/Source/Programs/Shared/EpicGames.UHT/Specifiers/UhtFunctionSpecifiers.cs
@@ -427,13 +427,20 @@ namespace EpicGames.UHT.Parsers
 				{
 					functionObj.FunctionExportFlags |= UhtFunctionExportFlags.VerseDecides;
 				}
+				else if (key.Equals("Name", StringComparison.OrdinalIgnoreCase))
+				{
+					functionObj.VerseName = kvp.Value.ToString();
+				}
 				else
 				{
 					functionObj.LogError($"Verse specifier option '{key}' is unknown");
 				}
 			}
 
-			functionObj.VerseName = functionObj.SourceName;
+			if (String.IsNullOrEmpty(functionObj.VerseName))
+			{
+				functionObj.VerseName = functionObj.SourceName;
+			}
 		}
 
 		private static void ParseNetServiceIdentifiers(UhtSpecifierContext specifierContext, List<KeyValuePair<StringView, StringView>> identifiers)
diff --git a/Engine/Source/Programs/Shared/EpicGames.UHT/Specifiers/UhtPropertyMemberSpecifiers.cs b/Engine/Source/Programs/Shared/EpicGames.UHT/Specifiers/UhtPropertyMemberSpecifiers.cs
--- a/Engine/Source/Programs/Shared/EpicGames.UHT/Specifiers/UhtPropertyMemberSpecifiers.cs
+++ b/Engine/Source/Programs/Shared/EpicGames.UHT/Specifiers/UhtPropertyMemberSpecifiers.cs
@@ -481,11 +481,11 @@ namespace EpicGames.UHT.Parsers
 			context.PropertySettings.PropertyExportFlags |= UhtPropertyExportFlags.FieldNotify;
 		}
 
-		[UhtSpecifier(Extends = UhtTableNames.PropertyMember, ValueType = UhtSpecifierValueType.String)]
-		private static void VerseNameSpecifier(UhtSpecifierContext specifierContext, StringView value)
-		{
-			UhtPropertySpecifierContext context = (UhtPropertySpecifierContext)specifierContext;
-			context.PropertySettings.VerseName = value.ToString();
-		}
+		//[UhtSpecifier(Extends = UhtTableNames.PropertyMember, ValueType = UhtSpecifierValueType.String)]
+		//private static void VerseNameSpecifier(UhtSpecifierContext specifierContext, StringView value)
+		//{
+		//	UhtPropertySpecifierContext context = (UhtPropertySpecifierContext)specifierContext;
+		//	context.PropertySettings.VerseName = value.ToString();
+		//}
 	}
 }
diff --git a/Engine/Source/Programs/Shared/EpicGames.UHT/Types/UhtField.cs b/Engine/Source/Programs/Shared/EpicGames.UHT/Types/UhtField.cs
--- a/Engine/Source/Programs/Shared/EpicGames.UHT/Types/UhtField.cs
+++ b/Engine/Source/Programs/Shared/EpicGames.UHT/Types/UhtField.cs
@@ -155,11 +155,9 @@ namespace EpicGames.UHT.Types
 		/// </summary>
 		/// <param name="builder">Destination builder</param>
 		/// <param name="mode">Controls how the name is generated</param>
-		public void AppendVerseFullName(StringBuilder builder, UhtVerseFullNameMode mode)
+		public virtual void AppendVerseFullName(StringBuilder builder, UhtVerseFullNameMode mode)
 		{
-			builder.Append('(');
-			AppendVersePathInternal(builder, mode, true);
-			builder.Append(":)");
+			AppendVersePathPart(builder, mode);
 			AppendVerseName(builder, mode);
 		}
 
@@ -191,6 +189,18 @@ namespace EpicGames.UHT.Types
 			AppendVersePath(builder, mode, isTopLevel);
 		}
 
+		/// <summary>
+		/// Append the verse path part to the given builder
+		/// </summary>
+		/// <param name="builder">Destination builder</param>
+		/// <param name="mode">Controls how the name is generated</param>
+		protected virtual void AppendVersePathPart(StringBuilder builder, UhtVerseFullNameMode mode)
+		{
+			builder.Append('(');
+			AppendVersePathInternal(builder, mode, true);
+			builder.Append(":)");
+		}
+
 		/// <summary>
 		/// Helper method that appends only this path contribution to the verse path.
 		/// </summary>
diff --git a/Engine/Source/Programs/Shared/EpicGames.UHT/Types/UhtFunction.cs b/Engine/Source/Programs/Shared/EpicGames.UHT/Types/UhtFunction.cs
--- a/Engine/Source/Programs/Shared/EpicGames.UHT/Types/UhtFunction.cs
+++ b/Engine/Source/Programs/Shared/EpicGames.UHT/Types/UhtFunction.cs
@@ -126,6 +126,11 @@ namespace EpicGames.UHT.Types
 		/// True if the function is a decides function
 		/// </summary>
 		VerseDecides = 1 << 19,
+
+		/// <summary>
+		/// True if the return value is wrapped with TVerseTask
+		/// </summary>
+		VerseSuspends = 1 << 20,
 	}
 
 	/// <summary>
@@ -382,12 +387,6 @@ namespace EpicGames.UHT.Types
 			}
 		}
 
-		/// <summary>
-		/// True if this function has an incomplete return type
-		/// </summary>
-		public bool HasVerseIncompleteReturn => ReturnProperty is UhtEnumProperty enumProperty && enumProperty.Enum == Session.EVerseNativeCallResult;
-
-
 		/// <summary>
 		/// Construct a new instance of a function
 		/// </summary>
@@ -419,12 +418,22 @@ namespace EpicGames.UHT.Types
 			base.AddChild(type);
 		}
 
+		/// <inheritdoc/>
+		protected override void AppendVersePathPart(StringBuilder builder, UhtVerseFullNameMode mode)
+		{
+			// suspends methods don't get the path part added to the name
+			if (!FunctionExportFlags.HasAnyFlags(UhtFunctionExportFlags.VerseSuspends))
+			{
+				base.AppendVersePathPart(builder, mode);
+			}
+		}
+
 		/// <inheritdoc/>
 		protected override void AppendVerseName(StringBuilder builder, UhtVerseFullNameMode mode)
 		{
 			UhtProperty? returnProperty = ReturnProperty;
 			ReadOnlyMemory<UhtType> parameterProperties = ParameterProperties;
-			if (returnProperty is UhtEnumProperty enumProperty && enumProperty.Enum == Session.EVerseNativeCallResult)
+			if (Session.IsIncompleteReturn(returnProperty))
 			{
 				parameterProperties = parameterProperties[..^1];
 			}
diff --git a/Engine/Source/Programs/Shared/EpicGames.UHT/Types/UhtProperty.cs b/Engine/Source/Programs/Shared/EpicGames.UHT/Types/UhtProperty.cs
--- a/Engine/Source/Programs/Shared/EpicGames.UHT/Types/UhtProperty.cs
+++ b/Engine/Source/Programs/Shared/EpicGames.UHT/Types/UhtProperty.cs
@@ -97,6 +97,11 @@ namespace EpicGames.UHT.Types
 		/// If true, the type is a TPtr
 		/// </summary>
 		TPtr = 0x00008000,
+
+		/// <summary>
+		/// If true, the type is a TVerseTask
+		/// </summary>
+		TVerseTask = 0x00010000,
 	};
 
 	/// <summary>
@@ -1875,6 +1880,21 @@ namespace EpicGames.UHT.Types
 				this.LogError($"'{String.Join(", ", transients.ToStringList(false))}' specifier(s) are only allowed on class member variables");
 			}
 
+			if (PropertyExportFlags.HasAnyFlags(UhtPropertyExportFlags.TVal) && PropertyCategory != UhtPropertyCategory.Member)
+			{
+				this.LogError($"'TVal' is only valid on member variables");
+			}
+
+			if (PropertyExportFlags.HasAnyFlags(UhtPropertyExportFlags.TPtr) && PropertyCategory != UhtPropertyCategory.Member)
+			{
+				this.LogError($"'TPtr' is only valid on member variables");
+			}
+
+			if (PropertyExportFlags.HasAnyFlags(UhtPropertyExportFlags.TVerseTask) && (PropertyCategory != UhtPropertyCategory.Return || !outerStruct.IsVerseField))
+			{
+				this.LogError($"'TVerseTask' is only valid on the return value for Verse functions");
+			}
+
 			if (!options.HasAnyFlags(UhtValidationOptions.IsKeyOrValue))
 			{
 				switch (PropertyCategory)
diff --git a/Engine/Source/Programs/Shared/EpicGames.UHT/Utils/UhtSession.cs b/Engine/Source/Programs/Shared/EpicGames.UHT/Utils/UhtSession.cs
--- a/Engine/Source/Programs/Shared/EpicGames.UHT/Utils/UhtSession.cs
+++ b/Engine/Source/Programs/Shared/EpicGames.UHT/Utils/UhtSession.cs
@@ -2568,5 +2568,17 @@ namespace EpicGames.UHT.Utils
 			Logger.LogInformation("Total of {NumFiles} written", totalWrittenFiles);
 		}
 		#endregion
+
+		#region Other helper methods
+		/// <summary>
+		/// Test to see if the given property is an incomplete return type
+		/// </summary>
+		/// <param name="property"></param>
+		/// <returns></returns>
+		public bool IsIncompleteReturn(UhtProperty? property)
+		{
+			return property is UhtEnumProperty enumProperty && enumProperty.Enum == EVerseNativeCallResult;
+		}
+		#endregion
 	}
 }
