Improved Horde's LocalizationIssueHandler to get ready for many C++ localization warnings to change to structured logs (UE_LOGFMT)

diff --git a/Engine/Source/Programs/Shared/EpicGames.Horde/Issues/Handlers/LocalizationIssueHandler.cs b/Engine/Source/Programs/Shared/EpicGames.Horde/Issues/Handlers/LocalizationIssueHandler.cs
--- a/Engine/Source/Programs/Shared/EpicGames.Horde/Issues/Handlers/LocalizationIssueHandler.cs
+++ b/Engine/Source/Programs/Shared/EpicGames.Horde/Issues/Handlers/LocalizationIssueHandler.cs
@@ -14,6 +14,12 @@ namespace EpicGames.Horde.Issues.Handlers
 	[IssueHandler]
 	public class LocalizationIssueHandler : IssueHandler
 	{
+		private class LocalizationData
+		{
+			public string LocalizationTarget = "";
+			public string LocalizationID = "";
+		}
+
 		readonly List<IssueEventGroup> _issues = new List<IssueEventGroup>();
 
 		/// <inheritdoc/>
@@ -40,40 +46,82 @@ namespace EpicGames.Horde.Issues.Handlers
 		}
 
 		/// <summary>
-		/// Extracts a list of source files from an event
+		/// Extracts the localization properties of an IssueEvent to generate all the IssueKeys
+		/// </summary>
+		/// <param name="keys">Set of keys</param>
+		/// <param name="issueEvent">The event data</param>
+		/// <param name="data">The localization data</param>
+		private static void AddLocalizationIssueKeys(HashSet<IssueKey> keys, IssueEvent issueEvent, LocalizationData data)
+		{
+			foreach (JsonLogEvent line in issueEvent.Lines)
+			{
+				JsonDocument document = JsonDocument.Parse(line.Data);
+
+				JsonElement properties;
+				if (document.RootElement.TryGetProperty("properties", out properties) && properties.ValueKind == JsonValueKind.Object)
+				{
+					foreach (JsonProperty property in properties.EnumerateObject())
+					{
+						if (property.NameEquals("conflict") && property.Value.ValueKind == JsonValueKind.String)
+						{
+							keys.AddSourceFile(property.Value.GetString()!, IssueKeyType.File, data.LocalizationID);
+						}
+						else if((property.NameEquals("localizationID") || property.NameEquals("locID")) && property.Value.ValueKind == JsonValueKind.String)
+						{
+							keys.Add(new IssueKey(property.Value.GetString()!, IssueKeyType.None));
+						}
+					}
+				}
+			}
+		}
+
+		/// <summary>
+		/// Extracts the localization data from the properties of an IssueEvent
 		/// </summary>
 		/// <param name="issueEvent">The event data</param>
-		/// <param name="sourceFiles">List of source files</param>
-		public static void GetSourceFiles(IssueEvent issueEvent, HashSet<IssueKey> sourceFiles)
+		/// <param name="data">The localization data</param>
+		private static void GetLocalizationData(IssueEvent issueEvent, LocalizationData data)
 		{
 			foreach (JsonLogEvent line in issueEvent.Lines)
 			{
 				JsonDocument document = JsonDocument.Parse(line.Data);
 
-				string? relativePath;
-				if (document.RootElement.TryGetNestedProperty("properties.file.relativePath", out relativePath) || document.RootElement.TryGetNestedProperty("properties.file", out relativePath))
+				JsonElement properties;
+				if (document.RootElement.TryGetProperty("properties", out properties) && properties.ValueKind == JsonValueKind.Object)
 				{
-					if (!relativePath.EndsWith(".manifest", StringComparison.OrdinalIgnoreCase))
+					foreach (JsonProperty property in properties.EnumerateObject())
 					{
-						int endIdx = relativePath.LastIndexOfAny(new char[] { '/', '\\' }) + 1;
-						string fileName = relativePath.Substring(endIdx);
-						sourceFiles.Add(new IssueKey(fileName, IssueKeyType.File));
+						if (property.NameEquals("localizationTarget") && property.Value.ValueKind == JsonValueKind.String)
+						{
+							data.LocalizationTarget = property.Value.GetString()!;
+						}
+						else if ((property.NameEquals("localizationID") || property.NameEquals("locID")) && property.Value.ValueKind == JsonValueKind.String)
+						{
+							data.LocalizationID = property.Value.GetString()!;
+						}
 					}
 				}
 			}
 		}
 
 		/// <inheritdoc/>
-		public override bool HandleEvent(IssueEvent logEvent)
+		public override bool HandleEvent(IssueEvent issueEvent)
 		{
-			if (logEvent.EventId != null)
+			if (issueEvent.EventId != null)
 			{
-				EventId eventId = logEvent.EventId.Value;
+				EventId eventId = issueEvent.EventId.Value;
 				if (IsMatchingEventId(eventId))
 				{
-					IssueEventGroup issue = new IssueEventGroup("Localization", "Localization {Severity} in {Files}", IssueChangeFilter.Code);
-					issue.Events.Add(logEvent);
-					GetSourceFiles(logEvent, issue.Keys);
+					LocalizationData data = new LocalizationData();
+					GetLocalizationData(issueEvent, data);
+
+					string localizationIssueGroup = (string.IsNullOrEmpty(data.LocalizationTarget) ? "Localization" : "Localization Target " + data.LocalizationTarget);
+					IssueEventGroup issue = new IssueEventGroup(localizationIssueGroup, "Localization {Severity} in {Files}", IssueChangeFilter.All);
+					
+					issue.Events.Add(issueEvent);
+
+					issue.Keys.AddSourceFiles(issueEvent, data.LocalizationID);
+					AddLocalizationIssueKeys(issue.Keys, issueEvent, data);
 
 					if (issue.Keys.Count > 0)
 					{
diff --git a/Engine/Source/Programs/Shared/EpicGames.Horde/Issues/IssueKey.cs b/Engine/Source/Programs/Shared/EpicGames.Horde/Issues/IssueKey.cs
--- a/Engine/Source/Programs/Shared/EpicGames.Horde/Issues/IssueKey.cs
+++ b/Engine/Source/Programs/Shared/EpicGames.Horde/Issues/IssueKey.cs
@@ -204,7 +204,8 @@ namespace EpicGames.Horde.Issues
 		/// </summary>
 		/// <param name="keys">Set of keys</param>
 		/// <param name="issueEvent">The event data</param>
-		public static void AddSourceFiles(this HashSet<IssueKey> keys, IssueEvent issueEvent)
+		/// <param name="scope">The scope of the key to add</param>
+		public static void AddSourceFiles(this HashSet<IssueKey> keys, IssueEvent issueEvent, string? scope = null)
 		{
 			foreach (JsonLogEvent line in issueEvent.Lines)
 			{
@@ -223,11 +224,11 @@ namespace EpicGames.Horde.Issues
 					{
 						if (property.NameEquals("file") && property.Value.ValueKind == JsonValueKind.String)
 						{
-							keys.AddSourceFile(property.Value.GetString()!, type);
+							keys.AddSourceFile(property.Value.GetString()!, type, scope);
 						}
 						if (property.Value.HasStringProperty("$type", "SourceFile") && property.Value.TryGetStringProperty("relativePath", out string? value))
 						{
-							keys.AddSourceFile(value, type);
+							keys.AddSourceFile(value, type, scope);
 						}
 					}
 				}
@@ -240,12 +241,13 @@ namespace EpicGames.Horde.Issues
 		/// <param name="keys">List of source files</param>
 		/// <param name="relativePath">File to add</param>
 		/// <param name="type">Type of key to add</param>
-		public static void AddSourceFile(this HashSet<IssueKey> keys, string relativePath, IssueKeyType type)
+		/// <param name="scope">The scope of the key to add</param>
+		public static void AddSourceFile(this HashSet<IssueKey> keys, string relativePath, IssueKeyType type, string? scope = null)
 		{
 			int endIdx = relativePath.LastIndexOfAny(new char[] { '/', '\\' }) + 1;
 
 			string fileName = relativePath.Substring(endIdx);
-			IssueKey key = new IssueKey(fileName, type);
+			IssueKey key = new IssueKey(fileName, type, scope);
 
 			keys.Add(key);
 		}
