[UBA]
* Bumped timeout from 60 seconds to 10 minutes now when things are stable (except the short timeout causing trouble)

diff --git a/Engine/Source/Programs/Shared/EpicGames.UBA/Impl/ServerImpl.cs b/Engine/Source/Programs/Shared/EpicGames.UBA/Impl/ServerImpl.cs
--- a/Engine/Source/Programs/Shared/EpicGames.UBA/Impl/ServerImpl.cs
+++ b/Engine/Source/Programs/Shared/EpicGames.UBA/Impl/ServerImpl.cs
@@ -36,7 +36,8 @@ namespace EpicGames.UBA.Impl
 		public ServerImpl(int maxWorkers, int sendSize, ILogger logger, bool useQuic)
 		{
 			_logger = logger;
-			_handle = NetworkServer_Create(_logger.GetHandle(), maxWorkers, sendSize, 60, (byte)(useQuic?1:0)); // We know this code is using 5s pings.. so 60 seconds timeout on recv socket is good
+			int receiveTimeoutSeconds = 10 * 60; // 10 minutes. This should never happen
+			_handle = NetworkServer_Create(_logger.GetHandle(), maxWorkers, sendSize, receiveTimeoutSeconds, (byte)(useQuic?1:0));
 		}
 
 		#region IDisposable
