Horde: Update TelemetryMetricUnitType and default to Value instead of Time

diff --git a/Engine/Source/Programs/Shared/EpicGames.Horde/Telemetry/TelemetryMessages.cs b/Engine/Source/Programs/Shared/EpicGames.Horde/Telemetry/TelemetryMessages.cs
--- a/Engine/Source/Programs/Shared/EpicGames.Horde/Telemetry/TelemetryMessages.cs
+++ b/Engine/Source/Programs/Shared/EpicGames.Horde/Telemetry/TelemetryMessages.cs
@@ -113,14 +113,26 @@ namespace EpicGames.Horde.Telemetry
 	public enum TelemetryMetricUnitType
 	{
 		/// <summary>
-		/// Time duration 
+		/// Time duration in seconds 
 		/// </summary>
-		Time,
+		Seconds,
 		/// <summary>
-		/// Ratio 0-100%
+		/// Ratio 0-1
 		/// </summary>
 		Ratio,
 		/// <summary>
+		/// Ratio 0-100%
+		/// </summary>
+		Percentage,
+		/// <summary>
+		/// Data size in bytes
+		/// </summary>
+		Bytes,
+		/// <summary>
+		/// Data transfer rate in bytes per second
+		/// </summary>
+		BytesPerSecond,
+		/// <summary>
 		/// Artbitrary numeric value
 		/// </summary>
 		Value
@@ -177,7 +189,7 @@ namespace EpicGames.Horde.Telemetry
 		/// <summary>
 		/// The unit of the data to display
 		/// </summary>
-		public TelemetryMetricUnitType Unit { get; set; } = TelemetryMetricUnitType.Time;
+		public TelemetryMetricUnitType Unit { get; set; } = TelemetryMetricUnitType.Value;
 
 		/// <summary>
 		/// The graph type 
