Fixed warnings in LocalizationIssueHandler.cs

diff --git a/Engine/Source/Programs/Shared/EpicGames.Horde/Issues/Handlers/LocalizationIssueHandler.cs b/Engine/Source/Programs/Shared/EpicGames.Horde/Issues/Handlers/LocalizationIssueHandler.cs
--- a/Engine/Source/Programs/Shared/EpicGames.Horde/Issues/Handlers/LocalizationIssueHandler.cs
+++ b/Engine/Source/Programs/Shared/EpicGames.Horde/Issues/Handlers/LocalizationIssueHandler.cs
@@ -1,6 +1,5 @@
 // Copyright Epic Games, Inc. All Rights Reserved.
 
-using System;
 using System.Collections.Generic;
 using System.Text.Json;
 using EpicGames.Core;
@@ -16,8 +15,8 @@ namespace EpicGames.Horde.Issues.Handlers
 	{
 		private class LocalizationData
 		{
-			public string LocalizationTarget = "";
-			public string LocalizationID = "";
+			public string _localizationTarget = "";
+			public string _localizationID = "";
 		}
 
 		readonly List<IssueEventGroup> _issues = new List<IssueEventGroup>();
@@ -64,7 +63,7 @@ namespace EpicGames.Horde.Issues.Handlers
 					{
 						if (property.NameEquals("conflict") && property.Value.ValueKind == JsonValueKind.String)
 						{
-							keys.AddSourceFile(property.Value.GetString()!, IssueKeyType.File, data.LocalizationID);
+							keys.AddSourceFile(property.Value.GetString()!, IssueKeyType.File, data._localizationID);
 						}
 						else if((property.NameEquals("localizationID") || property.NameEquals("locID")) && property.Value.ValueKind == JsonValueKind.String)
 						{
@@ -93,11 +92,11 @@ namespace EpicGames.Horde.Issues.Handlers
 					{
 						if (property.NameEquals("localizationTarget") && property.Value.ValueKind == JsonValueKind.String)
 						{
-							data.LocalizationTarget = property.Value.GetString()!;
+							data._localizationTarget = property.Value.GetString()!;
 						}
 						else if ((property.NameEquals("localizationID") || property.NameEquals("locID")) && property.Value.ValueKind == JsonValueKind.String)
 						{
-							data.LocalizationID = property.Value.GetString()!;
+							data._localizationID = property.Value.GetString()!;
 						}
 					}
 				}
@@ -115,12 +114,12 @@ namespace EpicGames.Horde.Issues.Handlers
 					LocalizationData data = new LocalizationData();
 					GetLocalizationData(issueEvent, data);
 
-					string localizationIssueGroup = (string.IsNullOrEmpty(data.LocalizationTarget) ? "Localization" : "Localization Target " + data.LocalizationTarget);
+					string localizationIssueGroup = (System.String.IsNullOrEmpty(data._localizationTarget) ? "Localization" : "Localization Target " + data._localizationTarget);
 					IssueEventGroup issue = new IssueEventGroup(localizationIssueGroup, "Localization {Severity} in {Files}", IssueChangeFilter.All);
 					
 					issue.Events.Add(issueEvent);
 
-					issue.Keys.AddSourceFiles(issueEvent, data.LocalizationID);
+					issue.Keys.AddSourceFiles(issueEvent, data._localizationID);
 					AddLocalizationIssueKeys(issue.Keys, issueEvent, data);
 
 					if (issue.Keys.Count > 0)
