Do not run matchers on partial json log line entries.
This prevents us from trying to match unreliable log lines which were only partially written out e.g. due to the process crashing mid output.

diff --git a/Engine/Source/Programs/Shared/EpicGames.Core/LogEventParser.cs b/Engine/Source/Programs/Shared/EpicGames.Core/LogEventParser.cs
--- a/Engine/Source/Programs/Shared/EpicGames.Core/LogEventParser.cs
+++ b/Engine/Source/Programs/Shared/EpicGames.Core/LogEventParser.cs
@@ -303,6 +303,8 @@ namespace EpicGames.Core
 		/// <param name="line">The line to output</param>
 		public void WriteLine(string line)
 		{
+			bool invalidJson = false;
+
 			if (line.Length > 0 && line[0] == '{')
 			{
 				int length = line.Length;
@@ -321,6 +323,8 @@ namespace EpicGames.Core
 						_logger.Log(jsonEvent.Level, jsonEvent.EventId, jsonEvent, null, JsonLogEvent.Format);
 						return;
 					}
+
+					invalidJson = true;
 				}
 				catch (Exception ex)
 				{
@@ -329,7 +333,8 @@ namespace EpicGames.Core
 			}
 
 			_buffer.AddLine(StringUtils.ParseEscapeCodes(line));
-			ProcessData(false);
+
+			ProcessData(false, disableMatching: invalidJson);
 		}
 
 		/// <summary>
@@ -425,7 +430,7 @@ namespace EpicGames.Core
 		/// Process any data in the buffer
 		/// </summary>
 		/// <param name="bFlush">Whether we've reached the end of the stream</param>
-		void ProcessData(bool bFlush)
+		void ProcessData(bool bFlush, bool disableMatching = false)
 		{
 			_activeTimer.Start();
 			int startLineCount = _buffer.Length;
@@ -441,7 +446,7 @@ namespace EpicGames.Core
 				{
 					_matchingEnabled++;
 				}
-				else if (_matchingEnabled >= 0)
+				else if (_matchingEnabled >= 0 && disableMatching == false)
 				{
 					events = MatchEvent();
 				}
