Horde: Stop disposing of supplied HTTP handler in HordeClient

When enabling interactive auth for the Horde agent, this error came up as the handler was improperly disposed.
This fix addreses that as lifetime should be managed by caller, in this case, the service provider/DI.

diff --git a/Engine/Source/Programs/Shared/EpicGames.Horde/HordeClient.cs b/Engine/Source/Programs/Shared/EpicGames.Horde/HordeClient.cs
--- a/Engine/Source/Programs/Shared/EpicGames.Horde/HordeClient.cs
+++ b/Engine/Source/Programs/Shared/EpicGames.Horde/HordeClient.cs
@@ -1,6 +1,7 @@
 ﻿// Copyright Epic Games, Inc. All Rights Reserved.
 
 using System;
+using System.Diagnostics.CodeAnalysis;
 using System.Net.Http;
 using System.Net.Http.Headers;
 using System.Net.Http.Json;
@@ -337,6 +338,8 @@ namespace EpicGames.Horde
 	{
 		readonly HordeHttpAuthHandlerState _authHandlerState;
 		readonly IHordeHttpMessageHandler _baseHttpMessageHandler;
+
+		[SuppressMessage("Usage", "CA2213:Disposable fields should be disposed", Justification = "Lifetime is managed by caller")]
 		readonly HttpMessageHandler _authHttpMessageHandler;
 
 		/// <summary>
@@ -359,7 +362,6 @@ namespace EpicGames.Horde
 			_authHandlerState.OnStateChanged -= NotifyAuthStateChanged;
 
 			await _authHandlerState.DisposeAsync();
-			_authHttpMessageHandler.Dispose();
 			await base.DisposeAsync();
 		}
 
