Allow loading config System{TYPE}.ini from FPlatformProcess::ApplicationSettingsDir()

diff --git a/Engine/Source/Programs/Shared/EpicGames.Build/Unreal.cs b/Engine/Source/Programs/Shared/EpicGames.Build/Unreal.cs
--- a/Engine/Source/Programs/Shared/EpicGames.Build/Unreal.cs
+++ b/Engine/Source/Programs/Shared/EpicGames.Build/Unreal.cs
@@ -166,6 +166,12 @@ namespace UnrealBuildBase
 		public static DirectoryReference EngineSourceDirectory => _EngineSourceDirectory.Value;
 		private static readonly Lazy<DirectoryReference> _EngineSourceDirectory = new(() => DirectoryReference.Combine(EngineDirectory, "Source"));
 
+		/// <summary>
+		/// Returns the Application Settings Directory path. This matches FPlatformProcess::ApplicationSettingsDir().
+		/// </summary>
+		public static DirectoryReference ApplicationSettingDirectory => _ApplicationSettingDirectory.Value;
+		private static readonly Lazy<DirectoryReference> _ApplicationSettingDirectory = new(GetApplicationSettingDirectory);
+
 		/// <summary>
 		/// Returns the User Settings Directory path. This matches FPlatformProcess::UserSettingsDir().
 		/// </summary>
@@ -444,9 +450,9 @@ namespace UnrealBuildBase
 		}
 
 		/// <summary>
-		/// Returns the User Settings Directory path. This matches FPlatformProcess::UserSettingsDir().
+		/// Returns the Application Settings Directory path. This matches FPlatformProcess::ApplicationSettingsDir().
 		/// </summary>
-		private static DirectoryReference GetUserSettingDirectory()
+		private static DirectoryReference GetApplicationSettingDirectory()
 		{
 			if (OperatingSystem.IsMacOS())
 			{
@@ -456,6 +462,20 @@ namespace UnrealBuildBase
 			{
 				return new DirectoryReference(Path.Combine(Environment.GetFolderPath(Environment.SpecialFolder.ApplicationData), "Epic"));
 			}
+
+			return new DirectoryReference(Path.Combine(Environment.GetFolderPath(Environment.SpecialFolder.CommonApplicationData), "Epic"));
+		}
+
+		/// <summary>
+		/// Returns the User Settings Directory path. This matches FPlatformProcess::UserSettingsDir().
+		/// </summary>
+		private static DirectoryReference GetUserSettingDirectory()
+		{
+			if (OperatingSystem.IsMacOS() || OperatingSystem.IsLinux())
+			{
+				// Mac and Linux use the same folder for UserSettingsDir and ApplicationSettingsDir
+				return GetApplicationSettingDirectory();
+			}
 			else
 			{
 				// Not all user accounts have a local application data directory (eg. SYSTEM, used by Jenkins for builds).
diff --git a/Engine/Source/Programs/UnrealBuildTool/System/ConfigHierarchy.cs b/Engine/Source/Programs/UnrealBuildTool/System/ConfigHierarchy.cs
--- a/Engine/Source/Programs/UnrealBuildTool/System/ConfigHierarchy.cs
+++ b/Engine/Source/Programs/UnrealBuildTool/System/ConfigHierarchy.cs
@@ -1013,6 +1013,8 @@ namespace UnrealBuildTool
 			"{PROJECT}/Config/{PLATFORM}/Generated{PLATFORM}{TYPE}.ini",
 			// Project/Platform/Custom/CustomConfig/Platform*.ini only if CustomConfig is defined
 			"{PROJECT}/Config/{PLATFORM}/Custom/{CUSTOMCONFIG}/{PLATFORM}{TYPE}.ini",
+			// AppSettings/.../System*.ini
+			"{APPSETTINGS}/Unreal Engine/Engine/Config/System{TYPE}.ini",
 			// UserSettings/.../User*.ini
 			"{USERSETTINGS}/Unreal Engine/Engine/Config/User{TYPE}.ini",
 			// UserDir/.../User*.ini
@@ -1041,11 +1043,8 @@ namespace UnrealBuildTool
 		{
 			string OutString = InString.Replace("{TYPE}", BaseIniName);
 
-			DirectoryReference? UserSettingsDir = Unreal.UserSettingDirectory;
-			if (UserSettingsDir != null)
-			{
-				OutString = OutString.Replace("{USERSETTINGS}", UserSettingsDir.FullName);
-			}
+			OutString = OutString.Replace("{APPSETTINGS}", Unreal.ApplicationSettingDirectory.FullName);
+			OutString = OutString.Replace("{USERSETTINGS}", Unreal.UserSettingDirectory.FullName);
 
 			DirectoryReference? UserDir = Unreal.UserDirectory;
 			if (UserDir != null)
diff --git a/Engine/Source/Runtime/Core/Private/Misc/ConfigContext.cpp b/Engine/Source/Runtime/Core/Private/Misc/ConfigContext.cpp
--- a/Engine/Source/Runtime/Core/Private/Misc/ConfigContext.cpp
+++ b/Engine/Source/Runtime/Core/Private/Misc/ConfigContext.cpp
@@ -658,6 +658,7 @@ static void ConditionalOverrideIniFilename(FString& IniFilename, const TCHAR* Ba
 static FString PerformBasicReplacements(const FString& InString, const TCHAR* BaseIniName)
 {
 	FString OutString = InString.Replace(TEXT("{TYPE}"), BaseIniName, ESearchCase::CaseSensitive);
+	OutString = OutString.Replace(TEXT("{APPSETTINGS}"), FPlatformProcess::ApplicationSettingsDir(), ESearchCase::CaseSensitive);
 	OutString = OutString.Replace(TEXT("{USERSETTINGS}"), FPlatformProcess::UserSettingsDir(), ESearchCase::CaseSensitive);
 	OutString = OutString.Replace(TEXT("{USER}"), FPlatformProcess::UserDir(), ESearchCase::CaseSensitive);
 	OutString = OutString.Replace(TEXT("{CUSTOMCONFIG}"), *FConfigCacheIni::GetCustomConfigString(), ESearchCase::CaseSensitive);
@@ -862,7 +863,7 @@ void FConfigContext::AddStaticLayersToHierarchy(TArray<FString>* GatheredLayerFi
 		if (!EnumHasAnyFlags(Layer.Flag, EConfigLayerFlags::NoExpand))
 		{
 			// we assume none of the more special tags in expanded ones
-			checkfSlow(FCString::Strstr(Layer.Path, TEXT("{USERSETTINGS}")) == nullptr && FCString::Strstr(Layer.Path, TEXT("{USER}")) == nullptr, TEXT("Expanded config %s shouldn't have a {USER*} tags in it"), Layer.Path);
+			checkfSlow(FCString::Strstr(Layer.Path, TEXT("{APPSETTINGS}")) == nullptr && FCString::Strstr(Layer.Path, TEXT("{USERSETTINGS}")) == nullptr && FCString::Strstr(Layer.Path, TEXT("{USER}")) == nullptr, TEXT("Expanded config %s shouldn't have {APPSETTINGS} or {USER*} tags in it"), Layer.Path);
 
 			// loop over all the possible expansions
 			for (int32 ExpansionIndex = 0; ExpansionIndex < UE_ARRAY_COUNT(GConfigExpansions); ExpansionIndex++)
diff --git a/Engine/Source/Runtime/Core/Public/Misc/ConfigHierarchy.h b/Engine/Source/Runtime/Core/Public/Misc/ConfigHierarchy.h
--- a/Engine/Source/Runtime/Core/Public/Misc/ConfigHierarchy.h
+++ b/Engine/Source/Runtime/Core/Public/Misc/ConfigHierarchy.h
@@ -35,6 +35,8 @@ inline FConfigLayer GConfigLayers[] =
 	{ TEXT("ProjectPlatformGenerated"),	TEXT("{PROJECT}/Config/{PLATFORM}/Generated{PLATFORM}{TYPE}.ini"), EConfigLayerFlags::UseGlobalConfigCache },
 	// Project/Platform/Custom/CustomConfig/Platform*.ini only if CustomConfig is defined
 	{ TEXT("CustomConfigPlatform"),		TEXT("{PROJECT}/Config/{PLATFORM}/Custom/{CUSTOMCONFIG}/{PLATFORM}{TYPE}.ini"), EConfigLayerFlags::RequiresCustomConfig| EConfigLayerFlags::UseGlobalConfigCache },
+	// AppSettings/.../System*.ini
+	{ TEXT("AppSettingsDir"),			TEXT("{APPSETTINGS}Unreal Engine/Engine/Config/System{TYPE}.ini"), EConfigLayerFlags::NoExpand },
 	// UserSettings/.../User*.ini
 	{ TEXT("UserSettingsDir"),			TEXT("{USERSETTINGS}Unreal Engine/Engine/Config/User{TYPE}.ini"), EConfigLayerFlags::NoExpand },
 	// UserDir/.../User*.ini
