Horde: Log reason when compute task allocation fails in compute client

diff --git a/Engine/Source/Programs/Shared/EpicGames.Horde/Compute/Clients/ServerComputeClient.cs b/Engine/Source/Programs/Shared/EpicGames.Horde/Compute/Clients/ServerComputeClient.cs
--- a/Engine/Source/Programs/Shared/EpicGames.Horde/Compute/Clients/ServerComputeClient.cs
+++ b/Engine/Source/Programs/Shared/EpicGames.Horde/Compute/Clients/ServerComputeClient.cs
@@ -295,7 +295,7 @@ namespace EpicGames.Horde.Compute.Clients
 
 				if (httpResponse.StatusCode is HttpStatusCode.ServiceUnavailable or HttpStatusCode.TooManyRequests)
 				{
-					_logger.LogDebug("No compute resource is available");
+					_logger.LogDebug("No compute resource is available. Reason: {Reason}", await httpResponse.Content.ReadAsStringAsync(cancellationToken));
 					yield break;
 				}
 
