Horde: Add a new workstation mode to agent

* An agent in workstation mode is expected to authenticate as the current user, and not via a pre-shared token like in a buildfarm scenario
* Uses interactive auth via EpicGames.Oidc to authenticate

diff --git a/Engine/Source/Programs/Horde/HordeAgent/AgentApp.cs b/Engine/Source/Programs/Horde/HordeAgent/AgentApp.cs
--- a/Engine/Source/Programs/Horde/HordeAgent/AgentApp.cs
+++ b/Engine/Source/Programs/Horde/HordeAgent/AgentApp.cs
@@ -181,8 +181,8 @@ namespace HordeAgent
 			services.AddHorde(options =>
 			{
 				options.ServerUrl = serverProfile.Url;
-				options.AccessToken = serverProfile.Token;
-				options.AllowAuthPrompt = false;
+				options.AccessToken = serverProfile.GetAuthToken();
+				options.AllowAuthPrompt = serverProfile.UseInteractiveAuth;
 				options.BackendCache.CacheDir = DirectoryReference.Combine(settings.WorkingDir, "Saved", "Bundles").FullName;
 				options.BackendCache.MaxSize = settings.BundleCacheSize * 1024 * 1024;
 			});
diff --git a/Engine/Source/Programs/Horde/HordeAgent/AgentSettings.cs b/Engine/Source/Programs/Horde/HordeAgent/AgentSettings.cs
--- a/Engine/Source/Programs/Horde/HordeAgent/AgentSettings.cs
+++ b/Engine/Source/Programs/Horde/HordeAgent/AgentSettings.cs
@@ -4,11 +4,24 @@ using System.ComponentModel.DataAnnotations;
 using EpicGames.Core;
 using EpicGames.Horde;
 using EpicGames.Horde.Utilities;
+using HordeCommon.Rpc;
 using HordeServer.Utilities;
 using Microsoft.Extensions.Configuration;
 
 namespace HordeAgent
 {
+	/// <summary>
+	/// Defines the operation mode of the agent
+	/// Duplicated to prevent depending on Protobuf structures (weak name reference when deserializing JSON)
+	/// </summary>
+	public enum AgentMode
+	{
+		/// <see cref="Mode.Dedicated" />
+		Dedicated,
+		/// <see cref="Mode.Workstation" />
+		Workstation,
+	}
+	
 	/// <summary>
 	/// Describes a network share to mount
 	/// </summary>
@@ -51,6 +64,11 @@ namespace HordeAgent
 		/// Bearer token to use to initiate the connection
 		/// </summary>
 		public string? Token { get; set; }
+		
+		/// <summary>
+		/// Whether to authenticate interactively in a desktop environment (for example, when agent is running on a user's workstation)
+		/// </summary>
+		public bool UseInteractiveAuth { get; set; } = false;
 
 		/// <summary>
 		/// Thumbprint of a certificate to trust. Allows using self-signed certs for the server.
@@ -61,6 +79,14 @@ namespace HordeAgent
 		/// Thumbprints of certificates to trust. Allows using self-signed certs for the server.
 		/// </summary>
 		public List<string> Thumbprints { get; } = new List<string>();
+		
+		/// <summary>
+		/// Returns auth token if interactive auth is disabled. Otherwise, null is returned.
+		/// </summary>
+		public string? GetAuthToken()
+		{
+			return UseInteractiveAuth ? null : Token;
+		}
 
 		/// <summary>
 		/// Checks whether the given certificate thumbprint should be trusted
@@ -108,6 +134,16 @@ namespace HordeAgent
 		/// By default, the computer's hostname will be used.
 		/// </summary>
 		public string? Name { get; set; }
+		
+		/// <summary>
+		/// Mode of operation for the agent
+		/// <see cref="AgentMode.Dedicated"/> - For trusted agents in controlled environments (e.g., build farms). 
+		/// These agents handle all lease types and run exclusively Horde workloads.
+		/// 
+		/// <see cref="AgentMode.Workstation"/> - For low-trust workstations, uses interactive authentication (human logs in).
+		/// These agents yield to non-Horde workloads and only support compute leases for remote execution.
+		/// </summary>
+		public AgentMode? Mode { get; set; } = AgentMode.Dedicated;
 
 		/// <summary>
 		/// Whether the server is running in 'installed' mode. In this mode, on Windows, the default data directory will use the common 
@@ -307,6 +343,16 @@ namespace HordeAgent
 		{
 			return Name ?? Environment.MachineName;
 		}
+		
+		internal Mode GetMode()
+		{
+			return Mode switch
+			{
+				AgentMode.Dedicated => HordeCommon.Rpc.Mode.Dedicated,
+				AgentMode.Workstation => HordeCommon.Rpc.Mode.Workstation,
+				_ => throw new Exception($"Unknown agent mode: {Mode}")
+			};
+		}
 	}
 
 	/// <summary>
diff --git a/Engine/Source/Programs/Horde/HordeAgent/Services/GrpcService.cs b/Engine/Source/Programs/Horde/HordeAgent/Services/GrpcService.cs
--- a/Engine/Source/Programs/Horde/HordeAgent/Services/GrpcService.cs
+++ b/Engine/Source/Programs/Horde/HordeAgent/Services/GrpcService.cs
@@ -136,7 +136,7 @@ namespace HordeAgent.Services
 			}
 
 			// Create the channel
-			_logger.LogInformation("Connecting to rpc server {BaseUrl}", serverUri);
+			_logger.LogInformation("gRPC channel connecting to {BaseUrl} ...", serverUri);
 			return GrpcChannel.ForAddress(serverUri, channelOptions);
 		}
 
diff --git a/Engine/Source/Programs/Horde/HordeAgent/Services/SessionFactory.cs b/Engine/Source/Programs/Horde/HordeAgent/Services/SessionFactory.cs
--- a/Engine/Source/Programs/Horde/HordeAgent/Services/SessionFactory.cs
+++ b/Engine/Source/Programs/Horde/HordeAgent/Services/SessionFactory.cs
@@ -157,7 +157,7 @@ namespace HordeAgent.Services
 			{
 				statusService.SetDescription(AgentStatusMessage.WaitingForEnrollment);
 
-				registrationInfo = await RegisterAgentAsync(grpcService, currentSettings, capabilities, logger, cancellationToken);
+				registrationInfo = await RegisterAgentAsync(hordeClientFactory, currentSettings, capabilities, logger, cancellationToken);
 				registrationList.Entries.Add(registrationInfo);
 
 				await WriteRegistrationListAsync(registrationFile, registrationList, cancellationToken);
@@ -172,7 +172,7 @@ namespace HordeAgent.Services
 			statusService.SetDescription(AgentStatusMessage.ConnectingToServer);
 
 			RpcCreateSessionResponse createSessionResponse;
-			await using (HordeClient sessionClient = hordeClientFactory.Create(accessToken: registrationInfo.Token))
+			await using (HordeClient sessionClient = hordeClientFactory.Create(accessToken: serverProfile.GetAuthToken()))
 			{
 				HordeRpc.HordeRpcClient rpcClient = await sessionClient.CreateGrpcClientAsync<HordeRpc.HordeRpcClient>(cancellationToken);
 
@@ -255,20 +255,18 @@ namespace HordeAgent.Services
 			await FileReference.WriteAllBytesAsync(settingsFile, data, cancellationToken);
 		}
 
-		static async Task<AgentRegistration> RegisterAgentAsync(GrpcService grpcService, AgentSettings agentSettings, RpcAgentCapabilities capabilities, ILogger logger, CancellationToken cancellationToken)
+		static async Task<AgentRegistration> RegisterAgentAsync(HordeClientFactory hordeClientFactory, AgentSettings agentSettings, RpcAgentCapabilities capabilities, ILogger logger, CancellationToken cancellationToken)
 		{
 			ServerProfile serverProfile = agentSettings.GetCurrentServerProfile();
-			using GrpcChannel grpcChannel = await grpcService.CreateGrpcChannelAsync(serverProfile.Token, cancellationToken);
+			await using HordeClient sessionClient = hordeClientFactory.Create(accessToken: serverProfile.GetAuthToken());
+			using GrpcChannel grpcChannel = await sessionClient.CreateGrpcChannelAsync(cancellationToken);
 
-			if (!String.IsNullOrEmpty(serverProfile.Token))
+			if (!String.IsNullOrEmpty(serverProfile.Token) || serverProfile.UseInteractiveAuth)
 			{
-				logger.LogInformation("Registering agent directly...");
-				HordeRpc.HordeRpcClient rpcClient = new HordeRpc.HordeRpcClient(grpcChannel);
-
-				RpcCreateAgentRequest createAgentRequest = new RpcCreateAgentRequest();
-				createAgentRequest.Name = agentSettings.GetAgentName();
-				createAgentRequest.Ephemeral = agentSettings.Ephemeral;
+				logger.LogInformation("Registering agent...");
+				HordeRpc.HordeRpcClient rpcClient = new (grpcChannel);
 
+				RpcCreateAgentRequest createAgentRequest = new() { Name = agentSettings.GetAgentName(), Ephemeral = agentSettings.Ephemeral, Mode = agentSettings.GetMode() };
 				RpcCreateAgentResponse createAgentResponse = await rpcClient.CreateAgentAsync(createAgentRequest, null, null, cancellationToken);
 				return new AgentRegistration(serverProfile.Url, createAgentResponse.Id, createAgentResponse.Token);
 			}
diff --git a/Engine/Source/Programs/Shared/EpicGames.Horde/HordeClient.cs b/Engine/Source/Programs/Shared/EpicGames.Horde/HordeClient.cs
--- a/Engine/Source/Programs/Shared/EpicGames.Horde/HordeClient.cs
+++ b/Engine/Source/Programs/Shared/EpicGames.Horde/HordeClient.cs
@@ -129,25 +129,21 @@ namespace EpicGames.Horde
 			Uri serverUri = ServerUrl;
 			bool useInsecureConnection = serverUri.Scheme.Equals("http", StringComparison.Ordinal);
 
-			// Get the server URL for gRPC traffic. If we're using an unencrpyted connection we need to use a different port for http/2, so 
-			// send a http1 request to the server to query it.
+			// Get the server URL for gRPC traffic. If we're using an unencrypted connection we need to use a different port for HTTP/2,
+			// so send a HTTP/1 request to the server to query it.
 			if (useInsecureConnection)
 			{
 				_logger.LogInformation("Querying server {BaseUrl} for rpc port", serverUri);
-				using (HttpClient httpClient = CreateUnauthenticatedHttpClient())
+				using HttpClient httpClient = CreateUnauthenticatedHttpClient();
+				httpClient.DefaultRequestHeaders.Add("Accept", "application/json");
+				httpClient.Timeout = TimeSpan.FromSeconds(210); // Need to make sure this doesn't cancel any long running gRPC streaming calls (eg. session update)
+				using HttpResponseMessage response = await httpClient.GetAsync(new Uri(serverUri, "api/v1/server/ports"), cancellationToken);
+				GetPortsResponse? ports = await response.Content.ReadFromJsonAsync<GetPortsResponse>(HordeHttpClient.JsonSerializerOptions, cancellationToken);
+				if (ports is { UnencryptedHttp2: not null } && ports.UnencryptedHttp2.Value != 0)
 				{
-					httpClient.DefaultRequestHeaders.Add("Accept", "application/json");
-					httpClient.Timeout = TimeSpan.FromSeconds(210); // Need to make sure this doesn't cancel any long running gRPC streaming calls (eg. session update)
-					using (HttpResponseMessage response = await httpClient.GetAsync(new Uri(serverUri, "api/v1/server/ports"), cancellationToken))
-					{
-						GetPortsResponse? ports = await response.Content.ReadFromJsonAsync<GetPortsResponse>(HordeHttpClient.JsonSerializerOptions, cancellationToken);
-						if (ports != null && ports.UnencryptedHttp2.HasValue && ports.UnencryptedHttp2.Value != 0)
-						{
-							UriBuilder builder = new UriBuilder(serverUri);
-							builder.Port = ports.UnencryptedHttp2.Value;
-							serverUri = builder.Uri;
-						}
-					}
+					UriBuilder builder = new UriBuilder(serverUri);
+					builder.Port = ports.UnencryptedHttp2.Value;
+					serverUri = builder.Uri;
 				}
 			}
 
@@ -177,7 +173,7 @@ namespace EpicGames.Horde
 			});
 
 			// Configure requests to send the bearer token
-			string? bearerToken = await GetAccessTokenAsync(false, cancellationToken);
+			string? bearerToken = await GetAccessTokenAsync(_hordeOptions.AllowAuthPrompt, cancellationToken);
 			if (!String.IsNullOrEmpty(bearerToken))
 			{
 				CallCredentials callCredentials = CallCredentials.FromInterceptor((context, metadata) =>
@@ -186,20 +182,12 @@ namespace EpicGames.Horde
 					return Task.CompletedTask;
 				});
 
-				if (useInsecureConnection)
-				{
-					channelOptions.Credentials = ChannelCredentials.Create(ChannelCredentials.Insecure, callCredentials);
-				}
-				else
-				{
-					channelOptions.Credentials = ChannelCredentials.Create(ChannelCredentials.SecureSsl, callCredentials);
-				}
-
+				channelOptions.Credentials = ChannelCredentials.Create(useInsecureConnection ? ChannelCredentials.Insecure : ChannelCredentials.SecureSsl, callCredentials);
 				channelOptions.UnsafeUseInsecureChannelCallCredentials = useInsecureConnection;
 			}
 
 			// Create the channel
-			_logger.LogInformation("Connecting to rpc server {BaseUrl}", serverUri);
+			_logger.LogInformation("gRPC channel connecting to {BaseUrl} ...", serverUri);
 			return GrpcChannel.ForAddress(serverUri, channelOptions);
 		}
 
