Fix GetApplicationSettingDirectory() on Mac not matching FPlatformProcess::ApplicationSettingsDir()

#rnx

diff --git a/Engine/Source/Programs/Shared/EpicGames.Build/Unreal.cs b/Engine/Source/Programs/Shared/EpicGames.Build/Unreal.cs
--- a/Engine/Source/Programs/Shared/EpicGames.Build/Unreal.cs
+++ b/Engine/Source/Programs/Shared/EpicGames.Build/Unreal.cs
@@ -454,11 +454,7 @@ namespace UnrealBuildBase
 		/// </summary>
 		private static DirectoryReference GetApplicationSettingDirectory()
 		{
-			if (OperatingSystem.IsMacOS())
-			{
-				return new DirectoryReference(Path.Combine(Environment.GetFolderPath(Environment.SpecialFolder.Personal), "Library", "Application Support", "Epic"));
-			}
-			else if (OperatingSystem.IsLinux())
+			if (OperatingSystem.IsMacOS() || OperatingSystem.IsLinux())
 			{
 				return new DirectoryReference(Path.Combine(Environment.GetFolderPath(Environment.SpecialFolder.ApplicationData), "Epic"));
 			}
