AutomationTool: Maintain encoding in FilterIniFile

#rnx

diff --git a/Engine/Source/Programs/AutomationTool/AutomationUtils/Utils.cs b/Engine/Source/Programs/AutomationTool/AutomationUtils/Utils.cs
--- a/Engine/Source/Programs/AutomationTool/AutomationUtils/Utils.cs
+++ b/Engine/Source/Programs/AutomationTool/AutomationUtils/Utils.cs
@@ -377,7 +377,8 @@ namespace AutomationTool
 
 		private static void FilterIniFile(string SourceName, string TargetName, List<string> IniKeyDenyList, List<string> InSectionDenyList, List<string> InSectionAllowList)
 		{
-			string[] Lines = File.ReadAllLines(SourceName);
+			Encoding encoding = FileUtils.GetEncoding(SourceName);
+			string[] Lines = File.ReadAllLines(SourceName, encoding);
 			StringBuilder NewLines = new StringBuilder("");
 			bool bFilteringSection = false;
 
@@ -433,7 +434,7 @@ namespace AutomationTool
 			{
 				File.Delete(TargetName);
 			}
-			File.WriteAllText(TargetName, NewLines.ToString());
+			File.WriteAllText(TargetName, NewLines.ToString(), encoding);
 
 			// other code assumes same timestamp for source and dest
 			File.SetLastWriteTimeUtc(TargetName, File.GetLastWriteTimeUtc(SourceName));
diff --git a/Engine/Source/Programs/Shared/EpicGames.Core/FileReference.cs b/Engine/Source/Programs/Shared/EpicGames.Core/FileReference.cs
--- a/Engine/Source/Programs/Shared/EpicGames.Core/FileReference.cs
+++ b/Engine/Source/Programs/Shared/EpicGames.Core/FileReference.cs
@@ -247,6 +247,60 @@ namespace EpicGames.Core
 			}
 		}
 
+		/// <summary>
+		/// Get the encoding of a file by checking for BOM bytes
+		/// </summary>
+		/// <param name="location">Location of the file</param>
+		/// <param name="skipBytes">Number of bytes to skip</param>
+		/// <returns>The encoding</returns>
+		public static Encoding GetEncoding(FileReference location, out int skipBytes)
+		{
+			try
+			{
+				Span<byte> span = stackalloc byte[4];
+				using (FileStream stream = Open(location, FileMode.Open, FileAccess.Read, FileShare.Read))
+				{
+					stream.Read(span);
+				}
+				return FileUtils.GetEncoding(span, out skipBytes);
+			}
+			catch (IOException)
+			{
+			}
+			skipBytes = 0;
+			return Encoding.UTF8;
+		}
+
+		/// <summary>
+		/// Get the encoding of a file by checking for BOM bytes
+		/// </summary>
+		/// <param name="location">Location of the file</param>
+		/// <returns>The encoding</returns>
+		public static Encoding GetEncoding(FileReference location) => GetEncoding(location, out int _);
+
+        /// <summary>
+        /// Get the encoding of a file by checking for BOM bytes
+        /// </summary>
+        /// <param name="location">Location of the file</param>
+        /// <param name="cancellationToken">Cancellation token for the operation</param>
+        /// <returns>The encoding</returns>
+        public static async Task<Encoding> GetEncodingAsync(FileReference location, CancellationToken cancellationToken = default)
+		{
+			try
+			{
+				byte[] buffer = new byte[4];
+				using (FileStream stream = Open(location, FileMode.Open, FileAccess.Read, FileShare.Read))
+				{
+					await stream.ReadAsync(buffer, 0, 4, cancellationToken);
+				}
+				return FileUtils.GetEncoding(buffer);
+			}
+			catch (IOException)
+			{
+			}
+			return Encoding.UTF8;
+		}
+
 		/// <summary>
 		/// Finds the correct case to match the location of this file on disk. Uses the given case for parts of the path that do not exist.
 		/// </summary>
@@ -774,7 +828,7 @@ namespace EpicGames.Core
 		/// <param name="contents">Contents to append to the file</param>
 		/// <param name="encoding">The encoding to use when parsing the file</param>
 		/// <param name="cancellationToken">Cancellation token for the operation</param>
-		public static Task AppendAllTextAsync(FileReference location, string contents, Encoding encoding, CancellationToken cancellationToken = default) => File.AppendAllTextAsync(location.FullName, contents, encoding, cancellationToken	);
+		public static Task AppendAllTextAsync(FileReference location, string contents, Encoding encoding, CancellationToken cancellationToken = default) => File.AppendAllTextAsync(location.FullName, contents, encoding, cancellationToken);
 
 		#endregion
 	}
diff --git a/Engine/Source/Programs/Shared/EpicGames.Core/FileUtils.cs b/Engine/Source/Programs/Shared/EpicGames.Core/FileUtils.cs
--- a/Engine/Source/Programs/Shared/EpicGames.Core/FileUtils.cs
+++ b/Engine/Source/Programs/Shared/EpicGames.Core/FileUtils.cs
@@ -17,72 +17,45 @@ namespace EpicGames.Core
 	/// <summary>
 	/// Exception used to represent caught file/directory exceptions.
 	/// </summary>
-	public class WrappedFileOrDirectoryException : Exception
+	/// <param name="inner">Inner exception</param>
+	/// <param name="message">Message to display</param>
+	public class WrappedFileOrDirectoryException(Exception inner, string message) : Exception(message, inner)
 	{
-		/// <summary>
-		/// Constructor
-		/// </summary>
-		/// <param name="inner">Inner exception</param>
-		/// <param name="message">Message to display</param>
-		public WrappedFileOrDirectoryException(Exception inner, string message) : base(message, inner)
-		{
-		}
-
-		/// <summary>
-		/// Returns the message to display for this exception
-		/// </summary>
-		/// <returns>Message to display</returns>
-		public override string ToString()
-		{
-			return Message;
-		}
+		/// <inheritdoc/>
+		public override string ToString() => Message;
 	}
 
 	/// <summary>
 	/// Information about a locked file
 	/// </summary>
-	public class FileLockInfoWin32
+	/// <param name="processId">Process id</param>
+	/// <param name="fileName">Path to the process holding the lock</param>
+	/// <param name="appName">Name of the application</param>
+	/// <param name="startTime">Time at which the process started</param>
+	public class FileLockInfoWin32(int processId, string? fileName, string appName, DateTime startTime)
 	{
 		/// <summary>
 		/// Process id
 		/// </summary>
-		public int ProcessId { get; }
+		public int ProcessId => processId;
 
 		/// <summary>
 		/// Path to the process holding the lock
 		/// </summary>
-		public string? FileName { get; }
+		public string? FileName => fileName;
 
 		/// <summary>
 		/// Name of the application
 		/// </summary>
-		public string AppName { get; }
+		public string AppName => appName;
 
 		/// <summary>
 		/// Time at which the process started
 		/// </summary>
-		public DateTime StartTime { get; }
-
-		/// <summary>
-		/// Constructor
-		/// </summary>
-		/// <param name="processId">Process id</param>
-		/// <param name="fileName">Path to the process holding the lock</param>
-		/// <param name="appName">Name of the application</param>
-		/// <param name="startTime">Time at which the process started</param>
-		public FileLockInfoWin32(int processId, string? fileName, string appName, DateTime startTime)
-		{
-			ProcessId = processId;
-			FileName = fileName;
-			AppName = appName;
-			StartTime = startTime;
-		}
+		public DateTime StartTime => startTime;
 
 		/// <inheritdoc/>
-		public override string ToString()
-		{
-			return $"{ProcessId}: {FileName ?? AppName} (started {StartTime})";
-		}
+		public override string ToString() => $"{ProcessId}: {FileName ?? AppName} (started {StartTime})";
 	}
 
 	/// <summary>
@@ -436,6 +409,83 @@ namespace EpicGames.Core
 			}
 		}
 
+		/// <summary>
+		/// Get the encoding of a span and number of bytes to skip at the start of the buffer
+		/// </summary>
+		/// <param name="bytes">Bytes to scan for a BOM</param>
+		/// <param name="skipBytes">Number of bytes to skip</param>
+		/// <returns>The encoding</returns>
+		public static Encoding GetEncoding(ReadOnlySpan<byte> bytes, out int skipBytes)
+		{
+			// https://simple.wikipedia.org/wiki/Byte_order_mark
+			switch (bytes.Length)
+			{
+				// UTF-32, little-endian
+				case >= 4 when bytes[0] == 0xff && bytes[1] == 0xfe && bytes[2] == 0x00 && bytes[3] == 0x00:
+					skipBytes = 4;
+					return Encoding.UTF32;
+				// UTF-32, big-endian
+				case >= 4 when bytes[0] == 0x00 && bytes[1] == 0x00 && bytes[2] == 0xfe && bytes[3] == 0xff:
+					skipBytes = 4;
+					// Encoding.BigEndianUTF32 is private, use GetEncoding() to return the static Encoding to prevent additional allocation
+					return Encoding.GetEncoding(12001);
+				// UTF-8 with BOM
+				case >= 3 when bytes[0] == 0xef && bytes[1] == 0xbb && bytes[2] == 0xbf:
+					skipBytes = 3;
+					return Encoding.UTF8;
+				// UTF-16, little-endian
+				case >= 2 when bytes[0] == 0xff && bytes[1] == 0xfe:
+					skipBytes = 2;
+					return Encoding.Unicode;
+				// UTF-16, big-endian
+				case >= 2 when bytes[0] == 0xfe && bytes[1] == 0xff:
+					skipBytes = 2;
+					return Encoding.BigEndianUnicode;
+				// Default to UTF-8
+				default:
+					skipBytes = 0;
+					return Encoding.UTF8;
+			}
+		}
+
+		/// <summary>
+		/// Get the encoding of a span
+		/// </summary>
+		/// <param name="bytes">Bytes to scan for a BOM</param>
+		/// <returns>The encoding</returns>
+		public static Encoding GetEncoding(ReadOnlySpan<byte> bytes) => GetEncoding(bytes, out int _);
+
+		/// <summary>
+		/// Get the encoding of a file
+		/// </summary>
+		/// <param name="filePath">File path to scan for a BOM</param>
+		/// <param name="skipBytes">Number of bytes to skip</param>
+		/// <returns>The encoding</returns>
+		public static Encoding GetEncoding(string filePath, out int skipBytes)
+		{
+			try
+			{
+				byte[] bytes = new byte[4];
+				using (FileStream stream = File.Open(filePath, FileMode.Open, FileAccess.Read, FileShare.Read))
+				{
+					stream.Read(bytes, 0, 4);
+				}
+				return GetEncoding(bytes, out skipBytes);
+			}
+			catch (IOException)
+			{
+			}
+			skipBytes = 0;
+			return Encoding.UTF8;
+		}
+
+		/// <summary>
+		/// Get the encoding of a span
+		/// </summary>
+		/// <param name="filePath">File path to scan for a BOM</param>
+		/// <returns>The encoding</returns>
+		public static Encoding GetEncoding(string filePath) => GetEncoding(filePath, out int _);
+
 		/// <summary>
 		/// Gets the file mode on Mac
 		/// </summary>
diff --git a/Engine/Source/Programs/Shared/EpicGames.UHT/Utils/UhtFileManager.cs b/Engine/Source/Programs/Shared/EpicGames.UHT/Utils/UhtFileManager.cs
--- a/Engine/Source/Programs/Shared/EpicGames.UHT/Utils/UhtFileManager.cs
+++ b/Engine/Source/Programs/Shared/EpicGames.UHT/Utils/UhtFileManager.cs
@@ -109,7 +109,7 @@ namespace EpicGames.UHT.Utils
 				int readLength = fs.Read(byteSpan);
 				byteSpan = byteSpan[..readLength];
 
-				Encoding encoding = GetEncoding(byteSpan, out int skipBytes);
+				Encoding encoding = FileUtils.GetEncoding(byteSpan, out int skipBytes);
 				byteSpan = byteSpan[skipBytes..];
 
 				int charCount = encoding.GetCharCount(byteBuffer.Buffer.Memory.Span);
@@ -141,7 +141,7 @@ namespace EpicGames.UHT.Utils
 				int readLength = fs.Read(byteSpan);
 				byteSpan = byteSpan[..readLength];
 
-				Encoding encoding = GetEncoding(byteSpan, out int skipBytes);
+				Encoding encoding = FileUtils.GetEncoding(byteSpan, out int skipBytes);
 				byteSpan = byteSpan[skipBytes..];
 
 				int charCount = encoding.GetCharCount(byteBuffer.Buffer.Memory.Span);
@@ -213,7 +213,7 @@ namespace EpicGames.UHT.Utils
 				int readLength = fs.Read(byteSpan);
 				byteSpan = byteSpan[..readLength];
 
-				Encoding encoding = GetEncoding(byteSpan, out int skipBytes);
+				Encoding encoding = FileUtils.GetEncoding(byteSpan, out int skipBytes);
 				byteSpan = byteSpan[skipBytes..];
 
 				int charCount = encoding.GetCharCount(byteSpan);
@@ -228,52 +228,5 @@ namespace EpicGames.UHT.Utils
 				return false;
 			}
 		}
-
-		/// <summary>
-		/// Get the encoding and number of bytes to skip at the start of the buffer
-		/// </summary>
-		/// <param name="bytes">Bytes to scan for a BOM</param>
-		/// <param name="skipBytes">Number of bytes to skip</param>
-		/// <returns>The encoding</returns>
-		// https://simple.wikipedia.org/wiki/Byte_order_mark
-		private static Encoding GetEncoding(ReadOnlySpan<byte> bytes, out int skipBytes)
-		{
-			skipBytes = 0;
-			Encoding encoding = Encoding.UTF8;
-
-			int length = bytes.Length;
-			if (length >= 2)
-			{
-				if (bytes[0] == 0xfe && bytes[1] == 0xff)
-				{
-					encoding = Encoding.BigEndianUnicode; // Big 16
-					skipBytes = 2;
-				}
-				else if (bytes[0] == 0xff && bytes[1] == 0xfe)
-				{
-					if (length >= 4 && bytes[2] == 0 && bytes[3] == 0)
-					{
-						encoding = Encoding.UTF32; // Little 32
-						skipBytes = 4;
-					}
-					else
-					{
-						encoding = Encoding.Unicode; // Little 16
-						skipBytes = 2;
-					}
-				}
-				else if (length >= 3 && bytes[0] == 0xef && bytes[1] == 0xbb && bytes[2] == 0xbf)
-				{
-					encoding = Encoding.UTF8;
-					skipBytes = 3;
-				}
-				else if (length >= 4 && bytes[0] == 0 && bytes[1] == 0 && bytes[2] == 0xfe && bytes[3] == 0xff)
-				{
-					encoding = new UTF32Encoding(bigEndian: true, byteOrderMark: true); // Big 32
-					skipBytes = 4;
-				}
-			}
-			return encoding;
-		}
 	}
 }
