[UHT]
* Removed dll storage on some methods that were inlined
* Fixed so rigvm functions can have dll storage

diff --git a/Engine/Source/Programs/Shared/EpicGames.UHT/Exporters/CodeGen/UhtHeaderCodeGenerator.cs b/Engine/Source/Programs/Shared/EpicGames.UHT/Exporters/CodeGen/UhtHeaderCodeGenerator.cs
--- a/Engine/Source/Programs/Shared/EpicGames.UHT/Exporters/CodeGen/UhtHeaderCodeGenerator.cs
+++ b/Engine/Source/Programs/Shared/EpicGames.UHT/Exporters/CodeGen/UhtHeaderCodeGenerator.cs
@@ -164,7 +164,7 @@ namespace EpicGames.UHT.Exporters.CodeGen
 				// was marked 'RequiredAPI', this is not needed as C++ will exports all methods automatically.
 				if (textType != UhtPropertyTextType.EventFunctionArgOrRetVal &&
 					!(outerClass != null && outerClass.ClassFlags.HasAnyFlags(EClassFlags.RequiredAPI)) &&
-					exportFlags.HasAnyFlags(UhtFunctionExportFlags.RequiredAPI))
+					exportFlags.HasAnyFlags(UhtFunctionExportFlags.RequiredAPI) && (!isInterface || textType == UhtPropertyTextType.InterfaceFunctionArgOrRetVal))
 				{
 					builder.Append(Module.Api);
 				}
diff --git a/Engine/Source/Programs/Shared/EpicGames.UHT/Exporters/CodeGen/UhtHeaderCodeGeneratorHFile.cs b/Engine/Source/Programs/Shared/EpicGames.UHT/Exporters/CodeGen/UhtHeaderCodeGeneratorHFile.cs
--- a/Engine/Source/Programs/Shared/EpicGames.UHT/Exporters/CodeGen/UhtHeaderCodeGeneratorHFile.cs
+++ b/Engine/Source/Programs/Shared/EpicGames.UHT/Exporters/CodeGen/UhtHeaderCodeGeneratorHFile.cs
@@ -296,7 +296,7 @@ namespace EpicGames.UHT.Exporters.CodeGen
 					}
 					
 					builder.Append('\t').Append(methodInfo.ReturnType).Append(' ').Append(methodInfo.Name).Append('(').Append(constPrefix).Append(scriptStruct.RigVMStructInfo.ExecuteContextType).Append("& InExecuteContext); \\\r\n");
-					builder.Append("\tstatic ").Append(methodInfo.ReturnType).Append(" Static").Append(methodInfo.Name).Append("( \\\r\n");
+					builder.Append("\tstatic ").Append(methodInfo.ReturnType).Append(methodInfo.DllStorage).Append(" Static").Append(methodInfo.Name).Append("( \\\r\n");
 					builder.Append("\t\t");
 					if (!scriptStruct.RigVMStructInfo.HasAnyExecuteContextMember)
 					{
diff --git a/Engine/Source/Programs/Shared/EpicGames.UHT/Parsers/UhtScriptStructParser.cs b/Engine/Source/Programs/Shared/EpicGames.UHT/Parsers/UhtScriptStructParser.cs
--- a/Engine/Source/Programs/Shared/EpicGames.UHT/Parsers/UhtScriptStructParser.cs
+++ b/Engine/Source/Programs/Shared/EpicGames.UHT/Parsers/UhtScriptStructParser.cs
@@ -212,6 +212,7 @@ namespace EpicGames.UHT.Parsers
 
 				// NOTE: The argument list reader doesn't handle templates with multiple arguments (i.e. the ',' issue)
 				topScope.TokenReader
+					.OptionalEndsWith("_API", (api) => { methodInfo.DllStorage = " " + api; })
 					.Optional("virtual")
 					.RequireIdentifier((ref UhtToken identifier) => methodInfo.ReturnType = identifier.Value.ToString())
 					.RequireIdentifier((ref UhtToken identifier) => methodInfo.Name = identifier.Value.ToString());
diff --git a/Engine/Source/Programs/Shared/EpicGames.UHT/Tokenizer/UhtTokenReaderGeneralExtensions.cs b/Engine/Source/Programs/Shared/EpicGames.UHT/Tokenizer/UhtTokenReaderGeneralExtensions.cs
--- a/Engine/Source/Programs/Shared/EpicGames.UHT/Tokenizer/UhtTokenReaderGeneralExtensions.cs
+++ b/Engine/Source/Programs/Shared/EpicGames.UHT/Tokenizer/UhtTokenReaderGeneralExtensions.cs
@@ -212,6 +212,24 @@ namespace EpicGames.UHT.Tokenizer
 			return tokenReader;
 		}
 
+		/// <summary>
+		/// Parse optional text that ends with with the given text
+		/// </summary>
+		/// <param name="tokenReader">Token reader</param>
+		/// <param name="text">Text to match</param>
+		/// <param name="action">Action to invoke if the text was found</param>
+		/// <returns>Token reader</returns>
+		public static IUhtTokenReader OptionalEndsWith(this IUhtTokenReader tokenReader, string text, Action<string> action)
+		{
+			ref UhtToken token = ref tokenReader.PeekToken();
+			if (token.IsIdentifier() && token.Value.Span.EndsWith(text))
+			{
+				tokenReader.ConsumeToken();
+				action(token.Value.Span.ToString());
+			}
+			return tokenReader;
+		}
+
 		/// <summary>
 		/// Parse attributes and optionally alignment specifier 
 		/// </summary>
diff --git a/Engine/Source/Programs/Shared/EpicGames.UHT/Types/UhtScriptStruct.cs b/Engine/Source/Programs/Shared/EpicGames.UHT/Types/UhtScriptStruct.cs
--- a/Engine/Source/Programs/Shared/EpicGames.UHT/Types/UhtScriptStruct.cs
+++ b/Engine/Source/Programs/Shared/EpicGames.UHT/Types/UhtScriptStruct.cs
@@ -527,6 +527,12 @@ namespace EpicGames.UHT.Types
 		/// </summary>
 		/// <returns>True if this method is a predicate</returns>
 		public bool IsPredicate { get; set; } = false;
+
+		/// <summary>
+		/// Is true if method has dll storage such as UE_API
+		/// </summary>
+		/// <returns>True if this method is a predicate</returns>
+		public string DllStorage { get; set; } = String.Empty;
 	}
 
 	/// <summary>
