Implement verse callable functions.

#rnx

diff --git a/Engine/Source/Programs/Shared/EpicGames.UHT/Exporters/CodeGen/UhtHeaderCodeGenerator.cs b/Engine/Source/Programs/Shared/EpicGames.UHT/Exporters/CodeGen/UhtHeaderCodeGenerator.cs
--- a/Engine/Source/Programs/Shared/EpicGames.UHT/Exporters/CodeGen/UhtHeaderCodeGenerator.cs
+++ b/Engine/Source/Programs/Shared/EpicGames.UHT/Exporters/CodeGen/UhtHeaderCodeGenerator.cs
@@ -709,6 +709,32 @@ namespace EpicGames.UHT.Exporters.CodeGen
 			}
 		}
 
+		/// <summary>
+		/// Append just the overall return type of the verse callable function
+		/// </summary>
+		/// <param name="builder">Destination builder</param>
+		/// <param name="function">Function in question</param>
+		/// <param name="isVerseVM">If true, this is for the verse VM where the outer TOptional is stripped</param>
+		/// <returns>Builder</returns>
+		public static StringBuilder AppendVerseCallableReturnType(this StringBuilder builder, UhtFunction function, bool isVerseVM)
+		{
+			UhtProperty? returnProperty = function.ReturnProperty;
+
+			if (returnProperty == null)
+			{
+				builder.Append("void");
+			}
+			else
+			{
+				if (isVerseVM && returnProperty is UhtOptionalProperty optionalProperty)
+				{
+					returnProperty = optionalProperty.ValueProperty;
+				}
+				builder.AppendTokens(returnProperty.TypeTokens.MyTokens);
+			}
+			return builder;
+		}
+
 		/// <summary>
 		/// Append just the overall return type of the native callable verse function
 		/// </summary>
diff --git a/Engine/Source/Programs/Shared/EpicGames.UHT/Exporters/CodeGen/UhtHeaderCodeGeneratorCppFile.cs b/Engine/Source/Programs/Shared/EpicGames.UHT/Exporters/CodeGen/UhtHeaderCodeGeneratorCppFile.cs
--- a/Engine/Source/Programs/Shared/EpicGames.UHT/Exporters/CodeGen/UhtHeaderCodeGeneratorCppFile.cs
+++ b/Engine/Source/Programs/Shared/EpicGames.UHT/Exporters/CodeGen/UhtHeaderCodeGeneratorCppFile.cs
@@ -738,7 +738,22 @@ namespace EpicGames.UHT.Exporters.CodeGen
 				if (scriptStruct.IsVerseField)
 				{
 					builder.Append("\t},\r\n");
-					builder.Append("\t\"").AppendVersePath(scriptStruct, UhtVerseFullNameMode.Qualified).Append("\"\r\n");
+
+					// NOTE: This is temporary while the new VM doesn't match the value in the old VM
+					string defaultName = scriptStruct.GetVerseFullName(UhtVerseFullNameMode.Default);
+					string qualifiedName = scriptStruct.GetVerseFullName(UhtVerseFullNameMode.Qualified);
+					if (defaultName.Equals(qualifiedName, StringComparison.Ordinal))
+					{
+						builder.Append($"\t\"{qualifiedName}\"\r\n");
+					}
+					else
+					{
+						builder.Append("\t#if WITH_VERSE_VM\r\n");
+						builder.Append($"\t\"{defaultName}\"\r\n");
+						builder.Append("\t#else\r\n");
+						builder.Append($"\t\"{qualifiedName}\"\r\n");
+						builder.Append("\t#endif\r\n");
+					}
 				}
 				builder.Append("};\r\n");
 			}
@@ -1130,6 +1145,7 @@ namespace EpicGames.UHT.Exporters.CodeGen
 			bool isRpc = IsRpcFunction(function) && ShouldExportFunction(function);
 			bool isCallback = IsCallbackFunction(function);
 			bool isVerseNativeCallable = function.FunctionExportFlags.HasAnyFlags(UhtFunctionExportFlags.VerseNativeCallable);
+			bool isVerseCallable = function.FunctionExportFlags.HasAnyFlags(UhtFunctionExportFlags.VerseCallable);
 			if (isNotDelegate || isRpc || isCallback)
 			{
 				using UhtCodeBlockComment blockCommand = new(builder, classObj, function);
@@ -1154,6 +1170,10 @@ namespace EpicGames.UHT.Exporters.CodeGen
 							AppendVerseNativeCallableFunction(builder, classObj, function, UhtNames.VerseProxySuffix);
 						}
 					}
+					if (isVerseCallable)
+					{
+						AppendVerseCallableFunction(builder, classObj, function, "verse_exe_");
+					}
 				}
 				if (isRpc)
 				{
@@ -1329,6 +1349,192 @@ namespace EpicGames.UHT.Exporters.CodeGen
 			return builder;
 		}
 
+		private static StringBuilder AppendVerseCallableFunction(StringBuilder builder, UhtClass classObj, UhtFunction function, string functionPrefix)
+		{
+			UhtProperty? returnProperty = function.ReturnProperty;
+
+			bool isDecides = function.FunctionExportFlags.HasAnyFlags(UhtFunctionExportFlags.VerseDecides);
+			bool isCoroutine = function.FunctionExportFlags.HasAnyFlags(UhtFunctionExportFlags.VerseSuspends);
+			bool isNoRollback = function.FunctionExportFlags.HasAnyFlags(UhtFunctionExportFlags.VerseNoRollback);
+			bool isIncompleteReturn = function.Session.IsIncompleteReturn(returnProperty);
+			bool isVoidReturn = returnProperty == null;
+
+			// Write _exec_ wrapper that checks for native method implementation
+			if (isNoRollback)
+			{
+				builder.Append("V_DEFINE_EXEC_NOROLLBACK(").Append(classObj.SourceName).Append(", ").Append(function.VerseName).Append(")\r\n");
+			}
+			else if (isCoroutine)
+			{
+				builder.Append("V_DEFINE_EXEC_SUSPENDS(").Append(classObj.SourceName).Append(", ").Append(function.VerseName).Append(")\r\n");
+			}
+			else
+			{
+				builder.Append("V_DEFINE_EXEC(").Append(classObj.SourceName).Append(", ").Append(function.VerseName).Append(")\r\n");
+			}
+
+			builder.Append("{\r\n");
+			builder.Append("\tV_CALL_IMPL_NO_PRED(").Append(function.VerseName).Append(");\r\n");
+			builder.Append("}\r\n");
+
+			builder.Append("V_DEFINE_IMPL_NO_PRED(").Append(classObj.SourceName).Append(", ").Append(function.VerseName).Append(")\r\n");
+			builder.Append("{\r\n");
+			// The bodies must differ based on the VM
+			builder.Append("#if WITH_VERSE_VM\r\n");
+			AppendVerseCallableFunctionBody(builder, classObj, function, returnProperty, true);
+			builder.Append("#else\r\n");
+			AppendVerseCallableFunctionBody(builder, classObj, function, returnProperty, false);
+			builder.Append("#endif\r\n");
+			builder.Append("}\r\n");
+			return builder;
+		}
+
+		private static StringBuilder AppendVerseCallableFunctionBody(StringBuilder builder, UhtClass classObj, UhtFunction function, UhtProperty? returnProperty, bool isVerseVM)
+		{
+			bool isDecides = function.FunctionExportFlags.HasAnyFlags(UhtFunctionExportFlags.VerseDecides);
+			bool isCoroutine = function.FunctionExportFlags.HasAnyFlags(UhtFunctionExportFlags.VerseSuspends);
+			bool isNoRollback = function.FunctionExportFlags.HasAnyFlags(UhtFunctionExportFlags.VerseNoRollback);
+			bool isIncompleteReturn = function.Session.IsIncompleteReturn(returnProperty);
+			bool isVoidReturn = returnProperty == null;
+
+			builder.Append("\tV_MARSHALLING_PARAM_BEGIN\r\n");
+			if (isVerseVM || !isCoroutine)
+			{
+				// TODO 
+				//WriteParamMarshallingMacros(ParamExpression, Signature.GetParams().begin(), InteropBlock);
+				builder.Append("\tV_MARSHAL_TUPLE_BEGIN\r\n");				
+				builder.Append("\tV_MARSHAL_TUPLE_END\r\n");
+			}
+			builder.Append("\tV_MARSHALLING_END\r\n");
+
+			builder.Append("\tV_NATIVE_BEGIN(").Append(classObj.SourceName).Append(", V_COMMA_SEPARATED(");
+			if (!isVerseVM)
+			{
+				if (isCoroutine)
+				{
+					builder.Append("FVerseResult");
+				}
+				else
+				{
+					builder.AppendVerseCallableReturnType(function, isVerseVM);
+				}
+			}
+			else
+			{
+				builder.AppendVerseCallableReturnType(function, isVerseVM);
+			}
+			builder.Append("))\r\n");
+
+			//if (bUseVerseBPVM && bCoroutine)
+			//{
+			//	WriteCoroutineTask(Name, Signature, Scope->GetKind() != CScope::EKind::Module, Function->_ExtensionFieldAccessorKind, InteropBlock);
+			//}
+			if (isVerseVM && !isCoroutine && !isDecides && (isIncompleteReturn || isVoidReturn))
+			{
+				builder.Append("\t__NativeReturnValue.Emplace();");
+			}
+
+			builder.Append('\t');
+			if (isCoroutine)
+			{
+				if (!isVerseVM)
+				{
+					builder.Append("*__NativeReturnValue = ");
+				}
+				else
+				{
+					builder.Append("__ControlFlow = ");
+				}
+			}
+			else if (!isIncompleteReturn && (isDecides || !isVoidReturn))
+			{
+				if (!isVerseVM)
+				{
+					builder.Append('*');
+				}
+				builder.Append("__NativeReturnValue = ");
+			}
+			//if (Scope->GetKind() == CScope::EKind::Module)
+			//{
+			//	InteropBlock.WriteFormat("verse::%s::", Scope->GetScopeName().AsCString());
+			//}
+			//else
+			{
+				if (!isVerseVM && isCoroutine)
+				{
+					builder.Append("((FTask*)V_THIS)->_Self->");
+				}
+				else
+				{
+					builder.Append("V_THIS->");
+				}
+			}
+
+			builder.Append(function.SourceName).Append('(');
+
+			bool needsComma = false;
+			if (isCoroutine)
+			{
+				if (!isVerseVM)
+				{
+					builder.Append("(FTask*)V_THIS");
+				}
+				else
+				{
+					builder.Append("{__Context}");
+				}
+				needsComma = true;
+			}
+
+			foreach (UhtType argumentType in function.ParameterProperties.Span)
+			{
+				if (argumentType is UhtProperty argument)
+				{
+					if (needsComma)
+					{
+						builder.Append(", ");
+					}
+					if (!isVerseVM && isCoroutine)
+					{
+						builder.Append("((FTask*)V_THIS)->");
+					}
+					builder.Append(argument.SourceName); // InteropBlock << GetRemappedSymbolString(Param->GetName());
+					needsComma = true;
+				}
+			}
+			if (!isCoroutine && isIncompleteReturn)
+			{
+				if (needsComma)
+				{
+					builder.Append(", ");
+				}
+				if (!isVerseVM)
+				{
+					builder.Append("*__NativeReturnValue");
+				}
+				else
+				{
+					if (isDecides)
+					{
+						builder.Append("__NativeReturnValue");
+					}
+					else
+					{
+						builder.Append("*__NativeReturnValue");
+					}
+				}
+				needsComma = true;
+			}
+			builder.Append(");\r\n");
+
+			if (!isVerseVM && isCoroutine)
+			{
+				builder.Append("\t((FTask*)V_THIS)->_bEverSuspended = true;\r\n");
+			}
+			builder.Append("\tV_NATIVE_END(").Append(function.FunctionExportFlags.HasAnyFlags(UhtFunctionExportFlags.VerseRTFMAlwaysOpen) ? "true" : "false").Append(")\r\n");
+			return builder;
+		}
+
 		private StringBuilder AppendClassFunctionCallback(StringBuilder builder, UhtClass classObj, UhtFunction function)
 		{
 			// Net response functions don't go into the VM
@@ -2144,6 +2350,7 @@ namespace EpicGames.UHT.Exporters.CodeGen
 
 			if (!functions.IsEmpty)
 			{
+				bool hasVerseCallableFunctions = false;
 				{
 					using UhtMacroBlockEmitter blockEmitter = new(builder, UhtDefineScopeNames.WithEditor, functions.SoleScope);
 					builder.Append("\tUClass* Class = ").Append(classObj.SourceName).Append("::StaticClass();\r\n");
@@ -2160,6 +2367,31 @@ namespace EpicGames.UHT.Exporters.CodeGen
 							.Append("::exec")
 							.Append(function.EngineName)
 							.Append(" },\r\n");
+						if (function.IsVerseField && function.FunctionExportFlags.HasAnyFlags(UhtFunctionExportFlags.VerseCallable))
+						{
+							hasVerseCallableFunctions = true;
+						}
+					}
+
+					if (hasVerseCallableFunctions)
+					{
+						builder.Append("#if WITH_VERSE_BPVM\r\n");
+						foreach (UhtFunction function in functions.Instances)
+						{
+							if (function.IsVerseField && function.FunctionExportFlags.HasAnyFlags(UhtFunctionExportFlags.VerseCallable))
+							{
+								blockEmitter.Set(function.DefineScope);
+								builder
+									.Append("\t\t{ ")
+									.AppendUTF8LiteralString(function.EncodedVersePath)
+									.Append(", &")
+									.AppendClassSourceNameOrInterfaceName(classObj)
+									.Append("::_exec_")
+									.Append(function.SourceName)
+									.Append("___ },\r\n");
+							}
+						}
+						builder.Append("#endif\r\n");
 					}
 
 					// This will close the block if we have one that isn't editor only
@@ -2168,6 +2400,36 @@ namespace EpicGames.UHT.Exporters.CodeGen
 					builder.Append("\t};\r\n");
 					builder.Append("\tFNativeFunctionRegistrar::RegisterFunctions(Class, Funcs, UE_ARRAY_COUNT(Funcs));\r\n");
 				}
+
+				if (hasVerseCallableFunctions)
+				{
+					builder.Append("#if WITH_VERSE_VM\r\n");
+					using UhtMacroBlockEmitter blockEmitter = new(builder, UhtDefineScopeNames.WithEditor, functions.SoleScope);
+					builder.Append("\tstatic const FVerseCallableThunk VerseFuncs[] = {\r\n");
+
+					foreach (UhtFunction function in functions.Instances)
+					{
+						if (function.IsVerseField && function.FunctionExportFlags.HasAnyFlags(UhtFunctionExportFlags.VerseCallable))
+						{
+							blockEmitter.Set(function.DefineScope);
+							builder
+								.Append("\t\t{ ")
+								.AppendUTF8LiteralString(function.VersePath)
+								.Append(", &")
+								.AppendClassSourceNameOrInterfaceName(classObj)
+								.Append("::_exec_")
+								.Append(function.SourceName)
+								.Append("___ },\r\n");
+						}
+					}
+
+					// This will close the block if we have one that isn't editor only
+					blockEmitter.Set(functions.SoleScope);
+
+					builder.Append("\t};\r\n");
+					builder.Append("\tVerse::CodeGen::Private::RegisterVerseCallableThunks(Class, VerseFuncs, UE_ARRAY_COUNT(VerseFuncs));\r\n");
+					builder.Append("#endif\r\n");
+				}
 			}
 
 			builder.Append("}\r\n");
diff --git a/Engine/Source/Programs/Shared/EpicGames.UHT/Exporters/CodeGen/UhtHeaderCodeGeneratorHFile.cs b/Engine/Source/Programs/Shared/EpicGames.UHT/Exporters/CodeGen/UhtHeaderCodeGeneratorHFile.cs
--- a/Engine/Source/Programs/Shared/EpicGames.UHT/Exporters/CodeGen/UhtHeaderCodeGeneratorHFile.cs
+++ b/Engine/Source/Programs/Shared/EpicGames.UHT/Exporters/CodeGen/UhtHeaderCodeGeneratorHFile.cs
@@ -1259,11 +1259,26 @@ namespace EpicGames.UHT.Exporters.CodeGen
 		{
 			foreach (UhtFunction function in functions)
 			{
-				if (!ShouldExportFunction(function))
+				if (ShouldExportFunction(function))
 				{
-					continue;
+					builder.Append("\tDECLARE_FUNCTION(").Append(function.UnMarshalAndCallName).Append("); \\\r\n");
+				}
+				if (function.FunctionExportFlags.HasAnyFlags(UhtFunctionExportFlags.VerseCallable))
+				{
+					if (function.FunctionExportFlags.HasAnyFlags(UhtFunctionExportFlags.VerseNoRollback))
+					{
+						builder.Append("\tV_DECLARE_EXEC_NOROLLBACK(").Append(function.VerseName).Append("); \\\r\n");
+					}
+					else if (function.FunctionExportFlags.HasAnyFlags(UhtFunctionExportFlags.VerseSuspends))
+					{
+						builder.Append("\tV_DECLARE_EXEC_SUSPENDS(").Append(function.VerseName).Append("); \\\r\n");
+					}
+					else
+					{
+						builder.Append("\tV_DECLARE_EXEC(").Append(function.VerseName).Append("); \\\r\n");
+					}
+					builder.Append("\tV_DECLARE_IMPL_NO_PRED(").Append(function.VerseName).Append("); \\\r\n");
 				}
-				builder.Append("\tDECLARE_FUNCTION(").Append(function.UnMarshalAndCallName).Append("); \\\r\n");
 			}
 			return builder;
 		}
diff --git a/Engine/Source/Programs/Shared/EpicGames.UHT/Specifiers/UhtFunctionSpecifiers.cs b/Engine/Source/Programs/Shared/EpicGames.UHT/Specifiers/UhtFunctionSpecifiers.cs
--- a/Engine/Source/Programs/Shared/EpicGames.UHT/Specifiers/UhtFunctionSpecifiers.cs
+++ b/Engine/Source/Programs/Shared/EpicGames.UHT/Specifiers/UhtFunctionSpecifiers.cs
@@ -419,15 +419,27 @@ namespace EpicGames.UHT.Parsers
 			foreach (KeyValuePair<StringView, StringView> kvp in values)
 			{
 				ReadOnlySpan<char> key = kvp.Key.Span;
-				if (key.Equals("NativeCallable", StringComparison.OrdinalIgnoreCase))
+				if (key.Equals("native_callable", StringComparison.OrdinalIgnoreCase))
 				{
 					functionObj.FunctionExportFlags |= UhtFunctionExportFlags.VerseNativeCallable;
 				}
-				else if (key.Equals("Decides", StringComparison.OrdinalIgnoreCase))
+				else if (key.Equals("verse_callable", StringComparison.OrdinalIgnoreCase))
+				{
+					functionObj.FunctionExportFlags |= UhtFunctionExportFlags.VerseCallable;
+				}
+				else if (key.Equals("decides", StringComparison.OrdinalIgnoreCase))
 				{
 					functionObj.FunctionExportFlags |= UhtFunctionExportFlags.VerseDecides;
 				}
-				else if (key.Equals("Name", StringComparison.OrdinalIgnoreCase))
+				else if (key.Equals("no_rollback", StringComparison.OrdinalIgnoreCase))
+				{
+					functionObj.FunctionExportFlags |= UhtFunctionExportFlags.VerseNoRollback;
+				}
+				else if (key.Equals("rtfm_always_open", StringComparison.OrdinalIgnoreCase))
+				{
+					functionObj.FunctionExportFlags |= UhtFunctionExportFlags.VerseRTFMAlwaysOpen;
+				}
+				else if (key.Equals("name", StringComparison.OrdinalIgnoreCase))
 				{
 					functionObj.VerseName = kvp.Value.ToString();
 				}
diff --git a/Engine/Source/Programs/Shared/EpicGames.UHT/Types/Properties/UhtOptionalProperty.cs b/Engine/Source/Programs/Shared/EpicGames.UHT/Types/Properties/UhtOptionalProperty.cs
--- a/Engine/Source/Programs/Shared/EpicGames.UHT/Types/Properties/UhtOptionalProperty.cs
+++ b/Engine/Source/Programs/Shared/EpicGames.UHT/Types/Properties/UhtOptionalProperty.cs
@@ -250,7 +250,7 @@ namespace EpicGames.UHT.Types
 				return null;
 			}
 
-			if (!value.PropertyCaps.HasAnyFlags(UhtPropertyCaps.CanBeContainerValue))
+			if (!value.PropertyCaps.HasAnyFlags(UhtPropertyCaps.CanBeOptionalValue))
 			{
 				tokenReader.LogError($"The type \'{value.GetUserFacingDecl()}\' can not be used as a value in a TOptional");
 			}
diff --git a/Engine/Source/Programs/Shared/EpicGames.UHT/Types/UhtClass.cs b/Engine/Source/Programs/Shared/EpicGames.UHT/Types/UhtClass.cs
--- a/Engine/Source/Programs/Shared/EpicGames.UHT/Types/UhtClass.cs
+++ b/Engine/Source/Programs/Shared/EpicGames.UHT/Types/UhtClass.cs
@@ -1625,9 +1625,9 @@ namespace EpicGames.UHT.Types
 				{
 					if (IsVerseField)
 					{
-						if (ClassType == UhtClassType.Interface && SuperClass != Session.UObject)
+						if (ClassType == UhtClassType.Interface && SuperClass != Session.UInterface)
 						{
-							this.LogError($"Verse Interface class '{SourceName}' must inherit from UObject'");
+							this.LogError($"Verse Interface class '{SourceName}' must inherit from UInterface'");
 						}
 						else if (ClassType == UhtClassType.NativeInterface && SuperClass != Session.IInterface)
 						{
diff --git a/Engine/Source/Programs/Shared/EpicGames.UHT/Types/UhtField.cs b/Engine/Source/Programs/Shared/EpicGames.UHT/Types/UhtField.cs
--- a/Engine/Source/Programs/Shared/EpicGames.UHT/Types/UhtField.cs
+++ b/Engine/Source/Programs/Shared/EpicGames.UHT/Types/UhtField.cs
@@ -117,18 +117,10 @@ namespace EpicGames.UHT.Types
 		/// <summary>
 		/// Return the complete verse name
 		/// </summary>
-		public string VersePath
-		{
-			get
-			{
-				StringBuilder builder = new();
-				AppendVerseFullName(builder, UhtVerseFullNameMode.Default);
-				return builder.ToString();
-			}
-		}
+		public string VersePath => GetVerseFullName(UhtVerseFullNameMode.Default);
 
 		/// <summary>
-		/// REturn the function name encoded for BPVM
+		/// Return the function name encoded for BPVM
 		/// </summary>
 		public string EncodedVersePath
 		{
@@ -150,6 +142,18 @@ namespace EpicGames.UHT.Types
 		{
 		}
 
+		/// <summary>
+		/// Return the full name of the verse type
+		/// </summary>
+		/// <param name="mode">Controls how the name is generated</param>
+		/// <returns></returns>
+		public string GetVerseFullName(UhtVerseFullNameMode mode)
+		{
+			StringBuilder builder = new();
+			AppendVerseFullName(builder, mode);
+			return builder.ToString();
+		}
+
 		/// <summary>
 		/// Append the verse path to the given builder
 		/// </summary>
diff --git a/Engine/Source/Programs/Shared/EpicGames.UHT/Types/UhtFunction.cs b/Engine/Source/Programs/Shared/EpicGames.UHT/Types/UhtFunction.cs
--- a/Engine/Source/Programs/Shared/EpicGames.UHT/Types/UhtFunction.cs
+++ b/Engine/Source/Programs/Shared/EpicGames.UHT/Types/UhtFunction.cs
@@ -131,6 +131,21 @@ namespace EpicGames.UHT.Types
 		/// True if the return value is wrapped with TVerseTask
 		/// </summary>
 		VerseSuspends = 1 << 20,
+
+		/// <summary>
+		/// Function can be called from verse
+		/// </summary>
+		VerseCallable = 1 << 21,
+
+		/// <summary>
+		/// True if the function is a no_rollback function
+		/// </summary>
+		VerseNoRollback = 1 << 22,
+
+		/// <summary>
+		/// True if the function is marked as rtfm_always_open
+		/// </summary>
+		VerseRTFMAlwaysOpen = 1 << 23,
 	}
 
 	/// <summary>
@@ -843,6 +858,16 @@ namespace EpicGames.UHT.Types
 						{
 							this.LogError($"Function '{SourceName}' is marked as a Verse function but appears in a class that isn't a Verse class");
 						}
+
+						if (!FunctionExportFlags.HasAnyFlags(UhtFunctionExportFlags.VerseNativeCallable | UhtFunctionExportFlags.VerseCallable))
+						{
+							this.LogError($"Verse function '{SourceName}' must either be marked as NativeCallable or VerseCallable");
+						}
+
+						if (VersePath == SourceName)
+						{
+							this.LogError($"Verse function '{SourceName}' has the same name as the Verse generated name.  Change the source name and then specify the name in the UFUNCTION attributes");
+						}
 					}
 					break;
 
diff --git a/Engine/Source/Programs/Shared/EpicGames.UHT/Types/UhtProperty.cs b/Engine/Source/Programs/Shared/EpicGames.UHT/Types/UhtProperty.cs
--- a/Engine/Source/Programs/Shared/EpicGames.UHT/Types/UhtProperty.cs
+++ b/Engine/Source/Programs/Shared/EpicGames.UHT/Types/UhtProperty.cs
@@ -409,6 +409,11 @@ namespace EpicGames.UHT.Types
 		/// True if the property should codegen as a byte enumeration
 		/// </summary>
 		IsRigVMEnumAsByte = 1 << 15,
+
+		/// <summary>
+		/// If true, the property type can be TOptional value
+		/// </summary>
+		CanBeOptionalValue = 1 << 16,
 	}
 
 	/// <summary>
@@ -1179,7 +1184,7 @@ namespace EpicGames.UHT.Types
 		protected UhtProperty(UhtType outer, int lineNumber) : base(outer.HeaderFile, outer, lineNumber)
 		{
 			PropertyFlags = EPropertyFlags.None;
-			PropertyCaps = UhtPropertyCaps.CanBeContainerValue | UhtPropertyCaps.CanBeContainerKey | UhtPropertyCaps.CanHaveConfig;
+			PropertyCaps = UhtPropertyCaps.CanBeContainerValue | UhtPropertyCaps.CanBeContainerKey | UhtPropertyCaps.CanBeOptionalValue | UhtPropertyCaps.CanHaveConfig;
 		}
 
 		/// <summary>
@@ -1208,7 +1213,7 @@ namespace EpicGames.UHT.Types
 			Getter = propertySettings.Getter;
 			Setter = propertySettings.Setter;
 			IsBitfield = propertySettings.IsBitfield;
-			PropertyCaps = UhtPropertyCaps.CanBeContainerValue | UhtPropertyCaps.CanBeContainerKey | UhtPropertyCaps.CanHaveConfig;
+			PropertyCaps = UhtPropertyCaps.CanBeContainerValue | UhtPropertyCaps.CanBeContainerKey | UhtPropertyCaps.CanBeOptionalValue | UhtPropertyCaps.CanHaveConfig;
 			TypeTokens = propertySettings.TypeTokens;
 		}
 
