[UHT]
* Added dll storage to generated Execute function in RIGVM_METHOD

diff --git a/Engine/Source/Programs/Shared/EpicGames.UHT/Exporters/CodeGen/UhtHeaderCodeGeneratorHFile.cs b/Engine/Source/Programs/Shared/EpicGames.UHT/Exporters/CodeGen/UhtHeaderCodeGeneratorHFile.cs
--- a/Engine/Source/Programs/Shared/EpicGames.UHT/Exporters/CodeGen/UhtHeaderCodeGeneratorHFile.cs
+++ b/Engine/Source/Programs/Shared/EpicGames.UHT/Exporters/CodeGen/UhtHeaderCodeGeneratorHFile.cs
@@ -295,8 +295,8 @@ namespace EpicGames.UHT.Exporters.CodeGen
 						continue;
 					}
 					
-					builder.Append('\t').Append(methodInfo.ReturnType).Append(' ').Append(methodInfo.Name).Append('(').Append(constPrefix).Append(scriptStruct.RigVMStructInfo.ExecuteContextType).Append("& InExecuteContext); \\\r\n");
-					builder.Append("\tstatic ").Append(methodInfo.ReturnType).Append(methodInfo.DllStorage).Append(" Static").Append(methodInfo.Name).Append("( \\\r\n");
+					builder.Append('\t').Append(methodInfo.DllStorage).Append(methodInfo.ReturnType).Append(' ').Append(methodInfo.Name).Append('(').Append(constPrefix).Append(scriptStruct.RigVMStructInfo.ExecuteContextType).Append("& InExecuteContext); \\\r\n");
+					builder.Append("\tstatic ").Append(methodInfo.ReturnType).Append(' ').Append(methodInfo.DllStorage).Append("Static").Append(methodInfo.Name).Append("( \\\r\n");
 					builder.Append("\t\t");
 					if (!scriptStruct.RigVMStructInfo.HasAnyExecuteContextMember)
 					{
diff --git a/Engine/Source/Programs/Shared/EpicGames.UHT/Parsers/UhtScriptStructParser.cs b/Engine/Source/Programs/Shared/EpicGames.UHT/Parsers/UhtScriptStructParser.cs
--- a/Engine/Source/Programs/Shared/EpicGames.UHT/Parsers/UhtScriptStructParser.cs
+++ b/Engine/Source/Programs/Shared/EpicGames.UHT/Parsers/UhtScriptStructParser.cs
@@ -212,7 +212,7 @@ namespace EpicGames.UHT.Parsers
 
 				// NOTE: The argument list reader doesn't handle templates with multiple arguments (i.e. the ',' issue)
 				topScope.TokenReader
-					.OptionalEndsWith("_API", (api) => { methodInfo.DllStorage = " " + api; })
+					.OptionalEndsWith("_API", (api) => { methodInfo.DllStorage = api + " "; })
 					.Optional("virtual")
 					.RequireIdentifier((ref UhtToken identifier) => methodInfo.ReturnType = identifier.Value.ToString())
 					.RequireIdentifier((ref UhtToken identifier) => methodInfo.Name = identifier.Value.ToString());
