Horde Server: Adjust http timeouts and fix issue with storage backend and global http client timeout

diff --git a/Engine/Source/Programs/Shared/EpicGames.Horde/HordeHttpMessageHandler.cs b/Engine/Source/Programs/Shared/EpicGames.Horde/HordeHttpMessageHandler.cs
--- a/Engine/Source/Programs/Shared/EpicGames.Horde/HordeHttpMessageHandler.cs
+++ b/Engine/Source/Programs/Shared/EpicGames.Horde/HordeHttpMessageHandler.cs
@@ -59,7 +59,7 @@ namespace EpicGames.Horde
 				logger.LogWarning(KnownLogEvents.Systemic_Horde_Http, ex, "{Method} {Url} retrying after {Time}s.", request.Method, request.RedactedRequestUri(), timespan.TotalSeconds);
 			}
 
-			TimeSpan[] retryTimes = new[] { TimeSpan.FromSeconds(5.0), TimeSpan.FromSeconds(10.0) };
+			TimeSpan[] retryTimes = new[] { TimeSpan.FromSeconds(5.0), TimeSpan.FromSeconds(10.0), TimeSpan.FromSeconds(15.0) };
 			AsyncRetryPolicy retryPolicy = Policy.Handle<TimeoutRejectedException>().WaitAndRetryAsync(retryTimes, OnRetry);
 			return retryPolicy.WrapAsync(timeoutPolicy);
 		}
diff --git a/Engine/Source/Programs/Shared/EpicGames.Horde/Storage/Backends/HttpStorageBackend.cs b/Engine/Source/Programs/Shared/EpicGames.Horde/Storage/Backends/HttpStorageBackend.cs
--- a/Engine/Source/Programs/Shared/EpicGames.Horde/Storage/Backends/HttpStorageBackend.cs
+++ b/Engine/Source/Programs/Shared/EpicGames.Horde/Storage/Backends/HttpStorageBackend.cs
@@ -157,6 +157,9 @@ namespace EpicGames.Horde.Storage.Backends
 				if (redirectResponse.UploadUrl != null)
 				{
 					using HttpClient uploadRedirectClient = _createUploadRedirectClient();
+					// The http client timeout default of 100 seconds acts as a global timeout across all retry and wait periods
+					// We're creating a client per upload, so bump the global client timeout to account for retries specified on the on policy
+					uploadRedirectClient.Timeout = TimeSpan.FromSeconds(240);
 					IAsyncPolicy<HttpResponseMessage> retryPolicy = HordeHttpMessageHandler.CreateDefaultTransientErrorPolicy();
 					
 					using HttpResponseMessage uploadResponse = await retryPolicy.ExecuteAsync(() => uploadRedirectClient.PutAsync(redirectResponse.UploadUrl, streamContent, cancellationToken));
