Add menu for importing snapshots via UGS when syncing.

[RN] none

diff --git a/Engine/Source/Programs/Shared/EpicGames.ProjectStore/ProjectStoreTypes.cs b/Engine/Source/Programs/Shared/EpicGames.ProjectStore/ProjectStoreTypes.cs
--- a/Engine/Source/Programs/Shared/EpicGames.ProjectStore/ProjectStoreTypes.cs
+++ b/Engine/Source/Programs/Shared/EpicGames.ProjectStore/ProjectStoreTypes.cs
@@ -15,6 +15,9 @@ using System.Text;
 using EpicGames.Core;
 using EpicGames.Horde;
 using UnrealBuildBase;
+using EpicGames.OIDC;
+using System.Threading.Tasks;
+using Microsoft.Extensions.Configuration;
 
 namespace EpicGames.ProjectStore
 {
@@ -153,6 +156,11 @@ namespace EpicGames.ProjectStore
 		/// <param name="asyncImport"></param>
 		public abstract void Import(FileReference projectFile, DirectoryReference platformCookedDirectory, bool forceImport, bool asyncImport);
 
+		/// <summary>
+		/// Returns the changelist associated with this snapshot.
+		/// </summary>
+		public abstract int GetChangelist();
+
 		private static string[] GetHostAddresses()
 		{
 			return Dns.GetHostAddresses("", System.Net.Sockets.AddressFamily.InterNetwork).Select(addr => addr.ToString()).ToArray();
@@ -356,6 +364,7 @@ namespace EpicGames.ProjectStore
 		}
 
 		private readonly SnapshotDescriptorCollection _snapshotDescriptors = new SnapshotDescriptorCollection();
+		private readonly int _changelist;
 
 		private void LoadSnapshotDescriptors(string snapshotDescriptorJSON)
 		{
@@ -372,10 +381,12 @@ namespace EpicGames.ProjectStore
 		/// Constructs a fileshare build using a JSON string
 		/// </summary>
 		/// <param name="snapshotDescriptorJSON"></param>
+		/// <param name="changelist">The changelist associated with this build.</param>
 		/// <param name="runApp">A delegate to use to run command-line apps.</param>
-		public FileshareBuild(string snapshotDescriptorJSON, RunAppDelegate runApp)
+		public FileshareBuild(string snapshotDescriptorJSON, int changelist, RunAppDelegate runApp)
 			: base(runApp)
 		{
+			_changelist = changelist;
 			LoadSnapshotDescriptors(snapshotDescriptorJSON);
 		}
 
@@ -383,10 +394,12 @@ namespace EpicGames.ProjectStore
 		/// Constructs a fileshare build using a snapshot descriptor from a network fileshare.
 		/// </summary>
 		/// <param name="snapshotDescriptor"></param>
+		/// <param name="changelist">The changelist associated with this build.</param>
 		/// <param name="runApp">A delegate to use to run command-line apps.</param>
-		public FileshareBuild(FileReference snapshotDescriptor, RunAppDelegate runApp)
+		public FileshareBuild(FileReference snapshotDescriptor, int changelist, RunAppDelegate runApp)
 			: base(runApp)
 		{
+			_changelist = changelist;
 			string data = FileReference.ReadAllText(snapshotDescriptor);
 			LoadSnapshotDescriptors(data);
 		}
@@ -452,6 +465,12 @@ namespace EpicGames.ProjectStore
 				_runApp(zenExe.FullName, oplogImportCommandLine.ToString());
 			}
 		}
+
+		/// <inheritdoc/>
+		public override int GetChangelist()
+		{
+			return _changelist;
+		}
 	}
 
 	/// <summary>
@@ -521,6 +540,12 @@ namespace EpicGames.ProjectStore
 			_runApp(zenExe.FullName, oplogCreateCommandLine.ToString());
 			_runApp(zenExe.FullName, oplogImportCommandLine.ToString());
 		}
+
+		/// <inheritdoc/>
+		public override int GetChangelist()
+		{
+			return _buildData.Metadata.Changelist;
+		}
 	}
 
 	/// <summary>
@@ -668,7 +693,7 @@ namespace EpicGames.ProjectStore
 			}
 
 			DirectoryReference buildDir = DirectoryReference.Combine(buildsDir, changelist.ToString());
-			return new FileshareBuild(DirectoryReference.EnumerateFiles(buildDir).First(), _runApp);
+			return new FileshareBuild(DirectoryReference.EnumerateFiles(buildDir).First(), foundChangelist, _runApp);
 		}
 	}
 
@@ -740,6 +765,7 @@ namespace EpicGames.ProjectStore
 
 		/// <summary>
 		/// Constructs an Unreal Cloud build index that will search the given Unreal Cloud host and namespace.
+		/// Uses OidcToken.exe to perform authorization with the host.
 		/// </summary>
 		/// <param name="host">The Unreal Cloud host to search.</param>
 		/// <param name="ns">The Unreal Cloud namespace to search.</param>
@@ -753,6 +779,63 @@ namespace EpicGames.ProjectStore
 			_accessToken = GetAccessToken(cloudAuthServiceName);
 		}
 
+		/// <summary>
+		/// Creates an Unreal Cloud build index that will search the given Unreal Cloud host and namespace,
+		/// using an existing OidcTokenInfo from OidcTokenManager.
+		/// </summary>
+		/// <param name="host">The Unreal CLoud host to search.</param>
+		/// <param name="ns">The Unreal Cloud namespace to search.</param>
+		/// <param name="accessTokenInfo"></param>
+		/// <param name="runApp">A delegate used to run command-line apps.</param>
+		public UnrealCloudDDCBuildIndex(string host, string ns, OidcTokenInfo accessTokenInfo, RunAppDelegate runApp)
+			: base(runApp)
+		{
+			_host = host;
+			_namespace = ns;
+			_accessToken = accessTokenInfo.AccessToken!;
+		}
+
+		/// <summary>
+		/// Creates an Unreal Cloud build index asynchronously using remote auth configuration provided by a given host.
+		/// </summary>
+		/// <param name="host">The Unreal CLoud host to search.</param>
+		/// <param name="ns">The Unreal Cloud namespace to search.</param>
+		/// <param name="runApp">A delegate used to run command-line apps.</param>
+		/// <param name="allowInteractiveLogin">Whether we should allow an interactive login or raise NotLoggedInException.</param>
+		/// <returns>A task that can be awaited to get the authorized Unreal Cloud build index.</returns>
+		/// <exception cref="BuildIndexException">Raised if we fail to determine the auth provider from the host.</exception>
+		/// <exception cref="NotLoggedInException">Raised if we require interactive login and allowInteractiveLogin is false..</exception>
+		public static async Task<UnrealCloudDDCBuildIndex> CreateBuildIndexAsync(string host, string ns, RunAppDelegate runApp, bool allowInteractiveLogin = false)
+		{
+			ClientAuthConfigurationV1? authConfig = await ProviderConfigurationFactory.ReadRemoteAuthConfigurationAsync(new Uri(host), ProviderConfigurationFactory.DefaultEncryptionKey);
+			if (authConfig?.DefaultProvider == null)
+			{
+				throw new BuildIndexException(String.Format("Failed to determine default auth provider from {0}", host));
+			}
+
+			IConfiguration config = ProviderConfigurationFactory.BindOptions(authConfig);
+			using (ITokenStore tokenStore = TokenStoreFactory.CreateTokenStore())
+			{ 
+				OidcTokenManager tokenManager = OidcTokenManager.CreateTokenManager(config, tokenStore);
+
+				try
+				{
+					OidcTokenInfo tokenInfo = await tokenManager.GetAccessToken(authConfig.DefaultProvider);
+					return new UnrealCloudDDCBuildIndex(host, ns, tokenInfo, runApp);
+				}
+				catch (NotLoggedInException)
+				{
+					if (!allowInteractiveLogin)
+					{
+						throw;
+					}
+
+					OidcTokenInfo tokenInfo = await tokenManager.LoginAsync(authConfig.DefaultProvider);
+					return new UnrealCloudDDCBuildIndex(host, ns, tokenInfo, runApp);
+				}
+			}
+		}
+
 		class AccessTokenOutput
 		{
 			public string? Token { get; set; }
@@ -796,6 +879,11 @@ namespace EpicGames.ProjectStore
 			public bool PartialResult { get; set; }
 		}
 
+		class QueryBucketsResponse
+		{
+			public List<string> Buckets { get; set; } = new List<string>();
+		}
+
 		[System.Diagnostics.CodeAnalysis.SuppressMessage("Usage", "CA1308", Justification = "We know the provided strings are safe to lowercase using ToLowerInvariant()")]
 		private static string GetBucket(string project, string branch, string platform, string runtime)
 		{
@@ -822,6 +910,62 @@ namespace EpicGames.ProjectStore
 			return String.Format("{0}.oplog.{1}.{2}",  project, branch, platform);
 		}
 
+		/// <summary>
+		/// Gets a list of the available platforms this build index can find for a given branch and project.
+		/// </summary>
+		/// <param name="branch">The branch to search for.</param>
+		/// <param name="project">The project to search for.</param>
+		/// <returns></returns>
+		/// <exception cref="BuildIndexException">Raised on request failure or unexpected response format.</exception>
+		public List<string> GetAvailablePlatforms(string branch, string project)
+		{
+			Uri uri = new Uri(String.Format("{0}/api/v2/builds/{1}", _host, _namespace));
+
+			HttpResponseMessage response;
+			using (HttpClient client = new HttpClient())
+			{
+				using (HttpRequestMessage request = new HttpRequestMessage())
+				{
+					request.Method = HttpMethod.Get;
+					request.RequestUri = uri;
+					request.Headers.Authorization = new AuthenticationHeaderValue("Bearer", _accessToken);
+
+					response = client.Send(request);
+				}
+			}
+
+			string body;
+			using (StreamReader reader = new StreamReader(response.Content.ReadAsStream()))
+			{
+				body = reader.ReadToEnd();
+			}
+
+			if (!response.IsSuccessStatusCode)
+			{
+				throw new BuildIndexException(String.Format("Failed to query builds: {0} ({1})", response.ReasonPhrase, body));
+			}
+
+			QueryBucketsResponse? buckets = JsonSerializer.Deserialize<QueryBucketsResponse>(body, Build.GetDefaultJsonSerializerOptions());
+			if (buckets == null)
+			{
+				throw new BuildIndexException("Failed to deserialize bucket info returned from Unreal Cloud");
+			}
+
+			branch = StringId.Sanitize(branch).ToString();
+			// Bucket format for snapshot buckets is <project>.oplog.<branch>.<platform>.
+			// Find buckets matching the format, with the specified project and branch,
+			// then return the associated platform.
+			return buckets.Buckets
+				.Select(b => b.Split("."))
+				.Where(b =>
+					b.Length == 4
+					&& b[0].Equals(project, StringComparison.OrdinalIgnoreCase)
+					&& b[1].Equals("oplog", StringComparison.OrdinalIgnoreCase)
+					&& b[2].Equals(branch, StringComparison.OrdinalIgnoreCase))
+				.Select(b => b[3].ToLower())
+				.ToList();
+		}
+
 		private List<UnrealCloudDDCBuildData> QueryBuilds(string branch, string project, string platform, string runtime, int changelist = 0)
 		{
 			string bucket = GetBucket(project, branch, platform, runtime);
diff --git a/Engine/Source/Programs/UnrealGameSync/UnrealGameSync/Controls/WorkspaceControl.cs b/Engine/Source/Programs/UnrealGameSync/UnrealGameSync/Controls/WorkspaceControl.cs
--- a/Engine/Source/Programs/UnrealGameSync/UnrealGameSync/Controls/WorkspaceControl.cs
+++ b/Engine/Source/Programs/UnrealGameSync/UnrealGameSync/Controls/WorkspaceControl.cs
@@ -6,6 +6,7 @@ using System.Diagnostics;
 using System.Diagnostics.CodeAnalysis;
 using System.Drawing;
 using System.Drawing.Drawing2D;
+using System.Globalization;
 using System.IO;
 using System.Linq;
 using System.Runtime.InteropServices;
@@ -18,7 +19,9 @@ using System.Windows.Forms;
 using System.Windows.Forms.VisualStyles;
 using EpicGames.Core;
 using EpicGames.Horde;
+using EpicGames.OIDC;
 using EpicGames.Perforce;
+using EpicGames.ProjectStore;
 using Microsoft.Extensions.DependencyInjection;
 using Microsoft.Extensions.Logging;
 using Microsoft.Win32;
@@ -328,6 +331,24 @@ namespace UnrealGameSync
 				}
 			}
 
+			if (projectConfigFile != null)
+			{
+				string? cloudHost = projectConfigFile.GetValue("CloudStorage.Host", null);
+				string? cloudNamespace = projectConfigFile.GetValue("CloudStorage.Namespace", null);
+
+				if (cloudHost != null && cloudNamespace != null && ProjectInfo.StreamName != null)
+				{
+					UpdateImportSnapshotDropdown(cloudHost, cloudNamespace, ProjectInfo.StreamName, ProjectInfo.LocalFileName.GetFileNameWithoutAnyExtensions());
+					_settings.SnapshotHost = cloudHost;
+					_settings.SnapshotNamespace = cloudNamespace;
+				}
+				else
+				{
+					OptionsContextMenu_ImportSnapshots.Enabled = false;
+					_logger.LogWarning("Unable to determine cloud host and namespace from project's UnrealGameSync.ini. Specify Host=<host> and Namespace=<namespace> in the [CloudStorage] section of your UnrealGameSync.ini to enable snapshot import.");					
+				}
+			}
+
 			DesiredTaskbarState = Tuple.Create(TaskbarState.NoProgress, 0.0f);
 
 			System.Reflection.PropertyInfo doubleBufferedProperty = typeof(Control).GetProperty("DoubleBuffered", System.Reflection.BindingFlags.NonPublic | System.Reflection.BindingFlags.Instance)!;
@@ -448,6 +469,11 @@ namespace UnrealGameSync
 			}
 
 			_owner.ToolUpdateMonitor.OnChange += UpdateStatusPanel_CrossThread;
+			FilterContextMenu_Author_Name.Text = AuthorFilterPlaceholderText;
+
+			FilterContextMenu_Author_Name.TextBox.GotFocus += new System.EventHandler(FilterContextMenu_Author_Name_GotFocus);
+			FilterContextMenu_Author_Name.TextBox.LostFocus += new System.EventHandler(FilterContextMenu_Author_Name_LostFocus);
+			FilterContextMenu_Author_Name.TextBox.TextChanged += new System.EventHandler(FilterContextMenu_Author_Name_Changed);
 		}
 
 		public void IssueMonitor_OnIssuesChangedAsync()
@@ -933,6 +959,24 @@ namespace UnrealGameSync
 		void StartSync(int changeNumber, bool syncOnly, WorkspaceUpdateCallback? callback, bool skipKnownChangeLists)
 		{
 			WorkspaceUpdateOptions options = WorkspaceUpdateOptions.Sync | WorkspaceUpdateOptions.SyncArchives;
+
+			_settings.Snapshots.Clear();
+
+			foreach (ToolStripMenuItem menuItem in OptionsContextMenu_ImportSnapshots.DropDownItems)
+			{
+				if (menuItem.Checked)
+				{
+					if (menuItem.Tag is SnapshotSettings settings)
+					{
+						_settings.Snapshots.Add(settings);
+					}
+				}
+			}
+
+			if (_settings.Snapshots.Count > 0)
+			{
+				options |= WorkspaceUpdateOptions.ImportSnapshots;
+			}
 			
 			if (_settings.SyncOptions.UseFastSync is true)
 			{
@@ -1029,6 +1073,14 @@ namespace UnrealGameSync
 					}
 				}
 			}
+
+			if (options.HasFlag(WorkspaceUpdateOptions.ImportSnapshots))
+			{
+				context.Snapshots.AddRange(_settings.Snapshots);
+				context.SnapshotHost = _settings.SnapshotHost;
+				context.SnapshotNamespace = _settings.SnapshotNamespace;
+			}
+
 			StartWorkspaceUpdate(context, callback);
 		}
 
@@ -5127,6 +5179,72 @@ namespace UnrealGameSync
 			}
 		}
 
+		private void UpdateImportSnapshotDropdown(string host, string ns, string branch, string project, bool allowInteractiveLogin = false)
+		{
+			OptionsContextMenu_ImportSnapshots.DropDownItems.Clear();
+
+			try
+			{
+				Task<UnrealCloudDDCBuildIndex> buildIndexTask = Task.Run(() => UnrealCloudDDCBuildIndex.CreateBuildIndexAsync(host, ns, (app, cmdline) => { }, allowInteractiveLogin));
+				buildIndexTask.Wait();
+
+				if (!buildIndexTask.IsCompletedSuccessfully)
+				{
+					OptionsContextMenu_ImportSnapshots.DropDownItems.Add(new ToolStripMenuItem("Unable to determine valid platforms for snapshot import"));
+					return;
+				}
+
+				UnrealCloudDDCBuildIndex buildIndex = buildIndexTask.Result;
+				List<string> cookPlatforms = buildIndex.GetAvailablePlatforms(branch, project);
+
+				foreach (string cookPlatform in cookPlatforms)
+				{
+					string runtime = cookPlatform.EndsWith("client", StringComparison.OrdinalIgnoreCase) ? "client"
+						: cookPlatform.EndsWith("server", StringComparison.OrdinalIgnoreCase) ? "server"
+						: "game";
+
+					string platform = cookPlatform.Replace(runtime, "", StringComparison.OrdinalIgnoreCase);
+
+					TextInfo textInfo = new CultureInfo("en-US").TextInfo;
+
+					ToolStripMenuItem item = new ToolStripMenuItem(textInfo.ToTitleCase(platform + " " + runtime));
+					item.Enabled = true;
+					item.CheckOnClick = true;
+					item.Tag = new SnapshotSettings { Platform = platform, Runtime = runtime };
+
+					OptionsContextMenu_ImportSnapshots.DropDown.Closing += (sender, eventArgs) =>
+					{
+						ToolStripDropDown dropDown = (ToolStripDropDown)sender!;
+						Point p = dropDown.PointToClient(Control.MousePosition);
+						if (dropDown.DisplayRectangle.Contains(p))
+						{
+							eventArgs.Cancel = true;
+						}
+					};
+
+					OptionsContextMenu_ImportSnapshots.DropDownItems.Add(item);
+				}
+			}
+			catch (AggregateException e)
+			{
+				if (e.GetBaseException() is NotLoggedInException)
+				{
+					ToolStripMenuItem item = new ToolStripMenuItem("Populate available platforms");
+					item.Enabled = true;
+					item.Click += (sender, eventArgs) =>
+					{
+						UpdateImportSnapshotDropdown(host, ns, branch, project, true);
+					};
+
+					OptionsContextMenu_ImportSnapshots.DropDownItems.Add(item);
+				}
+				else
+				{
+					throw;
+				}
+			}
+		}
+
 		private void BuildList_FontChanged(object? sender, EventArgs e)
 		{
 			_buildFont?.Dispose();
@@ -7028,7 +7146,7 @@ namespace UnrealGameSync
 		/// </summary>
 		/// <param name="sender"></param>
 		/// <param name="e"></param>
-		private void FilterContextMenu_Author_Name_GotFocus(object sender, EventArgs e)
+		private void FilterContextMenu_Author_Name_GotFocus(object? sender, EventArgs e)
 		{
 			if (FilterContextMenu_Author_Name.Text == AuthorFilterPlaceholderText)
 			{
@@ -7044,7 +7162,7 @@ namespace UnrealGameSync
 		/// </summary>
 		/// <param name="sender"></param>
 		/// <param name="e"></param>
-		private void FilterContextMenu_Author_Name_LostFocus(object sender, EventArgs e)
+		private void FilterContextMenu_Author_Name_LostFocus(object? sender, EventArgs e)
 		{
 			if (String.IsNullOrEmpty(FilterContextMenu_Author_Name.Text))
 			{
@@ -7060,7 +7178,7 @@ namespace UnrealGameSync
 		/// </summary>
 		/// <param name="sender"></param>
 		/// <param name="e"></param>
-		private void FilterContextMenu_Author_Name_Changed(object sender, EventArgs e)
+		private void FilterContextMenu_Author_Name_Changed(object? sender, EventArgs e)
 		{
 			_authorFilterText = FilterContextMenu_Author_Name.Text;
 			if (!String.IsNullOrWhiteSpace(_authorFilterText) && _authorFilterText != AuthorFilterPlaceholderText)
diff --git a/Engine/Source/Programs/UnrealGameSync/UnrealGameSyncShared/UnrealGameSyncShared.csproj b/Engine/Source/Programs/UnrealGameSync/UnrealGameSyncShared/UnrealGameSyncShared.csproj
--- a/Engine/Source/Programs/UnrealGameSync/UnrealGameSyncShared/UnrealGameSyncShared.csproj
+++ b/Engine/Source/Programs/UnrealGameSync/UnrealGameSyncShared/UnrealGameSyncShared.csproj
@@ -47,6 +47,7 @@
     <ProjectReference Include="..\..\Shared\EpicGames.Horde\EpicGames.Horde.csproj" />
     <ProjectReference Include="..\..\Shared\EpicGames.OIDC\EpicGames.OIDC.csproj" />
     <ProjectReference Include="..\..\Shared\EpicGames.Perforce\EpicGames.Perforce.csproj" />
+    <ProjectReference Include="..\..\Shared\EpicGames.ProjectStore\EpicGames.ProjectStore.csproj" />
   </ItemGroup>
 
 </Project>
diff --git a/Engine/Source/Programs/UnrealGameSync/UnrealGameSyncShared/UserSettings.cs b/Engine/Source/Programs/UnrealGameSync/UnrealGameSyncShared/UserSettings.cs
--- a/Engine/Source/Programs/UnrealGameSync/UnrealGameSyncShared/UserSettings.cs
+++ b/Engine/Source/Programs/UnrealGameSync/UnrealGameSyncShared/UserSettings.cs
@@ -245,6 +245,12 @@ namespace UnrealGameSync
 		}
 	}
 
+	public class SnapshotSettings
+	{
+		public string Runtime { get; set;  } = "";
+		public string Platform { get; set; } = "";
+	}
+
 	public class UserSelectedProjectSettings
 	{
 		public string? ServerAndPort { get; set; }
@@ -563,6 +569,10 @@ namespace UnrealGameSync
 		// Precompiled binaries
 		public List<ArchiveSettings> Archives { get; init; } = new List<ArchiveSettings>();
 
+		public List<SnapshotSettings> Snapshots { get; init; } = new List<SnapshotSettings>();
+		public string SnapshotHost { get; set; } = "";
+		public string SnapshotNamespace { get; set; } = "";
+
 		// OIDC Settings
 		public Dictionary<string, string> ProviderToRefreshTokens { get; init; } = new Dictionary<string, string>();
 
diff --git a/Engine/Source/Programs/UnrealGameSync/UnrealGameSyncShared/WorkspaceUpdate.cs b/Engine/Source/Programs/UnrealGameSync/UnrealGameSyncShared/WorkspaceUpdate.cs
--- a/Engine/Source/Programs/UnrealGameSync/UnrealGameSyncShared/WorkspaceUpdate.cs
+++ b/Engine/Source/Programs/UnrealGameSync/UnrealGameSyncShared/WorkspaceUpdate.cs
@@ -15,7 +15,10 @@ using System.Text.RegularExpressions;
 using System.Threading;
 using System.Threading.Tasks;
 using EpicGames.Core;
+using EpicGames.OIDC;
 using EpicGames.Perforce;
+using EpicGames.ProjectStore;
+using Google.Protobuf.WellKnownTypes;
 using Microsoft.Extensions.Logging;
 
 namespace UnrealGameSync
@@ -41,7 +44,8 @@ namespace UnrealGameSync
 		Clobber = 0x8000,
 		Refilter = 0x10000,
 		UprojectSpecificSolution = 0x20000,
-		FastSync = 0x40000
+		FastSync = 0x40000,
+		ImportSnapshots = 0x80000
 	}
 
 	public enum WorkspaceUpdateResult
@@ -106,6 +110,9 @@ namespace UnrealGameSync
 		public Dictionary<string, string> AdditionalVariables { get; } = new Dictionary<string, string>(StringComparer.OrdinalIgnoreCase);
 		public PerforceSyncOptions? PerforceSyncOptions { get; set; }
 		public List<HaveRecord> HaveFiles { get; } = new List<HaveRecord>(); // Cached when sync filter has changed
+		public List<SnapshotSettings> Snapshots {  get; } = new List<SnapshotSettings>();
+		public string SnapshotHost { get; set; } = "";
+		public string SnapshotNamespace { get; set; } = "";
 
 		// May be updated during sync
 		public ConfigFile? ProjectConfigFile { get; set; }
@@ -1387,6 +1394,57 @@ namespace UnrealGameSync
 				}
 			}
 
+			if (Context.Options.HasFlag(WorkspaceUpdateOptions.ImportSnapshots))
+			{
+				if (project.ProjectDir == null)
+				{
+					logger.LogWarning("Unable to determine project directory to use for snapshot import - snapshot import will be skipped");
+				}
+				else if (project.StreamName == null)
+				{
+					logger.LogWarning("Unable to determine stream to use for snapshot import - snapshot import will be skipped");
+				}
+				else
+				{
+					using (TelemetryStopwatch stopwatch = new TelemetryStopwatch("Workspace_ImportSnapshots", project.TelemetryProjectIdentifier))
+					{
+						stopwatch.AddData(new { MachineName = System.Net.Dns.GetHostName(), DomainName = Environment.UserDomainName, ServerAndPort = perforce.Settings.ServerAndPort, ClientName = perforce.Settings.ClientName, UserName = perforce.Settings.UserName });
+						List<Task<int>> snapshotTasks = new List<Task<int>>();
+
+						Action<string, string> runApp = (string app, string cmdLine) =>
+						{
+							logger.LogInformation("zen> Running Zen command {App} {CommandLine}", app, cmdLine);
+							snapshotTasks.Add(ExecuteShellCommandAsync(app + " " + cmdLine, line => ProcessOutput(line, "zen> ", Progress, logger), cancellationToken));
+						};
+
+						try
+						{
+							BuildIndex buildIndex = await UnrealCloudDDCBuildIndex.CreateBuildIndexAsync(Context.SnapshotHost, Context.SnapshotNamespace, (string app, string cmdLine) => runApp(app, cmdLine), true);
+							DirectoryReference cookedPlatformDir = DirectoryReference.Combine(project.EngineDir, "Saved", "Cooked", "{Platform}");
+
+							foreach (SnapshotSettings snapshot in Context.Snapshots)
+							{
+								Build build = buildIndex.GetBuild(project.StreamName, project.LocalFileName.GetFileNameWithoutAnyExtensions(), snapshot.Platform, snapshot.Runtime, Context.ChangeNumber, true);
+								if (build.GetChangelist() != Context.ChangeNumber)
+								{
+									logger.LogWarning("No exact match found for target changelist {}. Closest changelist found was {}. You may require an incremental cook for this snapshot to function properly.", Context.ChangeNumber, build.GetChangelist());
+								}
+
+								build.Import(project.LocalFileName, cookedPlatformDir, false, false);
+							}
+
+							await Task.WhenAll(snapshotTasks);
+
+							times.Add(new Tuple<string, TimeSpan>("ImportSnapshots", stopwatch.Stop("Success")));
+						}
+						catch (BuildIndexException e)
+						{
+							logger.LogWarning("Failed to import snapshots: {}", e.Message);
+						}
+					}
+				}
+			}
+
 			// Take the lock before doing anything else. Building and generating project files can only be done on one workspace at a time.
 			using SemaphoreScope scope = new SemaphoreScope(_updateSemaphore);
 			if (Context.Options.HasFlag(WorkspaceUpdateOptions.GenerateProjectFiles) || Context.Options.HasFlag(WorkspaceUpdateOptions.Build))
