Bump Shared Assembly info to 5.7

#rnx

diff --git a/Engine/Source/Programs/Shared/MetaData.cs b/Engine/Source/Programs/Shared/MetaData.cs
--- a/Engine/Source/Programs/Shared/MetaData.cs
+++ b/Engine/Source/Programs/Shared/MetaData.cs
@@ -18,8 +18,7 @@ using System.Runtime.InteropServices;
 
 #if !SPECIFIC_VERSION
 // Automatically generate a version number based on the time of compilation
-[assembly: AssemblyVersion("5.6.0.0")]
-[assembly: AssemblyFileVersion("5.6.0.0")]
-[assembly: AssemblyInformationalVersion("5.6.0.0")]
+[assembly: AssemblyVersion("5.7.0.0")]
+[assembly: AssemblyFileVersion("5.7.0.0")]
+[assembly: AssemblyInformationalVersion("5.7.0.0")]
 #endif
-
