Improved LocalizationIssueHandler to cover more incoming edgecases (when changing from UE_LOG to UE_LOGFMT)

diff --git a/Engine/Source/Programs/Shared/EpicGames.Horde/Issues/Handlers/LocalizationIssueHandler.cs b/Engine/Source/Programs/Shared/EpicGames.Horde/Issues/Handlers/LocalizationIssueHandler.cs
--- a/Engine/Source/Programs/Shared/EpicGames.Horde/Issues/Handlers/LocalizationIssueHandler.cs
+++ b/Engine/Source/Programs/Shared/EpicGames.Horde/Issues/Handlers/LocalizationIssueHandler.cs
@@ -1,5 +1,6 @@
 // Copyright Epic Games, Inc. All Rights Reserved.
 
+using System;
 using System.Collections.Generic;
 using System.Text.Json;
 using EpicGames.Core;
@@ -8,22 +9,48 @@ using Microsoft.Extensions.Logging;
 namespace EpicGames.Horde.Issues.Handlers
 {
 	/// <summary>
-	/// Instance of a localization error
+	/// Instance of an issue handler for localization warnings/errors
 	/// </summary>
 	[IssueHandler]
 	public class LocalizationIssueHandler : IssueHandler
 	{
 		private class LocalizationData
 		{
-			public string _localizationTarget = "";
-			public string _localizationID = "";
+			public string _file = "";
+			public int _line = -1;
+			public string _location = ""; // This can be {pathToFile}({line}), just {pathToFile} or multiple other possibilities including {notAFile}. Let's never assume it is easy to parse.
+			public string _localizationNamespace = "";
+			public string _localizationKey = "";
+
+			public string _conflictFile = "";
+			public int _conflictLine = -1;
+			public string _conflictLocation = ""; // This can be {pathToFile}({line}), just {pathToFile} or multiple other possibilities including {notAFile}. Let's never assume it is easy to parse.
+			public string _conflictLocalizationNamespace = "";
+			public string _conflictLocalizationKey = "";
+
+			public bool HasLocalizationKey() { return !System.String.IsNullOrEmpty(_localizationKey); } 
+			public string GetNamespaceKeyString() { return HasLocalizationKey() ? _localizationNamespace + "," + _localizationKey : ""; }
+
+			public bool HasFile() { return !System.String.IsNullOrEmpty(_file); }
+			public bool HasConflictFile() { return !System.String.IsNullOrEmpty(_conflictFile); }
+			public bool HasAnyFile() { return HasFile() || HasConflictFile(); }
 		}
 
+		readonly IssueHandlerContext _context;
 		readonly List<IssueEventGroup> _issues = new List<IssueEventGroup>();
+		readonly IssueEventGroup _issuesWithoutClearFiles = new IssueEventGroup("Localization", "Localization {Severity}", IssueChangeFilter.All);
 
 		/// <inheritdoc/>
 		public override int Priority => 10;
 
+		/// <summary>
+		/// Constructor
+		/// </summary>
+		public LocalizationIssueHandler(IssueHandlerContext context)
+		{
+			_context = context;
+		}
+
 		/// <summary>
 		/// Determines if the given event id matches
 		/// </summary>
@@ -45,13 +72,36 @@ namespace EpicGames.Horde.Issues.Handlers
 		}
 
 		/// <summary>
-		/// Extracts the localization properties of an IssueEvent to generate all the IssueKeys
+		/// Generate all the IssueKeys with the extracted LocalizationData
 		/// </summary>
 		/// <param name="keys">Set of keys</param>
-		/// <param name="issueEvent">The event data</param>
 		/// <param name="data">The localization data</param>
-		private static void AddLocalizationIssueKeys(HashSet<IssueKey> keys, IssueEvent issueEvent, LocalizationData data)
+		private static void AddLocalizationIssueKeys(HashSet<IssueKey> keys, LocalizationData data)
+		{
+			if(data.HasFile())
+			{
+				keys.AddSourceFile(data._file, IssueKeyType.File);
+			}
+
+			if (data.HasConflictFile())
+			{
+				keys.AddSourceFile(data._conflictFile, IssueKeyType.File);
+			}
+
+			if (data.HasLocalizationKey())
+			{
+				keys.Add(new IssueKey(data.GetNamespaceKeyString(), IssueKeyType.None));
+			}
+		}
+
+		/// <summary>
+		/// Extracts the localization data from the properties of an IssueEvent
+		/// </summary>
+		/// <param name="issueEvent">The event data</param>
+		private static LocalizationData GetLocalizationData(IssueEvent issueEvent)
 		{
+			LocalizationData data = new LocalizationData();
+
 			foreach (JsonLogEvent line in issueEvent.Lines)
 			{
 				JsonDocument document = JsonDocument.Parse(line.Data);
@@ -61,44 +111,117 @@ namespace EpicGames.Horde.Issues.Handlers
 				{
 					foreach (JsonProperty property in properties.EnumerateObject())
 					{
-						if (property.NameEquals("conflict") && property.Value.ValueKind == JsonValueKind.String)
+						if (property.NameEquals("file") && property.Value.ValueKind == JsonValueKind.String)
+						{
+							data._file = property.Value.GetString()!;
+						}
+						else if (property.NameEquals("line") && property.Value.ValueKind == JsonValueKind.Number)
 						{
-							keys.AddSourceFile(property.Value.GetString()!, IssueKeyType.File, data._localizationID);
+							data._line = property.Value.GetInt32();
 						}
-						else if((property.NameEquals("localizationID") || property.NameEquals("locID")) && property.Value.ValueKind == JsonValueKind.String)
+						else if (property.NameEquals("location") && property.Value.ValueKind == JsonValueKind.String)
 						{
-							keys.Add(new IssueKey(property.Value.GetString()!, IssueKeyType.None));
+							data._location = property.Value.GetString()!;
+						}
+						else if (property.NameEquals("locNamespace") && property.Value.ValueKind == JsonValueKind.String)
+						{
+							data._localizationNamespace = property.Value.GetString()!;
+						}
+						// "locID" is deprecated and will be removed soon
+						else if ((property.NameEquals("locKey") || property.NameEquals("locID")) && property.Value.ValueKind == JsonValueKind.String)
+						{
+							data._localizationKey = property.Value.GetString()!;
+						}
+						else if (property.NameEquals("conflictFile") && property.Value.ValueKind == JsonValueKind.String)
+						{
+							data._conflictFile = property.Value.GetString()!;
+						}
+						else if (property.NameEquals("conflictLine") && property.Value.ValueKind == JsonValueKind.Number)
+						{
+							data._conflictLine = property.Value.GetInt32();
+						}
+						else if (property.NameEquals("conflictLocation") && property.Value.ValueKind == JsonValueKind.String)
+						{
+							data._conflictLocation = property.Value.GetString()!;
+						}
+						else if (property.NameEquals("conflictLocNamespace") && property.Value.ValueKind == JsonValueKind.String)
+						{
+							data._conflictLocalizationNamespace = property.Value.GetString()!;
+						}
+						else if (property.NameEquals("conflictLocKey") && property.Value.ValueKind == JsonValueKind.String)
+						{
+							data._conflictLocalizationKey = property.Value.GetString()!;
 						}
 					}
 				}
 			}
+
+			ConvertLocationToFileAndLineNumber(data);
+			RemoveFilesExtension(data);
+
+			return data;
 		}
 
 		/// <summary>
-		/// Extracts the localization data from the properties of an IssueEvent
+		/// Location data in LocalizationData may contain file and line information. We can try to parse it.
 		/// </summary>
-		/// <param name="issueEvent">The event data</param>
 		/// <param name="data">The localization data</param>
-		private static void GetLocalizationData(IssueEvent issueEvent, LocalizationData data)
+		private static void ConvertLocationToFileAndLineNumber(LocalizationData data)
 		{
-			foreach (JsonLogEvent line in issueEvent.Lines)
+			if (!System.String.IsNullOrEmpty(data._location))
 			{
-				JsonDocument document = JsonDocument.Parse(line.Data);
+				// If the TextLocation contains a '(', we assume it is a format: "/Path/To/File.cpp(10)" where 10 is a line number that can be parsed
+				int startLineIndex = data._location.IndexOf('(', System.StringComparison.Ordinal);
+				// If the TextLocation contains a '.', we assume it is a format: "/Path/To/Filename.*" where the filename can be retrieved
+				int startExtensionIndex = data._location.IndexOf('.', System.StringComparison.Ordinal);
+				if (startLineIndex > 0 && data._line == -1)
+				{
+					data._line = Int32.Parse(data._location.Substring(startLineIndex + 1, data._location.Length - startLineIndex - 2));
+				}
+				if (startExtensionIndex > 0 && System.String.IsNullOrEmpty(data._file))
+				{
+					data._file = data._location.Substring(0, startExtensionIndex);
+				}
+			}
+			if (!System.String.IsNullOrEmpty(data._conflictLocation))
+			{
+				// If the TextLocation contains a '(', we assume it is a format: "/Path/To/File.cpp(10)" where 10 is a line number that can be parsed
+				int startLineIndex = data._conflictLocation.IndexOf('(', System.StringComparison.Ordinal);
+				// If the TextLocation contains a '.', we assume it is a format: "/Path/To/Filename.*" where the filename can be retrieved
+				int startExtensionIndex = data._conflictLocation.IndexOf('.', System.StringComparison.Ordinal);
+				if (startLineIndex > 0 && data._conflictLine == -1)
+				{
+					data._conflictLine = Int32.Parse(data._conflictLocation.Substring(startLineIndex + 1, data._conflictLocation.Length - startLineIndex - 2));
+				}
+				if (startExtensionIndex > 0 && System.String.IsNullOrEmpty(data._conflictFile))
+				{
+					data._conflictFile = data._conflictLocation.Substring(0, startExtensionIndex);
+				}
+			}
+		}
 
-				JsonElement properties;
-				if (document.RootElement.TryGetProperty("properties", out properties) && properties.ValueKind == JsonValueKind.Object)
+		/// <summary>
+		/// LocalizationData file paths may contain extensions, remove them
+		/// </summary>
+		/// <param name="data">The localization data</param>
+		private static void RemoveFilesExtension(LocalizationData data)
+		{
+			if (!System.String.IsNullOrEmpty(data._file))
+			{
+				// If the File path contains a '.', we assume it is a format: "/Path/To/Filename.extension" where the filename can be retrieved
+				int startExtensionIndex = data._file.IndexOf('.', System.StringComparison.Ordinal);
+				if (startExtensionIndex > 0)
 				{
-					foreach (JsonProperty property in properties.EnumerateObject())
-					{
-						if (property.NameEquals("localizationTarget") && property.Value.ValueKind == JsonValueKind.String)
-						{
-							data._localizationTarget = property.Value.GetString()!;
-						}
-						else if ((property.NameEquals("localizationID") || property.NameEquals("locID")) && property.Value.ValueKind == JsonValueKind.String)
-						{
-							data._localizationID = property.Value.GetString()!;
-						}
-					}
+					data._file = data._file.Substring(0, startExtensionIndex);
+				}
+			}
+			if (!System.String.IsNullOrEmpty(data._conflictFile))
+			{
+				// If the ConflictFile contains a '.', we assume it is a format: "/Path/To/Filename.extension" where the filename can be retrieved
+				int startExtensionIndex = data._conflictFile.IndexOf('.', System.StringComparison.Ordinal);
+				if (startExtensionIndex > 0)
+				{
+					data._conflictFile = data._conflictFile.Substring(0, startExtensionIndex);
 				}
 			}
 		}
@@ -111,20 +234,36 @@ namespace EpicGames.Horde.Issues.Handlers
 				EventId eventId = issueEvent.EventId.Value;
 				if (IsMatchingEventId(eventId))
 				{
-					LocalizationData data = new LocalizationData();
-					GetLocalizationData(issueEvent, data);
+					LocalizationData data = GetLocalizationData(issueEvent);
+
+					IssueEventGroup issueGroup;
+					if(data.HasAnyFile())
+					{
+						// Create a new issue group (might be grouped later based on keys)
+						string localizationIssueGroup = "Localization";
+						string localizationSummaryTemplate = (!System.String.IsNullOrEmpty(_context.NodeName) ? "{Meta:Node}: " : "") + "Localization {Severity} in {Files}";
+						issueGroup = new IssueEventGroup(localizationIssueGroup, localizationSummaryTemplate, IssueChangeFilter.All);
+						if (!System.String.IsNullOrEmpty(_context.NodeName))
+						{
+							issueGroup.Metadata.Add("Node", _context.NodeName);
+						}
+					}
+					else
+					{
+						// If we have no idea on the file an issue is associated with, then we won't be able to narrow down a suspect.
+						// If we can't narrow down a suspect. There is no gain to split the issues amongst different issueGroup.
+						// Actually, there is only a risk of creating thousands of issue group if localization is in a weird/broken state.
+						issueGroup = _issuesWithoutClearFiles;
+					}
 
-					string localizationIssueGroup = (System.String.IsNullOrEmpty(data._localizationTarget) ? "Localization" : "Localization Target " + data._localizationTarget);
-					IssueEventGroup issue = new IssueEventGroup(localizationIssueGroup, "Localization {Severity} in {Files}", IssueChangeFilter.All);
-					
-					issue.Events.Add(issueEvent);
+					// Add current event to the chosen issueGroup
+					issueGroup.Events.Add(issueEvent);
 
-					issue.Keys.AddSourceFiles(issueEvent, data._localizationID);
-					AddLocalizationIssueKeys(issue.Keys, issueEvent, data);
+					AddLocalizationIssueKeys(issueGroup.Keys, data);
 
-					if (issue.Keys.Count > 0)
+					if(issueGroup.Keys.Count > 0 && !_issues.Contains(issueGroup))
 					{
-						_issues.Add(issue);
+						_issues.Add(issueGroup);
 					}
 
 					return true;
