Horde: Add KnownLogEvents.Horde_BuildHealth_Ignore to suppress build health issue generation regardless of log level

diff --git a/Engine/Source/Programs/Horde/Drivers/JobDriver/Execution/JobExecutor.cs b/Engine/Source/Programs/Horde/Drivers/JobDriver/Execution/JobExecutor.cs
--- a/Engine/Source/Programs/Horde/Drivers/JobDriver/Execution/JobExecutor.cs
+++ b/Engine/Source/Programs/Horde/Drivers/JobDriver/Execution/JobExecutor.cs
@@ -519,7 +519,7 @@ namespace JobDriver.Execution
 
 				if (stepCancellationToken.IsCancellationRequested && IsCancellationException(ex))
 				{
-					stepLogger.LogError("The step was intentionally cancelled");
+					stepLogger.LogError(KnownLogEvents.Horde_BuildHealth_Ignore, "The step was intentionally cancelled");
 					return (JobStepOutcome.Failure, JobStepState.Aborted);
 				}
 
diff --git a/Engine/Source/Programs/Shared/EpicGames.Core/KnownLogEvents.cs b/Engine/Source/Programs/Shared/EpicGames.Core/KnownLogEvents.cs
--- a/Engine/Source/Programs/Shared/EpicGames.Core/KnownLogEvents.cs
+++ b/Engine/Source/Programs/Shared/EpicGames.Core/KnownLogEvents.cs
@@ -381,6 +381,16 @@ namespace EpicGames.Core
 		/// </summary>
 		public static EventId Horde_BlobRead { get; } = new EventId(1002);
 
+		/// <summary>
+		/// Event used to supporess the logging from the Horde build health service, regardless of log level
+		/// </summary>
+		public static EventId Horde_BuildHealth_Ignore { get; } = new EventId(1003);
+
+		/// <summary>
+		/// Maximum code for Horde events. Add new events in the 1000-1099 range.
+		/// </summary>
+		public static EventId Horde_Max { get; } = new EventId(1099);
+
 		/// <summary>
 		/// An event while running zen cli
 		/// </summary>
diff --git a/Engine/Source/Programs/Shared/EpicGames.Horde/Issues/Handlers/IgnoredIssueHandler.cs b/Engine/Source/Programs/Shared/EpicGames.Horde/Issues/Handlers/IgnoredIssueHandler.cs
new file mode 100644
--- /dev/null
+++ b/Engine/Source/Programs/Shared/EpicGames.Horde/Issues/Handlers/IgnoredIssueHandler.cs
@@ -0,0 +1,56 @@
+// Copyright Epic Games, Inc. All Rights Reserved.
+
+using System.Collections.Generic;
+using EpicGames.Core;
+using Microsoft.Extensions.Logging;
+
+namespace EpicGames.Horde.Issues.Handlers
+{
+	/// <summary>
+	/// Instance of a particular systemic error
+	/// </summary>
+	[IssueHandler]
+	public class IgnoredIssueHandler : IssueHandler
+	{
+		/// <summary>
+		///  Known systemic errors
+		/// </summary>
+		static readonly HashSet<EventId> s_handledEvents = new HashSet<EventId> { KnownLogEvents.Horde_BuildHealth_Ignore };
+
+		/// <inheritdoc/>
+		public override int Priority => 100;
+
+		/// <summary>
+		/// Determines if the given event id matches
+		/// </summary>
+		/// <param name="eventId">The event id to compare</param>
+		/// <returns>True if the given event id matches</returns>
+		public static bool IsMatchingEventId(EventId eventId)
+		{
+			return s_handledEvents.Contains(eventId);
+		}
+
+		readonly List<IssueEventGroup> _issues = new List<IssueEventGroup>();
+
+		/// <summary>
+		/// Constructor
+		/// </summary>
+		public IgnoredIssueHandler() {}
+
+		/// <inheritdoc/>
+		public override bool HandleEvent(IssueEvent logEvent)
+		{
+			if (logEvent.EventId != null)
+			{
+				if (IsMatchingEventId(logEvent.EventId.Value))
+				{
+					return true;
+				}
+			}
+			return false;
+		}
+
+		/// <inheritdoc/>
+		public override IEnumerable<IssueEventGroup> GetIssues() => _issues;
+	}
+}
