Silence stderr from 'security' calls on Mac as these interfere with OidcToken --ResultToConsole.

diff --git a/Engine/Source/Programs/Shared/EpicGames.OIDC/MacOSTokenStore.cs b/Engine/Source/Programs/Shared/EpicGames.OIDC/MacOSTokenStore.cs
--- a/Engine/Source/Programs/Shared/EpicGames.OIDC/MacOSTokenStore.cs
+++ b/Engine/Source/Programs/Shared/EpicGames.OIDC/MacOSTokenStore.cs
@@ -53,6 +53,7 @@ namespace EpicGames.OIDC
 			{
 				FileName = "security",
 				RedirectStandardOutput = true,
+				RedirectStandardError = true
 			};
 			startInfo.ArgumentList.Add(verb);
 			foreach (string s in args)
