UnrealBuildTool.Tests: Fix test when default compiler is changed

#rnx

diff --git a/Engine/Source/Programs/Shared/UnrealBuildTool.Tests/CppCompileWarningsTests.cs b/Engine/Source/Programs/Shared/UnrealBuildTool.Tests/CppCompileWarningsTests.cs
--- a/Engine/Source/Programs/Shared/UnrealBuildTool.Tests/CppCompileWarningsTests.cs
+++ b/Engine/Source/Programs/Shared/UnrealBuildTool.Tests/CppCompileWarningsTests.cs
@@ -294,6 +294,7 @@ namespace UnrealBuildTool.Tests
 				// Generate rules assembly & subsequent targets
 				RulesAssembly generatedTestRulesAssembly = RulesCompiler.CreateTargetRulesAssembly(testUproject, "TestStubEditor", false, false, false, null, false, new TestLogger());
 				TargetRules testTargetRules = generatedTestRulesAssembly.CreateTargetRules("TestStubEditor", UnrealTargetPlatform.Win64, UnrealTargetConfiguration.Development, null, testUproject, null, new TestLogger());
+				testTargetRules.WindowsPlatform.Compiler = WindowsCompiler.VisualStudio2022;
 				CppCompileWarnings baseWarnings = new CppCompileWarnings(testTargetRules, null);
 
 				IEnumerable<string> args = null!;
