Remove '.' characters from bucket names when querying Unreal Cloud snapshots.

[RN] none

diff --git a/Engine/Source/Programs/Shared/EpicGames.ProjectStore/ProjectStoreTypes.cs b/Engine/Source/Programs/Shared/EpicGames.ProjectStore/ProjectStoreTypes.cs
--- a/Engine/Source/Programs/Shared/EpicGames.ProjectStore/ProjectStoreTypes.cs
+++ b/Engine/Source/Programs/Shared/EpicGames.ProjectStore/ProjectStoreTypes.cs
@@ -888,7 +888,7 @@ namespace EpicGames.ProjectStore
 		private static string GetBucket(string project, string branch, string platform, string runtime)
 		{
 			project = project.ToLowerInvariant();
-			branch = StringId.Sanitize(branch).ToString();
+			branch = StringId.Sanitize(branch).ToString().Replace(".", "-", StringComparison.Ordinal);
 			platform = platform.ToLowerInvariant();
 			runtime = runtime.ToLowerInvariant();
 
