Python Callback attached to Dynamic Delegate

Calling a python delegate currently crashes in `FFrame::StepExplicitProperty` because the `FFrame::OutParms` can be null even though the property is flagged as `CPF_OutParm`.
The `OutParms` can be null here because `UObject::ProcessEvent` does not create these `OutParms` if the UFunction is not flagged with `FUNC_HasOutParms`.
Looking at the UHT code, this flag is currently only set for functions in a class, and with the delegates being defined outside of a class, this flag is never set.
This leads to the delegates not having the flag `FUNC_HasOutParms`, so the `FFrame::OutParms` is still null when reaching `FFrame::StepExplicitProperty`, but as ref and const-ref parameters would still be tagged with `CPF_OutParm`, a crash would follow.

This CL ensures that the flag `FUNC_HasOutParms` is set for all `UFunctions`, which would now correctly create the `FFrame::OutParms` and stop the crash.

Test code that was originally crashing and that now runs:
```cpp
// PythonTestLibrary.h
#pragma once

#include "CoreMinimal.h"
#include "PythonTestLibrary.generated.h"

USTRUCT(BlueprintType)
struct FPythonTestStruct
{
GENERATED_BODY()

UPROPERTY(BlueprintReadWrite, EditAnywhere, Category = "PythonTestLibrary")
int Int = 0;
UPROPERTY(BlueprintReadWrite, EditAnywhere, Category = "PythonTestLibrary")
bool bIsEven = false;
};

DECLARE_DYNAMIC_DELEGATE_OneParam(FOnPythonTestLibrary_ConstRef, const FPythonTestStruct&, Struct);
DECLARE_DYNAMIC_DELEGATE_RetVal_OneParam(const FPythonTestStruct&, FOnPythonTestLibrary_ConstRef_Return, const FPythonTestStruct&, Struct);
DECLARE_DYNAMIC_DELEGATE_RetVal_OneParam(FPythonTestStruct, FOnPythonTestLibrary_Ref_Return, FPythonTestStruct&, Struct);
```

```python
import unreal

RESET_VALUE = 2
FUNCTION_NEW_RETURN_VALUE = 333

def same_values(a, b):
return a.int == b.int and a.is_even == b.is_even
def reset(val = RESET_VALUE): # resets the struct with the wrong value of is_even
struct = unreal.PythonTestStruct()
struct.int = val
struct.is_even = False
return struct
def SetIsEvenAndReturn(struct):
struct.is_even = (struct.int % 2 == 0)
return struct
def CreateNewStructAndReturnTuple():
struct = reset(RESET_VALUE)
struct.is_even = (struct.int % 2 == 0)
new = reset(FUNCTION_NEW_RETURN_VALUE)
return (new, struct)

### Const-ref return version
### DECLARE_DYNAMIC_DELEGATE_RetVal_OneParam(const FPythonTestStruct&, FOnPythonTestLibrary_ConstRef_Return, const FPythonTestStruct&, Struct);
unreal.log_warning(" === unreal.OnPythonTestLibrary_ConstRef_Return")
struct = reset(RESET_VALUE)
print(" - input:", struct)
delegate = unreal.OnPythonTestLibrary_ConstRef_Return()
delegate.bind_callable(SetIsEvenAndReturn)
ret = delegate.execute(struct)
print("original:", struct)
print("return  :", ret)
assert same_values(struct, unreal.PythonTestStruct(RESET_VALUE, False)) # untouched original
assert same_values(ret, unreal.PythonTestStruct(RESET_VALUE, True)) # the function should have set is_even to true

### Ref return version
### DECLARE_DYNAMIC_DELEGATE_RetVal_OneParam(FPythonTestStruct, FOnPythonTestLibrary_Ref_Return, FPythonTestStruct&, Struct);
unreal.log_warning(" === unreal.OnPythonTestLibrary_Ref_Return")
delegate = unreal.OnPythonTestLibrary_Ref_Return()
delegate.bind_callable(CreateNewStructAndReturnTuple) # return 2 values, one for the return and one for the ref param
ret, ref = delegate.execute()
print("return  :", ret)
print("refparam:", ref)
assert same_values(ret, unreal.PythonTestStruct(FUNCTION_NEW_RETURN_VALUE, False))
assert same_values(ref, unreal.PythonTestStruct(RESET_VALUE, True))  # the function should have set is_even to true
```

diff --git a/Engine/Source/Programs/Shared/EpicGames.UHT/Types/UhtFunction.cs b/Engine/Source/Programs/Shared/EpicGames.UHT/Types/UhtFunction.cs
--- a/Engine/Source/Programs/Shared/EpicGames.UHT/Types/UhtFunction.cs
+++ b/Engine/Source/Programs/Shared/EpicGames.UHT/Types/UhtFunction.cs
@@ -597,27 +597,27 @@ namespace EpicGames.UHT.Types
 						}
 					}
 
-					// The following code is only performed on functions in a class.
-					if (Outer is UhtClass)
+					foreach (UhtType type in Children)
 					{
-						foreach (UhtType type in Children)
+						if (type is not UhtProperty property)
 						{
-							if (type is UhtProperty property)
+							continue;
+						}
+						
+						if (property.PropertyFlags.HasExactFlags(EPropertyFlags.OutParm | EPropertyFlags.ReturnParm, EPropertyFlags.OutParm))
+						{
+							FunctionFlags |= EFunctionFlags.HasOutParms;
+						}
+						// The following code is only performed on functions in a class.
+						if (property is UhtStructProperty structProperty && Outer is UhtClass)
+						{
+							if (structProperty.ScriptStruct.HasDefaults)
 							{
-								if (property.PropertyFlags.HasExactFlags(EPropertyFlags.OutParm | EPropertyFlags.ReturnParm, EPropertyFlags.OutParm))
-								{
-									FunctionFlags |= EFunctionFlags.HasOutParms;
-								}
-								if (property is UhtStructProperty structProperty)
-								{
-									if (structProperty.ScriptStruct.HasDefaults)
-									{
-										FunctionFlags |= EFunctionFlags.HasDefaults;
-									}
-								}
+								FunctionFlags |= EFunctionFlags.HasDefaults;
 							}
 						}
 					}
+
 					break;
 			}
 			return results;
