[CoreUObject/UHT]
* Reduced number of exported functions by making UObject::staticClass() call generated function which in turn is the only exported method. GetPrivateStaticClass is not exported anymore

diff --git a/Engine/Source/Programs/Shared/EpicGames.UHT/Exporters/CodeGen/UhtHeaderCodeGeneratorCppFile.cs b/Engine/Source/Programs/Shared/EpicGames.UHT/Exporters/CodeGen/UhtHeaderCodeGeneratorCppFile.cs
--- a/Engine/Source/Programs/Shared/EpicGames.UHT/Exporters/CodeGen/UhtHeaderCodeGeneratorCppFile.cs
+++ b/Engine/Source/Programs/Shared/EpicGames.UHT/Exporters/CodeGen/UhtHeaderCodeGeneratorCppFile.cs
@@ -2096,7 +2096,7 @@ namespace EpicGames.UHT.Exporters.CodeGen
 			{
 				builder.Append("UClass* ").Append(GetSingletonName(classObj, false)).Append("()\r\n");
 				builder.Append("{\r\n");
-				builder.Append("\treturn ").Append(classObj.SourceName).Append("::StaticClass();\r\n");
+				builder.Append("\treturn ").Append(classObj.SourceName).Append("::GetPrivateStaticClass();\r\n");
 				builder.Append("}\r\n");
 			}
 
diff --git a/Engine/Source/Programs/Shared/EpicGames.UHT/Exporters/CodeGen/UhtHeaderCodeGeneratorHFile.cs b/Engine/Source/Programs/Shared/EpicGames.UHT/Exporters/CodeGen/UhtHeaderCodeGeneratorHFile.cs
--- a/Engine/Source/Programs/Shared/EpicGames.UHT/Exporters/CodeGen/UhtHeaderCodeGeneratorHFile.cs
+++ b/Engine/Source/Programs/Shared/EpicGames.UHT/Exporters/CodeGen/UhtHeaderCodeGeneratorHFile.cs
@@ -828,6 +828,9 @@ namespace EpicGames.UHT.Exporters.CodeGen
 				AppendSerializer(builder, classObj, api, UhtSerializerArchiveType.StructuredArchiveRecord, "DECLARE_FSTRUCTUREDARCHIVE_SERIALIZER");
 			}
 
+			// Declaration of globla function fetching UClass of type
+			builder.Append(classObj.Module.Api).Append("UClass* ").Append(ObjectInfos[classObj.ObjectTypeIndex].UnregisteredSingletonName).Append("();\r\n\r\n");
+
 			if (classObj.ClassFlags.HasAnyFlags(EClassFlags.Interface))
 			{
 				if (usesLegacy)
@@ -1538,20 +1541,26 @@ namespace EpicGames.UHT.Exporters.CodeGen
 		/// <returns>Output builder</returns>
 		private StringBuilder AppendCommonGeneratedBody(StringBuilder builder, UhtClass classObj, string api)
 		{
+			UhtCodeGenerator.ObjectInfo info = ObjectInfos[classObj.ObjectTypeIndex];
+
 			// Export the class's native function registration.
 			builder.Append("private: \\\r\n");
 			builder.Append("\tstatic void StaticRegisterNatives").Append(classObj.SourceName).Append("(); \\\r\n");
 			if (!classObj.ClassExportFlags.HasAnyFlags(UhtClassExportFlags.NoExport))
 			{
-				builder.Append("\tfriend struct ").Append(ObjectInfos[classObj.ObjectTypeIndex].RegisteredSingletonName).Append("_Statics; \\\r\n");
+				builder.Append("\tfriend struct ").Append(info.RegisteredSingletonName).Append("_Statics; \\\r\n");
 			}
+			
+			builder.Append('\t').Append("static UClass* GetPrivateStaticClass(); \\\r\n");
+			builder.Append("\tfriend ").Append(classObj.Module.Api).Append("UClass* ").Append(info.UnregisteredSingletonName).Append("(); \\\r\n");
+
 			builder.Append("public: \\\r\n");
 
 			UhtClass? superClass = classObj.SuperClass;
 			bool castedClass = classObj.ClassCastFlags.HasAnyFlags(EClassCastFlags.AllFlags) && superClass != null && classObj.ClassCastFlags != superClass.ClassCastFlags;
 
 			builder
-				.Append("\tDECLARE_CLASS(")
+				.Append("\tDECLARE_CLASS2(")
 				.Append(classObj)
 				.Append(", ")
 				.Append(superClass != null ? superClass.SourceName : "None")
@@ -1570,7 +1579,7 @@ namespace EpicGames.UHT.Exporters.CodeGen
 			{
 				builder.Append("CASTCLASS_None");
 			}
-			builder.Append(", TEXT(\"").Append(classObj.Package.SourceName).Append("\"), ").Append(api[0..^1]).Append(") \\\r\n");
+			builder.Append(", TEXT(\"").Append(classObj.Package.SourceName).Append("\"), ").Append(info.UnregisteredSingletonName).Append(") \\\r\n");
 
 			builder.Append("\tDECLARE_SERIALIZER(").Append(classObj.SourceName).Append(") \\\r\n");
 
diff --git a/Engine/Source/Runtime/CoreUObject/Public/UObject/Object.h b/Engine/Source/Runtime/CoreUObject/Public/UObject/Object.h
--- a/Engine/Source/Runtime/CoreUObject/Public/UObject/Object.h
+++ b/Engine/Source/Runtime/CoreUObject/Public/UObject/Object.h
@@ -1795,6 +1795,9 @@ private:
 		checkf(false, TEXT("REPLICATED_BASE_CLASS macro is not implemented in %s. Something went very wrong."), *GetFullName());
 	}
 
+	// Here for UHT to be happy
+	friend COREUOBJECT_API UClass* Z_Construct_UClass_UObject_NoRegister();
+
 public:
 
 	/** Should only ever be used by internal systems. */
diff --git a/Engine/Source/Runtime/CoreUObject/Public/UObject/ObjectMacros.h b/Engine/Source/Runtime/CoreUObject/Public/UObject/ObjectMacros.h
--- a/Engine/Source/Runtime/CoreUObject/Public/UObject/ObjectMacros.h
+++ b/Engine/Source/Runtime/CoreUObject/Public/UObject/ObjectMacros.h
@@ -1806,11 +1806,14 @@ namespace UM
 	Class declaration macros.
 -----------------------------------------------------------------------------*/
 
-#define DECLARE_CLASS( TClass, TSuperClass, TStaticFlags, TStaticCastFlags, TPackage, TRequiredAPI  ) \
+#define DECLARE_CLASS( TClass, TSuperClass, TStaticFlags, TStaticCastFlags, TPackage, TRequiredAPI ) \
+	TRequiredAPI static UClass* GetPrivateStaticClass(); \
+	DECLARE_CLASS2(TClass, TSuperClass, TStaticFlags, TStaticCastFlags, TPackage, GetPrivateStaticClass )
+
+#define DECLARE_CLASS2( TClass, TSuperClass, TStaticFlags, TStaticCastFlags, TPackage, TPrivateAccessor ) \
 private: \
     TClass& operator=(TClass&&);   \
     TClass& operator=(const TClass&);   \
-	TRequiredAPI static UClass* GetPrivateStaticClass(); \
 public: \
 	/** Bitwise union of #EClassFlags pertaining to this class.*/ \
 	static constexpr EClassFlags StaticClassFlags=EClassFlags(TStaticFlags); \
@@ -1821,7 +1824,7 @@ public: \
 	/** Returns a UClass object representing this class at runtime */ \
 	inline static UClass* StaticClass() \
 	{ \
-		return GetPrivateStaticClass(); \
+		return TPrivateAccessor(); \
 	} \
 	/** Returns the package this class belongs in */ \
 	inline static const TCHAR* StaticPackage() \
