[UBT] Adding support for the Activate flag in PluginReferenceDescriptors.
Fixing issue where the Activate flag would be deleted from plugin descriptors with some UAT programs

#tests running NewPluginLocalizationTarget command on plugin with game feature plugin activations, running  UBT tests

diff --git a/Engine/Source/Programs/Shared/UnrealBuildTool.Tests/PluginDescriptorTests.cs b/Engine/Source/Programs/Shared/UnrealBuildTool.Tests/PluginDescriptorTests.cs
--- a/Engine/Source/Programs/Shared/UnrealBuildTool.Tests/PluginDescriptorTests.cs
+++ b/Engine/Source/Programs/Shared/UnrealBuildTool.Tests/PluginDescriptorTests.cs
@@ -104,6 +104,21 @@ namespace UnrealBuildTool.Tests
 			Assert.AreEqual(json, saveContents);
 		}
 
+		[TestMethod]
+		public void Save2PluginGameFeatureActivation()
+		{
+			string json = CreateUpluginWithGameFeatureActivation();
+			string filename = "GameFeatureActivation-Save.uplugin";
+			FileReference fileToWrite = CreateTempUPluginFile(filename, json);
+			PluginDescriptor descriptor = PluginDescriptor.FromFile(fileToWrite);
+			string fileToSave = GetAbsolutePathToTempFile("GameFeatureActivation-Save.uplugin");
+			descriptor.Save2(fileToSave);
+			string saveContents = File.ReadAllText(fileToSave);
+
+			// Nothing should change between the format of the json and what was written to disk 
+			Assert.AreEqual(json, saveContents);
+		}
+
 		[TestMethod]
 		public void Save2CompareWithSave()
 		{
@@ -254,6 +269,63 @@ namespace UnrealBuildTool.Tests
 			return FixVerbatimStringIdentation(json);
 		}
 
+		private static string CreateUpluginWithGameFeatureActivation()
+		{
+			string json = @"
+				{
+					""FileVersion"": 3,
+					""Version"": 1,
+					""VersionName"": """",
+					""FriendlyName"": """",
+					""Description"": """",
+					""Category"": """",
+					""CreatedBy"": """",
+					""CreatedByURL"": """",
+					""DocsURL"": """",
+					""MarketplaceURL"": """",
+					""SupportURL"": """",
+					""CanContainContent"": true,
+					""ExplicitlyLoaded"": true,
+					""EnabledByDefault"": false,
+					""EditorCustomVirtualPath"": ""MyPlugin"",
+					""Modules"": [
+						{
+							""Name"": ""Module1"",
+							""Type"": ""Runtime"",
+							""LoadingPhase"": ""Default""
+						},
+						{
+							""Name"": ""Module2"",
+							""Type"": ""ClientOnly"",
+							""LoadingPhase"": ""Default""
+						}
+					],
+					""Plugins"": [
+						{
+							""Name"": ""Plugin1"",
+							""Enabled"": true
+						},
+						{
+							""Name"": ""Plugin2"",
+							""Enabled"": true,
+							""Activate"": true
+						},
+						{
+							""Name"": ""Plugin3"",
+							""Enabled"": true,
+							""Optional"": true,
+							""Activate"": true
+						},
+						{
+							""Name"": ""Plugin4"",
+							""Enabled"": true,
+							""Optional"": true
+						}
+					]
+				}";
+			return FixVerbatimStringIdentation(json);
+		}
+
 		private static string FixVerbatimStringIdentation(string verbatimString)
 		{
 			if (String.IsNullOrEmpty(verbatimString))
diff --git a/Engine/Source/Programs/UnrealBuildTool/System/PluginReferenceDescriptor.cs b/Engine/Source/Programs/UnrealBuildTool/System/PluginReferenceDescriptor.cs
--- a/Engine/Source/Programs/UnrealBuildTool/System/PluginReferenceDescriptor.cs
+++ b/Engine/Source/Programs/UnrealBuildTool/System/PluginReferenceDescriptor.cs
@@ -29,6 +29,11 @@ namespace UnrealBuildTool
 		/// </summary>
 		public bool bOptional;
 
+		/// <summary>
+		/// Whether this plugin should be activated, if it is a Game Feature Plugin
+		/// </summary>
+		public bool bActivate;
+
 		/// <summary>
 		/// Description of the plugin for users that do not have it installed.
 		/// </summary>
@@ -110,6 +115,10 @@ namespace UnrealBuildTool
 			{
 				Writer.WriteValue("Optional", bOptional);
 			}
+			if (bEnabled && bActivate)
+			{
+				Writer.WriteValue("Activate", bActivate);
+			}
 			if (!String.IsNullOrEmpty(Description))
 			{
 				Writer.WriteValue("Description", Description);
@@ -166,6 +175,10 @@ namespace UnrealBuildTool
 			{
 				PluginReferenceObject.AddOrSetFieldValue("Optional", bOptional);
 			}
+			if (bEnabled && bActivate)
+			{
+				PluginReferenceObject.AddOrSetFieldValue("Activate", bActivate);
+			}
 			if (!String.IsNullOrEmpty(Description))
 			{
 				PluginReferenceObject.AddOrSetFieldValue("Description", Description);
@@ -275,6 +288,7 @@ namespace UnrealBuildTool
 		{
 			PluginReferenceDescriptor Descriptor = new PluginReferenceDescriptor(RawObject.GetStringField("Name"), null, RawObject.GetBoolField("Enabled"));
 			RawObject.TryGetBoolField("Optional", out Descriptor.bOptional);
+			RawObject.TryGetBoolField("Activate", out Descriptor.bActivate);
 			RawObject.TryGetStringField("Description", out Descriptor.Description);
 			RawObject.TryGetStringField("MarketplaceURL", out Descriptor.MarketplaceURL);
 
