Avoid use of Unreal.RootDirectory and Unreal.EngineDirectory when importing snapshots via UGS.

[RN] none

diff --git a/Engine/Source/Programs/Shared/EpicGames.ProjectStore/ProjectStoreTypes.cs b/Engine/Source/Programs/Shared/EpicGames.ProjectStore/ProjectStoreTypes.cs
--- a/Engine/Source/Programs/Shared/EpicGames.ProjectStore/ProjectStoreTypes.cs
+++ b/Engine/Source/Programs/Shared/EpicGames.ProjectStore/ProjectStoreTypes.cs
@@ -156,6 +156,18 @@ namespace EpicGames.ProjectStore
 		/// <param name="asyncImport"></param>
 		public abstract void Import(FileReference projectFile, DirectoryReference platformCookedDirectory, bool forceImport, bool asyncImport);
 
+		/// <summary>
+		/// Performs snapshot import for the discovered descriptors into a specified project and cooked output directory,
+		/// using a specified Unreal root/engine directory.
+		/// </summary>
+		/// <param name="projectFile"></param>
+		/// <param name="rootDirectory"></param>
+		/// <param name="engineDirectory"></param>
+		/// <param name="platformCookedDirectory"></param>
+		/// <param name="forceImport"></param>
+		/// <param name="asyncImport"></param>
+		public abstract void Import(FileReference projectFile, DirectoryReference rootDirectory, DirectoryReference engineDirectory, DirectoryReference platformCookedDirectory, bool forceImport, bool asyncImport);
+
 		/// <summary>
 		/// Returns the changelist associated with this snapshot.
 		/// </summary>
@@ -209,27 +221,38 @@ namespace EpicGames.ProjectStore
 		}
 
 		/// <summary>
-		/// Creates a new project in zenserver if it doesn't exist.
+		/// Creates a new project in zenserver if it doesn't exist, using the default Unreal root and engine directories.
 		/// </summary>
 		/// <param name="projectFile">Full path to a .uproject file to create a zenserver project for.</param>
 		protected void CreateProject(FileReference projectFile)
+		{
+			CreateProject(projectFile, Unreal.RootDirectory, Unreal.EngineDirectory);
+		}
+
+		/// <summary>
+		/// Creates a new project in zenserver if it doesn't exist.
+		/// </summary>
+		/// <param name="projectFile">Full path to a .uproject file to create a zenserver project for.</param>
+		/// <param name="rootDirectory">Full path to the Unreal root directory.</param>
+		/// <param name="engineDirectory">Full path to Engine/.</param>
+		protected void CreateProject(FileReference projectFile, DirectoryReference rootDirectory, DirectoryReference engineDirectory)
 		{
 			string projectId = GetProjectID(projectFile);
 			StringBuilder projectCreateCommandLine = new StringBuilder();
-			projectCreateCommandLine.AppendFormat("project-create --force-update {0} {1} {2} {3} {4}", projectId, Unreal.RootDirectory.FullName, Unreal.EngineDirectory.FullName, projectFile.Directory.FullName, projectFile.FullName);
+			projectCreateCommandLine.AppendFormat("project-create --force-update {0} {1} {2} {3} {4}", projectId, rootDirectory.FullName, engineDirectory.FullName, projectFile.Directory.FullName, projectFile.FullName);
 
-			_runApp(GetZenExecutable().FullName, projectCreateCommandLine.ToString());
+			_runApp(GetZenExecutable(engineDirectory).FullName, projectCreateCommandLine.ToString());
 		}
 
 		/// <summary>
 		/// Returns a FileReference to the published Zen executable in Engine/Binaries/...
 		/// </summary>
-		protected static FileReference GetZenExecutable()
+		protected static FileReference GetZenExecutable(DirectoryReference engineDirectory)
 		{
 			string platform = (OperatingSystem.IsWindows() ? "Win64"
 				: (OperatingSystem.IsMacOS() ? "Mac" : "Linux"));
 
-			return new FileReference(Path.Combine(Unreal.EngineDirectory.FullName, "Binaries", platform, String.Format("zen{0}", RuntimePlatform.ExeExtension)));
+			return new FileReference(Path.Combine(engineDirectory.FullName, "Binaries", platform, String.Format("zen{0}", RuntimePlatform.ExeExtension)));
 		}
 
 		/// <summary>
@@ -404,11 +427,17 @@ namespace EpicGames.ProjectStore
 			LoadSnapshotDescriptors(data);
 		}
 
-		/// <inheritdoc />
+		/// <inheritdoc/>
 		public override void Import(FileReference projectFile, DirectoryReference platformCookedDirectory, bool forceImport, bool asyncImport)
+		{
+			Import(projectFile, Unreal.RootDirectory, Unreal.EngineDirectory, platformCookedDirectory, forceImport, asyncImport);
+		}
+
+		/// <inheritdoc />
+		public override void Import(FileReference projectFile, DirectoryReference rootDirectory, DirectoryReference engineDirectory, DirectoryReference platformCookedDirectory, bool forceImport, bool asyncImport)
 		{
 			string projectId = GetProjectID(projectFile);
-			CreateProject(projectFile);
+			CreateProject(projectFile, rootDirectory, engineDirectory);
 
 			if (_snapshotDescriptors == null || _snapshotDescriptors.Snapshots == null)
 			{
@@ -460,7 +489,7 @@ namespace EpicGames.ProjectStore
 						break;
 				}
 
-				FileReference zenExe = GetZenExecutable();
+				FileReference zenExe = GetZenExecutable(engineDirectory);
 				_runApp(zenExe.FullName, oplogCreateCommandLine.ToString());
 				_runApp(zenExe.FullName, oplogImportCommandLine.ToString());
 			}
@@ -505,9 +534,15 @@ namespace EpicGames.ProjectStore
 
 		/// <inheritdoc />
 		public override void Import(FileReference projectFile, DirectoryReference platformCookedDirectory, bool forceImport, bool asyncImport)
+		{
+			Import(projectFile, Unreal.RootDirectory, Unreal.EngineDirectory, platformCookedDirectory, forceImport, asyncImport);
+		}
+
+		/// <inheritdoc />
+		public override void Import(FileReference projectFile, DirectoryReference rootDirectory, DirectoryReference engineDirectory, DirectoryReference platformCookedDirectory, bool forceImport, bool asyncImport)
 		{
 			string projectId = GetProjectID(projectFile);
-			CreateProject(projectFile);
+			CreateProject(projectFile, rootDirectory, engineDirectory);
 
 			platformCookedDirectory = new DirectoryReference(platformCookedDirectory.FullName.Replace("{Platform}", _buildData.Metadata.CookPlatform, StringComparison.Ordinal));
 			FileReference projectStoreFile = WriteProjectStoreFile(projectId, _buildData.Metadata.CookPlatform, platformCookedDirectory);
@@ -536,7 +571,7 @@ namespace EpicGames.ProjectStore
 			oplogImportCommandLine.AppendFormat(" --builds {0}", _host);
 			oplogImportCommandLine.AppendFormat(" --builds-id {0}", _buildData.BuildId);
 
-			FileReference zenExe = GetZenExecutable();
+			FileReference zenExe = GetZenExecutable(engineDirectory);
 			_runApp(zenExe.FullName, oplogCreateCommandLine.ToString());
 			_runApp(zenExe.FullName, oplogImportCommandLine.ToString());
 		}
diff --git a/Engine/Source/Programs/UnrealGameSync/UnrealGameSyncShared/WorkspaceUpdate.cs b/Engine/Source/Programs/UnrealGameSync/UnrealGameSyncShared/WorkspaceUpdate.cs
--- a/Engine/Source/Programs/UnrealGameSync/UnrealGameSyncShared/WorkspaceUpdate.cs
+++ b/Engine/Source/Programs/UnrealGameSync/UnrealGameSyncShared/WorkspaceUpdate.cs
@@ -1430,7 +1430,7 @@ namespace UnrealGameSync
 									logger.LogWarning("No exact match found for target changelist {}. Closest changelist found was {}. You may require an incremental cook for this snapshot to function properly.", Context.ChangeNumber, build.GetChangelist());
 								}
 
-								build.Import(project.LocalFileName, cookedPlatformDir, false, false);
+								build.Import(project.LocalFileName, project.LocalRootPath, project.EngineDir, cookedPlatformDir, false, false);
 							}
 
 							await Task.WhenAll(snapshotTasks);
