UHT: add support for engine version comparisons as preprocessor conditions

diff --git a/Engine/Source/Programs/Shared/EpicGames.UHT/Parsers/UhtHeaderFileParser.cs b/Engine/Source/Programs/Shared/EpicGames.UHT/Parsers/UhtHeaderFileParser.cs
--- a/Engine/Source/Programs/Shared/EpicGames.UHT/Parsers/UhtHeaderFileParser.cs
+++ b/Engine/Source/Programs/Shared/EpicGames.UHT/Parsers/UhtHeaderFileParser.cs
@@ -994,6 +994,10 @@ namespace EpicGames.UHT.Parsers
 					{
 						return notPresent ? UhtCompilerDirective.NotCPPBlock : UhtCompilerDirective.CPPBlock;
 					}
+					else if (define.ValueStartsWith("UE_VERSION"))
+					{
+						return ParserEngineVersionConditional(define, notPresent, lineTokenReader);
+					}
 					break;
 
 				case UhtTokenType.EndOfFile:
@@ -1006,6 +1010,52 @@ namespace EpicGames.UHT.Parsers
 			return UhtCompilerDirective.Unrecognized;
 		}
 
+		private UhtCompilerDirective ParserEngineVersionConditional(UhtToken define, bool notPresent, UhtTokenBufferReader lineTokenReader)
+		{
+			lineTokenReader.Require('(');
+			UhtToken majorToken = lineTokenReader.GetToken();
+			lineTokenReader.Require(',');
+			UhtToken minorToken = lineTokenReader.GetToken();
+			lineTokenReader.Require(',');
+			UhtToken patchToken = lineTokenReader.GetToken();
+			lineTokenReader.Require(')');
+
+			EngineVersion engineVersion = HeaderFile.Session.EngineVersion;
+
+			int major, minor, patch;
+			majorToken.GetConstInt(out major);
+			minorToken.GetConstInt(out minor);
+			patchToken.GetConstInt(out patch);
+			EngineVersion defineVersion = new EngineVersion(major, minor, patch);
+
+			bool? conditionPasses = null;
+			if (define.IsValue("UE_VERSION_NEWER_THAN_OR_EQUAL"))
+			{
+				conditionPasses = engineVersion.CompareTo(defineVersion) >= 0;
+			}
+			else if (define.IsValue("UE_VERSION_NEWER_THAN"))
+			{
+				conditionPasses = engineVersion.CompareTo(defineVersion) > 0;
+			}
+			else if (define.IsValue("UE_VERSION_OLDER_THAN"))
+			{
+				conditionPasses = engineVersion.CompareTo(defineVersion) < 0;
+			}
+			
+			if (conditionPasses.HasValue)
+			{
+				if (notPresent)
+				{
+					conditionPasses = !conditionPasses.Value;
+				}
+				return conditionPasses.Value ? UhtCompilerDirective.OneBlock : UhtCompilerDirective.ZeroBlock;
+			}
+			else
+			{
+				throw new UhtException(String.Format("Unknown engine version preprocessor conditional: {0}", define.Value));
+			}
+		}
+
 		/// <summary>
 		/// Add a new compiler directive to the stack
 		/// </summary>
diff --git a/Engine/Source/Programs/Shared/EpicGames.UHT/Utils/UhtSession.cs b/Engine/Source/Programs/Shared/EpicGames.UHT/Utils/UhtSession.cs
--- a/Engine/Source/Programs/Shared/EpicGames.UHT/Utils/UhtSession.cs
+++ b/Engine/Source/Programs/Shared/EpicGames.UHT/Utils/UhtSession.cs
@@ -7,6 +7,7 @@ using System.IO.Enumeration;
 using System.Linq;
 using System.Text;
 using System.Text.Json;
+using System.Text.RegularExpressions;
 using System.Threading;
 using System.Threading.Tasks;
 using EpicGames.Core;
@@ -47,6 +48,71 @@ namespace EpicGames.UHT.Utils
 		public StringView Data { get; set; }
 	}
 
+	/// <summary>
+	/// A structure that represents an engine version.
+	/// </summary>
+	public struct EngineVersion
+	{
+		/// <summary>
+		/// The major engine version.
+		/// </summary>
+		public int MajorVersion { get; set; } = 0;
+		/// <summary>
+		/// The minor engine version.
+		/// </summary>
+		public int MinorVersion { get; set; } = 0;
+		/// <summary>
+		/// The patch engine version.
+		/// </summary>
+		public int PatchVersion { get; set; } = 0;
+
+		/// <summary>
+		/// Build a default engine version (0.0.0).
+		/// </summary>
+		public EngineVersion() { }
+
+		/// <summary>
+		/// Build a new engine version with the given numbers.
+		/// </summary>
+		public EngineVersion(int major, int minor, int patch)
+		{
+			MajorVersion = major;
+			MinorVersion = minor;
+			PatchVersion = patch;
+		}
+
+		/// <summary>
+		/// Compares this engine version to another one.
+		/// </summary>
+		/// <param name="other">The other engine version</param>
+		/// <returns>Returns -1 if this version is less than the other, 1 if it is greater than the other, and 0 if they're equal</returns>
+		public int CompareTo(EngineVersion other)
+		{
+			if (MajorVersion != other.MajorVersion)
+			{
+				return MajorVersion.CompareTo(other.MajorVersion);
+			}
+			if (MinorVersion != other.MinorVersion)
+			{
+				return MinorVersion.CompareTo(other.MinorVersion);
+			}
+			if (PatchVersion != other.PatchVersion)
+			{
+				return PatchVersion.CompareTo(other.PatchVersion);
+			}
+			return 0;
+		}
+
+		/// <summary>
+		/// Generates a string representation of the engine version.
+		/// </summary>
+		/// <returns></returns>
+		public override string ToString()
+		{
+			return String.Format("{0}.{1}.{2}", MajorVersion, MinorVersion, PatchVersion);
+		}
+	}
+
 	/// <summary>
 	/// Implementation of the export factory
 	/// </summary>
@@ -567,7 +633,7 @@ namespace EpicGames.UHT.Utils
 	/// <summary>
 	/// Session object that represents a UHT run
 	/// </summary>
-	public class UhtSession : IUhtMessageSite, IUhtMessageSession
+	public partial class UhtSession : IUhtMessageSite, IUhtMessageSession
 	{
 
 		/// <summary>
@@ -746,6 +812,11 @@ namespace EpicGames.UHT.Utils
 		/// </summary>
 		public IReadOnlyDictionary<string, UhtHeaderFile> HeaderFileDictionary => _headerFileDictionary;
 
+		/// <summary>
+		/// Gets the version of the engine this session belongs to.
+		/// </summary>
+		public EngineVersion EngineVersion => _engineVersion.Value;
+
 		/// <summary>
 		/// After headers are parsed, returns the UObject class.
 		/// </summary>
@@ -840,6 +911,7 @@ namespace EpicGames.UHT.Utils
 		private readonly List<UhtHeaderFile> _headerFiles = new();
 		private readonly List<UhtHeaderFile> _sortedHeaderFiles = new();
 		private readonly Dictionary<string, UhtHeaderFile> _headerFileDictionary = new(StringComparer.OrdinalIgnoreCase);
+		private readonly Lazy<EngineVersion> _engineVersion;
 		private long _errorCount = 0;
 		private long _warningCount = 0;
 		private readonly List<UhtMessage> _messages = new();
@@ -892,6 +964,8 @@ namespace EpicGames.UHT.Utils
 		public UhtSession(ILogger logger)
 		{
 			Logger = logger;
+
+			_engineVersion = new Lazy<EngineVersion>(() => LoadEngineVersion(), true);
 		}
 
 		/// <summary>
@@ -1918,6 +1992,69 @@ namespace EpicGames.UHT.Utils
 			}
 		}
 
+		[GeneratedRegex(@"^\#define\s+ENGINE_(?<Name>[A-Z]+)_VERSION\s+(?<Number>\d+)\s*$")]
+		private static partial Regex GetEngineVersionFragmentRegex();
+
+		private EngineVersion LoadEngineVersion()
+		{
+			if (EngineDirectory == null)
+			{
+				throw new UhtException("Can't parse engine version without an engine directory");
+			}
+
+			string versionHeader = Path.Combine(EngineDirectory, @"Source\Runtime\Launch\Resources\Version.h");
+			Log.Logger.LogDebug("Reading engine version from file: {VersionHeader}", versionHeader);
+			string[] lines;
+			try
+			{
+				lines = File.ReadAllLines(versionHeader);
+			}
+			catch (Exception ex)
+			{
+				throw new UhtException(String.Format("Error reading engine version header '{0}': {1}", versionHeader, ex));
+			}
+
+			EngineVersion loadedVersion = new EngineVersion();
+			bool gotMajor = false, gotMinor = false, gotPatch = false;
+			Regex versionRegex = GetEngineVersionFragmentRegex();
+			foreach (string line in lines)
+			{
+				Match match = versionRegex.Match(line);
+				if (match.Success)
+				{
+					int matchNumber = Int32.Parse(match.Groups["Number"].Value);
+					switch (match.Groups["Name"].Value)
+					{
+						case "MAJOR":
+							loadedVersion.MajorVersion = matchNumber;
+							gotMajor = true;
+							break;
+						case "MINOR":
+							loadedVersion.MinorVersion = matchNumber;
+							gotMinor = true;
+							break;
+						case "PATCH":
+							loadedVersion.PatchVersion = matchNumber;
+							gotPatch = true;
+							break;
+					}
+
+					if (gotMajor && gotMinor && gotPatch)
+					{
+						break;
+					}
+				}
+			}
+
+			if (!gotMajor || !gotMinor || !gotPatch)
+			{
+				throw new UhtException(String.Format("Got incomplete engine version {0}", loadedVersion));
+			}
+
+			Log.Logger.LogDebug("Got engine version {EngineVersion}", loadedVersion);
+			return loadedVersion;
+		}
+
 		#region Run steps
 		private void StepReadManifestFile(string manifestFilePath)
 		{
diff --git a/Engine/Source/Programs/UnrealBuildTool/System/ExternalExecution.cs b/Engine/Source/Programs/UnrealBuildTool/System/ExternalExecution.cs
--- a/Engine/Source/Programs/UnrealBuildTool/System/ExternalExecution.cs
+++ b/Engine/Source/Programs/UnrealBuildTool/System/ExternalExecution.cs
@@ -764,6 +764,29 @@ namespace UnrealBuildTool
 			return DateTime.MaxValue;
 		}
 
+		/// <summary>
+		/// Gets the timestamp of the Version.h header file. Throws an exception if not found.
+		/// </summary>
+		/// <returns>Last write time of Version.h</returns>
+		private static DateTime GetEngineVersionHeaderTimestampUtc()
+		{
+			// In Installed Builds, we don't check the timestamps on engine headers.  Default to a very old date.
+			if (Unreal.IsEngineInstalled())
+			{
+				return DateTime.MinValue;
+			}
+
+			string EngineVersionHeaderPath = Path.Combine(Unreal.EngineSourceDirectory.FullName, "Runtime", "Launch", "Resources", "Version.h");
+			FileInfo EngineVersionHeaderFileInfo = new(EngineVersionHeaderPath);
+			if (EngineVersionHeaderFileInfo.Exists)
+			{
+				return EngineVersionHeaderFileInfo.LastWriteTimeUtc;
+			}
+
+			// The Version.h header doesn't exist somehow (maybe it was moved).
+			throw new BuildException("Could not find engine version header, expected: {0}", EngineVersionHeaderPath);
+		}
+
 		/// <summary>
 		/// Checks the class header files and determines if generated UObject code files are out of date in comparison.
 		/// </summary>
@@ -793,6 +816,10 @@ namespace UnrealBuildTool
 				}
 			}
 
+			// Also get the Version.h timestamp.  If it is newer than the last UHT run, it could affect preprocessor conditionals
+			// that check for the engine version
+			DateTime EngineVersionHeaderTimestampUtc = GetEngineVersionHeaderTimestampUtc();
+
 			foreach (UHTModuleInfo Module in UObjectModules)
 			{
 				// If we're using a precompiled engine, skip checking timestamps for modules that are under the engine directory
@@ -838,6 +865,12 @@ namespace UnrealBuildTool
 					Logger.LogDebug("UnrealHeaderTool needs to run because CoreUObject timestamp ({Time}) is newer than timestamp for module {ModuleName} ({ModuleTime})", CoreGeneratedTimestampUtc.Value.ToLocalTime(), Module.ModuleName, SavedTimestampUtc.ToLocalTime());
 					return true;
 				}
+				if (EngineVersionHeaderTimestampUtc > SavedTimestampUtc)
+				{
+					// Generated code is older than engine version header.  Out of date!
+					Logger.LogDebug("UnrealHeaderTool needs to run because engine version header timestamp ({Time}) is newer than timestamp for module {ModuleName} ({ModuleTime})", EngineVersionHeaderTimestampUtc.ToLocalTime(), Module.ModuleName, SavedTimestampUtc.ToLocalTime());
+					return true;
+				}
 
 				// Has the .build.cs file changed since we last generated headers successfully?
 				FileInfo ModuleRulesFile = new FileInfo(Module.ModuleRulesFile.FullName);
diff --git a/Engine/Source/Runtime/Core/Public/Misc/EngineVersionComparison.h b/Engine/Source/Runtime/Core/Public/Misc/EngineVersionComparison.h
--- a/Engine/Source/Runtime/Core/Public/Misc/EngineVersionComparison.h
+++ b/Engine/Source/Runtime/Core/Public/Misc/EngineVersionComparison.h
@@ -8,6 +8,12 @@
 #define UE_GREATER_SORT(Value, ValueToBeGreaterThan, TieBreaker) \
 	(((Value) > (ValueToBeGreaterThan)) || (((Value) == (ValueToBeGreaterThan)) && (TieBreaker)))
 
+// Version comparison macro that is defined to true if the UE version is newer or equal than MajorVer.MinorVer.PatchVer 
+// and false otherwise
+// (a typical use is for backward compatible code)
+#define UE_VERSION_NEWER_THAN_OR_EQUAL(MajorVersion, MinorVersion, PatchVersion)\
+	UE_GREATER_SORT(ENGINE_MAJOR_VERSION, MajorVersion, UE_GREATER_SORT(ENGINE_MINOR_VERSION, MinorVersion, UE_GREATER_SORT(ENGINE_PATCH_VERSION, PatchVersion, true)))
+
 // Version comparison macro that is defined to true if the UE version is newer than MajorVer.MinorVer.PatchVer and false otherwise
 // (a typical use is to alert integrators to revisit this code when upgrading to a new engine version)
 #define UE_VERSION_NEWER_THAN(MajorVersion, MinorVersion, PatchVersion) \
