Renaming methods used in generating verse names

#rnx

diff --git a/Engine/Source/Programs/Shared/EpicGames.UHT/Exporters/CodeGen/UhtHeaderCodeGeneratorCppFile.cs b/Engine/Source/Programs/Shared/EpicGames.UHT/Exporters/CodeGen/UhtHeaderCodeGeneratorCppFile.cs
--- a/Engine/Source/Programs/Shared/EpicGames.UHT/Exporters/CodeGen/UhtHeaderCodeGeneratorCppFile.cs
+++ b/Engine/Source/Programs/Shared/EpicGames.UHT/Exporters/CodeGen/UhtHeaderCodeGeneratorCppFile.cs
@@ -454,7 +454,7 @@ namespace EpicGames.UHT.Exporters.CodeGen
 
 			if (enumObj.IsVerseField)
 			{
-				builder.Append("V_DEFINE_IMPORTED_ENUM(").Append(Module.Api.TrimEnd()).Append(", \"").AppendVerseUEPackageName(enumObj).Append("\", \"").Append(enumObj.VersePath).Append("\", ").Append(enumObj.SourceName).Append(");\r\n");
+				builder.Append("V_DEFINE_IMPORTED_ENUM(").Append(Module.Api.TrimEnd()).Append(", \"").AppendVerseUEPackageName(enumObj).Append("\", \"").Append(enumObj.GetVerseScopeAndName(UhtVerseNameMode.Default)).Append("\", ").Append(enumObj.SourceName).Append(");\r\n");
 			}
 
 			// Everything from this point on will be part of the definition hash
@@ -515,7 +515,7 @@ namespace EpicGames.UHT.Exporters.CodeGen
 				if (enumObj.IsVerseField)
 				{
 					builder.Append("\t},\r\n");
-					builder.Append("\t\"").AppendVersePath(enumObj, UhtVerseFullNameMode.Qualified).Append("\"\r\n");
+					builder.Append("\t\"").AppendVerseScopeAndName(enumObj, UhtVerseNameMode.Qualified).Append("\"\r\n");
 				}
 				builder.Append("};\r\n");
 			}
@@ -740,8 +740,8 @@ namespace EpicGames.UHT.Exporters.CodeGen
 					builder.Append("\t},\r\n");
 
 					// NOTE: This is temporary while the new VM doesn't match the value in the old VM
-					string defaultName = scriptStruct.GetVerseFullName(UhtVerseFullNameMode.Default);
-					string qualifiedName = scriptStruct.GetVerseFullName(UhtVerseFullNameMode.Qualified);
+					string defaultName = scriptStruct.GetVerseScopeAndName(UhtVerseNameMode.Default);
+					string qualifiedName = scriptStruct.GetVerseScopeAndName(UhtVerseNameMode.Qualified);
 					if (defaultName.Equals(qualifiedName, StringComparison.Ordinal))
 					{
 						builder.Append($"\t\"{qualifiedName}\"\r\n");
@@ -961,7 +961,7 @@ namespace EpicGames.UHT.Exporters.CodeGen
 				{
 					AppendPropertiesDefs(builder, context, properties, 0);
 
-					string functionEngineName = isVerseNativeCallable ? function.EncodedVersePath : function.EngineName;
+					string functionEngineName = isVerseNativeCallable ? function.GetEncodedVerseScopeAndName(UhtVerseNameMode.Default) : function.EngineName;
 					builder
 						.Append("const UECodeGen_Private::").AppendFunctionParamsType(function).Append(' ')
 						.Append(staticsName).Append("::FuncParams = { { ")
@@ -1013,9 +1013,9 @@ namespace EpicGames.UHT.Exporters.CodeGen
 							{
 								builder.Append("\r\n");
 								builder.Append("#if WITH_VERSE_VM\r\n");
-								builder.AppendUTF8LiteralString(function.VersePath).Append(",\r\n");
+								builder.AppendUTF8LiteralString(function.GetVerseScopeAndName(UhtVerseNameMode.Default)).Append(",\r\n");
 								builder.Append("#else\r\n");
-								builder.AppendUTF8LiteralString(function.EncodedVersePath).Append(",\r\n");
+								builder.AppendUTF8LiteralString(function.GetEncodedVerseScopeAndName(UhtVerseNameMode.Default)).Append(",\r\n");
 								builder.Append("#endif\r\n");
 							}
 							break;
@@ -1309,14 +1309,14 @@ namespace EpicGames.UHT.Exporters.CodeGen
 			builder.Append("\t\tusing CalleeType = ").AppendVerseNativeCallableTypeDef(function).Append(";\r\n");
 
 			// Callee paths
-			string versePath = function.VersePath;
+			string verseScopeAndName = function.GetVerseScopeAndName(UhtVerseNameMode.Default);
 			if (isVerseVM)
 			{
-				builder.Append("\t\tconst char* CalleePath = \"").Append(versePath).Append("\";\r\n");
+				builder.Append("\t\tconst char* CalleePath = \"").Append(verseScopeAndName).Append("\";\r\n");
 			}
 			else
 			{
-				builder.Append("\t\tconst TCHAR* CalleePath = TEXT(\"").AppendEncodedVerseName(versePath).Append("\");\r\n");
+				builder.Append("\t\tconst TCHAR* CalleePath = TEXT(\"").AppendEncodedVerseName(verseScopeAndName).Append("\");\r\n");
 			}
 
 			// Callee declaration
@@ -2383,7 +2383,7 @@ namespace EpicGames.UHT.Exporters.CodeGen
 								blockEmitter.Set(function.DefineScope);
 								builder
 									.Append("\t\t{ ")
-									.AppendUTF8LiteralString(function.EncodedVersePath)
+									.AppendUTF8LiteralString(function.GetEncodedVerseScopeAndName(UhtVerseNameMode.Default))
 									.Append(", &")
 									.AppendClassSourceNameOrInterfaceName(classObj)
 									.Append("::_exec_")
@@ -2414,7 +2414,7 @@ namespace EpicGames.UHT.Exporters.CodeGen
 							blockEmitter.Set(function.DefineScope);
 							builder
 								.Append("\t\t{ ")
-								.AppendUTF8LiteralString(function.VersePath)
+								.AppendUTF8LiteralString(function.GetVerseScopeAndName(UhtVerseNameMode.Default))
 								.Append(", &")
 								.AppendClassSourceNameOrInterfaceName(classObj)
 								.Append("::_exec_")
diff --git a/Engine/Source/Programs/Shared/EpicGames.UHT/Exporters/CodeGen/UhtHeaderCodeGeneratorHFile.cs b/Engine/Source/Programs/Shared/EpicGames.UHT/Exporters/CodeGen/UhtHeaderCodeGeneratorHFile.cs
--- a/Engine/Source/Programs/Shared/EpicGames.UHT/Exporters/CodeGen/UhtHeaderCodeGeneratorHFile.cs
+++ b/Engine/Source/Programs/Shared/EpicGames.UHT/Exporters/CodeGen/UhtHeaderCodeGeneratorHFile.cs
@@ -662,7 +662,7 @@ namespace EpicGames.UHT.Exporters.CodeGen
 						.Append("V_DECLARE_STRUCT_RUNTIME_TRAITS(, ")
 						.Append(scriptStruct.SourceName).Append(", \"")
 						.AppendVerseUEPackageName(scriptStruct).Append("\", \"")
-						.AppendVersePath(scriptStruct, UhtVerseFullNameMode.Default).Append("\")\r\n");
+						.AppendVerseScopeAndName(scriptStruct, UhtVerseNameMode.Default).Append("\")\r\n");
 				}
 
 				Session.Tables!.CodeGeneratorInjectorTable.Inject(builder, scriptStruct, UhtCodeGeneratorInjectionLocation.Header);
diff --git a/Engine/Source/Programs/Shared/EpicGames.UHT/Types/Properties/UhtEnumProperty.cs b/Engine/Source/Programs/Shared/EpicGames.UHT/Types/Properties/UhtEnumProperty.cs
--- a/Engine/Source/Programs/Shared/EpicGames.UHT/Types/Properties/UhtEnumProperty.cs
+++ b/Engine/Source/Programs/Shared/EpicGames.UHT/Types/Properties/UhtEnumProperty.cs
@@ -142,7 +142,7 @@ namespace EpicGames.UHT.Types
 					break;
 
 				case UhtPropertyTextType.VerseMangledType:
-					builder.AppendVersePath(enumObj, UhtVerseFullNameMode.Default);
+					builder.AppendVerseScopeAndName(enumObj, UhtVerseNameMode.Default);
 					break;
 
 				default:
diff --git a/Engine/Source/Programs/Shared/EpicGames.UHT/Types/Properties/UhtObjectPropertyBase.cs b/Engine/Source/Programs/Shared/EpicGames.UHT/Types/Properties/UhtObjectPropertyBase.cs
--- a/Engine/Source/Programs/Shared/EpicGames.UHT/Types/Properties/UhtObjectPropertyBase.cs
+++ b/Engine/Source/Programs/Shared/EpicGames.UHT/Types/Properties/UhtObjectPropertyBase.cs
@@ -581,7 +581,7 @@ namespace EpicGames.UHT.Types
 				case UhtObjectCppForm.TNonNullPtr:
 				case UhtObjectCppForm.TInterfaceInstance:
 				case UhtObjectCppForm.TSubclassOf:
-					builder.AppendVersePath(type, UhtVerseFullNameMode.Default);
+					builder.AppendVerseScopeAndName(type, UhtVerseNameMode.Default);
 					break;
 
 				default:
diff --git a/Engine/Source/Programs/Shared/EpicGames.UHT/Types/Properties/UhtStructProperty.cs b/Engine/Source/Programs/Shared/EpicGames.UHT/Types/Properties/UhtStructProperty.cs
--- a/Engine/Source/Programs/Shared/EpicGames.UHT/Types/Properties/UhtStructProperty.cs
+++ b/Engine/Source/Programs/Shared/EpicGames.UHT/Types/Properties/UhtStructProperty.cs
@@ -125,7 +125,7 @@ namespace EpicGames.UHT.Types
 			switch (textType)
 			{
 				case UhtPropertyTextType.VerseMangledType:
-					builder.AppendVersePath(ScriptStruct, UhtVerseFullNameMode.Default);
+					builder.AppendVerseScopeAndName(ScriptStruct, UhtVerseNameMode.Default);
 					break;
 
 				default:
diff --git a/Engine/Source/Programs/Shared/EpicGames.UHT/Types/UhtField.cs b/Engine/Source/Programs/Shared/EpicGames.UHT/Types/UhtField.cs
--- a/Engine/Source/Programs/Shared/EpicGames.UHT/Types/UhtField.cs
+++ b/Engine/Source/Programs/Shared/EpicGames.UHT/Types/UhtField.cs
@@ -72,7 +72,7 @@ namespace EpicGames.UHT.Types
 	/// <summary>
 	/// Controls how the verse name is appended to the builder
 	/// </summary>
-	public enum UhtVerseFullNameMode
+	public enum UhtVerseNameMode
 	{
 		/// <summary>
 		/// Default mode
@@ -114,24 +114,6 @@ namespace EpicGames.UHT.Types
 		/// </summary>
 		public bool IsVerseField => VerseName != null;
 
-		/// <summary>
-		/// Return the complete verse name
-		/// </summary>
-		public string VersePath => GetVerseFullName(UhtVerseFullNameMode.Default);
-
-		/// <summary>
-		/// Return the function name encoded for BPVM
-		/// </summary>
-		public string EncodedVersePath
-		{
-			get
-			{
-				StringBuilder builder = new();
-				builder.AppendEncodedVerseName(VersePath);
-				return builder.ToString();
-			}
-		}
-
 		/// <summary>
 		/// Construct a new field
 		/// </summary>
@@ -143,36 +125,62 @@ namespace EpicGames.UHT.Types
 		}
 
 		/// <summary>
-		/// Return the full name of the verse type
+		/// Return the scope and name of the verse type (scope:)name
+		/// </summary>
+		/// <param name="mode">Controls how the name is generated</param>
+		/// <returns></returns>
+		public string GetVerseScopeAndName(UhtVerseNameMode mode)
+		{
+			StringBuilder builder = new();
+			AppendVerseScopeAndName(builder, mode);
+			return builder.ToString();
+		}
+
+		/// <summary>
+		/// Return the encoded verse scope and name of the field
 		/// </summary>
 		/// <param name="mode">Controls how the name is generated</param>
 		/// <returns></returns>
-		public string GetVerseFullName(UhtVerseFullNameMode mode)
+		public string GetEncodedVerseScopeAndName(UhtVerseNameMode mode)
 		{
 			StringBuilder builder = new();
-			AppendVerseFullName(builder, mode);
+			AppendVerseScopeAndName(builder, mode);
+			string name = builder.ToString();
+			builder.Clear().AppendEncodedVerseName(name);
 			return builder.ToString();
 		}
 
 		/// <summary>
-		/// Append the verse path to the given builder
+		/// Append the scope and name of the verse type (scope:)name
 		/// </summary>
 		/// <param name="builder">Destination builder</param>
 		/// <param name="mode">Controls how the name is generated</param>
-		public virtual void AppendVerseFullName(StringBuilder builder, UhtVerseFullNameMode mode)
+		public virtual void AppendVerseScopeAndName(StringBuilder builder, UhtVerseNameMode mode)
 		{
-			AppendVersePathPart(builder, mode);
+			builder.Append('(');
+			AppendVerseScope(builder, mode);
+			builder.Append(":)");
 			AppendVerseName(builder, mode);
 		}
 
 		/// <summary>
-		/// Append the verse path to the given builder
+		/// Append the verse scope part to the given builder
 		/// </summary>
 		/// <param name="builder">Destination builder</param>
 		/// <param name="mode">Controls how the name is generated</param>
-		/// <param name="isTopLevel">If true, this is the type that initiated the verse path</param>
+		public void AppendVerseScope(StringBuilder builder, UhtVerseNameMode mode)
+		{
+			AppendVerseScopeInternal(builder, mode, true);
+		}
+
+		/// <summary>
+		/// Append the verse scope to the given builder
+		/// </summary>
+		/// <param name="builder">Destination builder</param>
+		/// <param name="mode">Controls how the name is generated</param>
+		/// <param name="isTopLevel">If true, this is the type that initiated the verse scope</param>
 		/// <exception cref="UhtException"></exception>
-		private void AppendVersePathInternal(StringBuilder builder, UhtVerseFullNameMode mode, bool isTopLevel)
+		private void AppendVerseScopeInternal(StringBuilder builder, UhtVerseNameMode mode, bool isTopLevel)
 		{
 			if (!IsVerseField)
 			{
@@ -180,7 +188,7 @@ namespace EpicGames.UHT.Types
 			}
 			if (Outer is UhtField outerField)
 			{
-				outerField.AppendVersePathInternal(builder, mode, false);
+				outerField.AppendVerseScopeInternal(builder, mode, false);
 			}
 			else
 			{
@@ -190,19 +198,7 @@ namespace EpicGames.UHT.Types
 					builder.Append('/').Append(VerseModule);
 				}
 			}
-			AppendVersePath(builder, mode, isTopLevel);
-		}
-
-		/// <summary>
-		/// Append the verse path part to the given builder
-		/// </summary>
-		/// <param name="builder">Destination builder</param>
-		/// <param name="mode">Controls how the name is generated</param>
-		protected virtual void AppendVersePathPart(StringBuilder builder, UhtVerseFullNameMode mode)
-		{
-			builder.Append('(');
-			AppendVersePathInternal(builder, mode, true);
-			builder.Append(":)");
+			AppendVerseMyScope(builder, mode, isTopLevel);
 		}
 
 		/// <summary>
@@ -211,7 +207,7 @@ namespace EpicGames.UHT.Types
 		/// <param name="builder"></param>
 		/// <param name="mode">Controls how the name is generated</param>
 		/// <param name="isTopLevel">If true, this is the type that initiated the verse path</param>
-		protected virtual void AppendVersePath(StringBuilder builder, UhtVerseFullNameMode mode, bool isTopLevel)
+		protected virtual void AppendVerseMyScope(StringBuilder builder, UhtVerseNameMode mode, bool isTopLevel)
 		{
 			if (!isTopLevel)
 			{
@@ -224,7 +220,7 @@ namespace EpicGames.UHT.Types
 		/// </summary>
 		/// <param name="builder"></param>
 		/// <param name="mode">Controls how the name is generated</param>
-		protected virtual void AppendVerseName(StringBuilder builder, UhtVerseFullNameMode mode)
+		protected virtual void AppendVerseName(StringBuilder builder, UhtVerseNameMode mode)
 		{
 			builder.Append(VerseName);
 		}
@@ -269,87 +265,15 @@ namespace EpicGames.UHT.Types
 		}
 
 		/// <summary>
-		/// Append the full verse name of the field
+		/// Append the scope and name of the verse type (scope:)name
 		/// </summary>
 		/// <param name="builder">Destination builder</param>
 		/// <param name="fieldObj">Field to serialize</param>
 		/// <param name="mode">Controls the type of string generated</param>
 		/// <returns>Builder</returns>
-		public static StringBuilder AppendVersePath(this StringBuilder builder, UhtField fieldObj, UhtVerseFullNameMode mode)
-		{
-			fieldObj.AppendVerseFullName(builder, mode);
-			return builder;
-		}
-
-		/// <summary>
-		/// Given a verse name, encode it so it can be represented as an FName
-		/// </summary>
-		/// <param name="builder">Destination builder</param>
-		/// <param name="name">Name to encode</param>
-		/// <returns>Builder</returns>
-		public static StringBuilder AppendEncodedVerseName(this StringBuilder builder, ReadOnlySpan<char> name)
+		public static StringBuilder AppendVerseScopeAndName(this StringBuilder builder, UhtField fieldObj, UhtVerseNameMode mode)
 		{
-			bool isFirstChar = true;
-			while (!name.IsEmpty)
-			{
-				char c = name[0];
-				name = name[1..];
-
-				if ((c >= 'a' && c <= 'z')
-					|| (c >= 'A' && c <= 'Z')
-					|| (c >= '0' && c <= '9' && !isFirstChar))
-				{
-					builder.Append(c);
-				}
-				else if (c == '[' && !name.IsEmpty && name[0] == ']')
-				{
-					name = name[1..];
-					builder.Append("_K");
-				}
-				else if (c == '-' && !name.IsEmpty && name[0] == '>')
-				{
-					name = name[1..];
-					builder.Append("_T");
-				}
-				else if (c == '_')
-				{
-					builder.Append("__");
-				}
-				else if (c == '(')
-				{
-					builder.Append("_L");
-				}
-				else if (c == ',')
-				{
-					builder.Append("_M");
-				}
-				else if (c == ':')
-				{
-					builder.Append("_N");
-				}
-				else if (c == '^')
-				{
-					builder.Append("_P");
-				}
-				else if (c == '?')
-				{
-					builder.Append("_Q");
-				}
-				else if (c == ')')
-				{
-					builder.Append("_R");
-				}
-				else if (c == '\'')
-				{
-					builder.Append("_U");
-				}
-				else
-				{
-					builder.Append('_').AppendFormat(CultureInfo.InvariantCulture, "{0:x2}", (uint)c);
-				}
-
-				isFirstChar = false;
-			}
+			fieldObj.AppendVerseScopeAndName(builder, mode);
 			return builder;
 		}
 	}
diff --git a/Engine/Source/Programs/Shared/EpicGames.UHT/Types/UhtFunction.cs b/Engine/Source/Programs/Shared/EpicGames.UHT/Types/UhtFunction.cs
--- a/Engine/Source/Programs/Shared/EpicGames.UHT/Types/UhtFunction.cs
+++ b/Engine/Source/Programs/Shared/EpicGames.UHT/Types/UhtFunction.cs
@@ -434,17 +434,21 @@ namespace EpicGames.UHT.Types
 		}
 
 		/// <inheritdoc/>
-		protected override void AppendVersePathPart(StringBuilder builder, UhtVerseFullNameMode mode)
+		public override void AppendVerseScopeAndName(StringBuilder builder, UhtVerseNameMode mode)
 		{
 			// suspends methods don't get the path part added to the name
 			if (!FunctionExportFlags.HasAnyFlags(UhtFunctionExportFlags.VerseSuspends))
 			{
-				base.AppendVersePathPart(builder, mode);
+				base.AppendVerseScopeAndName(builder, mode);
+			}
+			else
+			{
+				AppendVerseName(builder, mode);
 			}
 		}
 
 		/// <inheritdoc/>
-		protected override void AppendVerseName(StringBuilder builder, UhtVerseFullNameMode mode)
+		protected override void AppendVerseName(StringBuilder builder, UhtVerseNameMode mode)
 		{
 			UhtProperty? returnProperty = ReturnProperty;
 			ReadOnlyMemory<UhtType> parameterProperties = ParameterProperties;
@@ -864,7 +868,7 @@ namespace EpicGames.UHT.Types
 							this.LogError($"Verse function '{SourceName}' must either be marked as NativeCallable or VerseCallable");
 						}
 
-						if (VersePath == SourceName)
+						if (GetVerseScopeAndName(UhtVerseNameMode.Default) == SourceName)
 						{
 							this.LogError($"Verse function '{SourceName}' has the same name as the Verse generated name.  Change the source name and then specify the name in the UFUNCTION attributes");
 						}
diff --git a/Engine/Source/Programs/Shared/EpicGames.UHT/Types/UhtScriptStruct.cs b/Engine/Source/Programs/Shared/EpicGames.UHT/Types/UhtScriptStruct.cs
--- a/Engine/Source/Programs/Shared/EpicGames.UHT/Types/UhtScriptStruct.cs
+++ b/Engine/Source/Programs/Shared/EpicGames.UHT/Types/UhtScriptStruct.cs
@@ -952,19 +952,19 @@ namespace EpicGames.UHT.Types
 
 
 		/// <inheritdoc/>
-		protected override void AppendVersePath(StringBuilder builder, UhtVerseFullNameMode mode, bool isTopLevel)
+		protected override void AppendVerseMyScope(StringBuilder builder, UhtVerseNameMode mode, bool isTopLevel)
 		{
-			if (!isTopLevel || (mode == UhtVerseFullNameMode.Qualified && ScriptStructExportFlags.HasAnyFlags(UhtScriptStructExportFlags.IsVerseParametric)))
+			if (!isTopLevel || (mode == UhtVerseNameMode.Qualified && ScriptStructExportFlags.HasAnyFlags(UhtScriptStructExportFlags.IsVerseParametric)))
 			{
 				builder.Append('/').Append(VerseName);
 			}
 		}
 
 		/// <inheritdoc/>
-		protected override void AppendVerseName(StringBuilder builder, UhtVerseFullNameMode mode)
+		protected override void AppendVerseName(StringBuilder builder, UhtVerseNameMode mode)
 		{
 			builder.Append(VerseName);
-			if (mode == UhtVerseFullNameMode.Qualified && ScriptStructExportFlags.HasAnyFlags(UhtScriptStructExportFlags.IsVerseParametric))
+			if (mode == UhtVerseNameMode.Qualified && ScriptStructExportFlags.HasAnyFlags(UhtScriptStructExportFlags.IsVerseParametric))
 			{
 				builder.Append("(t)");
 			}
diff --git a/Engine/Source/Programs/Shared/EpicGames.UHT/Utils/UhtStringBuilderExtensions.cs b/Engine/Source/Programs/Shared/EpicGames.UHT/Utils/UhtStringBuilderExtensions.cs
--- a/Engine/Source/Programs/Shared/EpicGames.UHT/Utils/UhtStringBuilderExtensions.cs
+++ b/Engine/Source/Programs/Shared/EpicGames.UHT/Utils/UhtStringBuilderExtensions.cs
@@ -469,6 +469,78 @@ namespace EpicGames.UHT.Utils
 					throw new NotImplementedException();
 			}
 		}
+
+		/// <summary>
+		/// Given a verse name, encode it so it can be represented as an FName
+		/// </summary>
+		/// <param name="builder">Destination builder</param>
+		/// <param name="name">Name to encode</param>
+		/// <returns>Builder</returns>
+		public static StringBuilder AppendEncodedVerseName(this StringBuilder builder, ReadOnlySpan<char> name)
+		{
+			bool isFirstChar = true;
+			while (!name.IsEmpty)
+			{
+				char c = name[0];
+				name = name[1..];
+
+				if ((c >= 'a' && c <= 'z')
+					|| (c >= 'A' && c <= 'Z')
+					|| (c >= '0' && c <= '9' && !isFirstChar))
+				{
+					builder.Append(c);
+				}
+				else if (c == '[' && !name.IsEmpty && name[0] == ']')
+				{
+					name = name[1..];
+					builder.Append("_K");
+				}
+				else if (c == '-' && !name.IsEmpty && name[0] == '>')
+				{
+					name = name[1..];
+					builder.Append("_T");
+				}
+				else if (c == '_')
+				{
+					builder.Append("__");
+				}
+				else if (c == '(')
+				{
+					builder.Append("_L");
+				}
+				else if (c == ',')
+				{
+					builder.Append("_M");
+				}
+				else if (c == ':')
+				{
+					builder.Append("_N");
+				}
+				else if (c == '^')
+				{
+					builder.Append("_P");
+				}
+				else if (c == '?')
+				{
+					builder.Append("_Q");
+				}
+				else if (c == ')')
+				{
+					builder.Append("_R");
+				}
+				else if (c == '\'')
+				{
+					builder.Append("_U");
+				}
+				else
+				{
+					builder.Append('_').AppendFormat(CultureInfo.InvariantCulture, "{0:x2}", (uint)c);
+				}
+
+				isFirstChar = false;
+			}
+			return builder;
+		}
 	}
 
 	/// <summary>
diff --git a/Engine/Source/Programs/Shared/EpicGames.UHT/Utils/UhtUtilities.cs b/Engine/Source/Programs/Shared/EpicGames.UHT/Utils/UhtUtilities.cs
--- a/Engine/Source/Programs/Shared/EpicGames.UHT/Utils/UhtUtilities.cs
+++ b/Engine/Source/Programs/Shared/EpicGames.UHT/Utils/UhtUtilities.cs
@@ -6,6 +6,7 @@ using System.Collections.Generic;
 using System.Linq;
 using System.Text;
 using EpicGames.Core;
+using EpicGames.UHT.Types;
 
 namespace EpicGames.UHT.Utils
 {
@@ -399,5 +400,17 @@ namespace EpicGames.UHT.Utils
 				return (false, name);
 			}
 		}
+
+		/// <summary>
+		/// Given a verse name, encode it
+		/// </summary>
+		/// <param name="name">Name to be encoded</param>
+		/// <returns>Encoded name</returns>
+		public static string EncodeVerseName(string name)
+		{
+			StringBuilder builder = new();
+			builder.AppendEncodedVerseName(name);
+			return builder.ToString();
+		}
 	}
 }
