Adding V_REPORT_FUNCTION_CALL support to UHT

#rnx

diff --git a/Engine/Source/Programs/Shared/EpicGames.UHT/Exporters/CodeGen/UhtHeaderCodeGeneratorCppFile.cs b/Engine/Source/Programs/Shared/EpicGames.UHT/Exporters/CodeGen/UhtHeaderCodeGeneratorCppFile.cs
--- a/Engine/Source/Programs/Shared/EpicGames.UHT/Exporters/CodeGen/UhtHeaderCodeGeneratorCppFile.cs
+++ b/Engine/Source/Programs/Shared/EpicGames.UHT/Exporters/CodeGen/UhtHeaderCodeGeneratorCppFile.cs
@@ -1532,6 +1532,7 @@ namespace EpicGames.UHT.Exporters.CodeGen
 				builder.Append("\t((FTask*)V_THIS)->_bEverSuspended = true;\r\n");
 			}
 			builder.Append("\tV_NATIVE_END(").Append(function.FunctionExportFlags.HasAnyFlags(UhtFunctionExportFlags.VerseRTFMAlwaysOpen) ? "true" : "false").Append(")\r\n");
+			builder.Append("\tV_REPORT_FUNCTION_CALL(").AppendUTF8LiteralString(function.GetVerseName(UhtVerseNameMode.Default)).Append(", ").AppendUTF8LiteralString(function.GetVerseScope(UhtVerseNameMode.Default)).Append(");\r\n");
 			return builder;
 		}
 
diff --git a/Engine/Source/Programs/Shared/EpicGames.UHT/Types/UhtField.cs b/Engine/Source/Programs/Shared/EpicGames.UHT/Types/UhtField.cs
--- a/Engine/Source/Programs/Shared/EpicGames.UHT/Types/UhtField.cs
+++ b/Engine/Source/Programs/Shared/EpicGames.UHT/Types/UhtField.cs
@@ -107,6 +107,7 @@ namespace EpicGames.UHT.Types
 		/// <summary>
 		/// Cased name of the verse field
 		/// </summary>
+		[System.Diagnostics.CodeAnalysis.SuppressMessage("Naming", "CA1721:Property names should not match get methods", Justification = "<Pending>")]
 		public string? VerseName { get; set; } = null;
 
 		/// <summary>
@@ -150,6 +151,30 @@ namespace EpicGames.UHT.Types
 			return builder.ToString();
 		}
 
+		/// <summary>
+		/// Return the verse scope of the field
+		/// </summary>
+		/// <param name="mode">Controls how the name is generated</param>
+		/// <returns></returns>
+		public string GetVerseScope(UhtVerseNameMode mode)
+		{
+			StringBuilder builder = new();
+			AppendVerseScope(builder, mode);
+			return builder.ToString();
+		}
+
+		/// <summary>
+		/// Return the verse name of the field
+		/// </summary>
+		/// <param name="mode">Controls how the name is generated</param>
+		/// <returns></returns>
+		public string GetVerseName(UhtVerseNameMode mode)
+		{
+			StringBuilder builder = new();
+			AppendVerseName(builder, mode);
+			return builder.ToString();
+		}
+
 		/// <summary>
 		/// Append the scope and name of the verse type (scope:)name
 		/// </summary>
