UHT: fix engine version path for *nix machines

diff --git a/Engine/Source/Programs/Shared/EpicGames.UHT/Utils/UhtSession.cs b/Engine/Source/Programs/Shared/EpicGames.UHT/Utils/UhtSession.cs
--- a/Engine/Source/Programs/Shared/EpicGames.UHT/Utils/UhtSession.cs
+++ b/Engine/Source/Programs/Shared/EpicGames.UHT/Utils/UhtSession.cs
@@ -2002,7 +2002,7 @@ namespace EpicGames.UHT.Utils
 				throw new UhtException("Can't parse engine version without an engine directory");
 			}
 
-			string versionHeader = Path.Combine(EngineDirectory, @"Source\Runtime\Launch\Resources\Version.h");
+			string versionHeader = Path.Combine(EngineDirectory, "Source", "Runtime", "Launch", "Resources", "Version.h");
 			Log.Logger.LogDebug("Reading engine version from file: {VersionHeader}", versionHeader);
 			string[] lines;
 			try
