Horde: Disable trace level logging of agent channel data

diff --git a/Engine/Source/Programs/Shared/EpicGames.Horde/Compute/AgentMessageChannel.cs b/Engine/Source/Programs/Shared/EpicGames.Horde/Compute/AgentMessageChannel.cs
--- a/Engine/Source/Programs/Shared/EpicGames.Horde/Compute/AgentMessageChannel.cs
+++ b/Engine/Source/Programs/Shared/EpicGames.Horde/Compute/AgentMessageChannel.cs
@@ -2,6 +2,7 @@
 
 using System;
 using System.Buffers.Binary;
+using System.Diagnostics.CodeAnalysis;
 using System.Text;
 using System.Threading;
 using System.Threading.Tasks;
@@ -54,10 +55,10 @@ namespace EpicGames.Horde.Compute
 				header[0] = (byte)_type;
 				BinaryPrimitives.WriteInt32LittleEndian(header.Slice(1, 4), _length);
 
-				if (_channel._logger.IsEnabled(LogLevel.Trace))
-				{
-					_channel.LogMessageInfo("SEND", _type, _sendBufferWriter.GetWriteBuffer().Slice(HeaderLength, _length).Span);
-				}
+				// if (_channel._logger.IsEnabled(LogLevel.Trace))
+				// {
+				// 	_channel.LogMessageInfo("SEND", _type, _sendBufferWriter.GetWriteBuffer().Slice(HeaderLength, _length).Span);
+				// }
 
 				_sendBufferWriter.AdvanceWritePosition(HeaderLength + _length);
 				_length = 0;
@@ -164,10 +165,10 @@ namespace EpicGames.Horde.Compute
 
 				AgentMessageType type = (AgentMessageType)memory.Span[0];
 				AgentMessage message = new AgentMessage(type, memory.Slice(HeaderLength, messageLength));
-				if (_logger.IsEnabled(LogLevel.Trace))
-				{
-					LogMessageInfo("RECV", message.Type, message.Data.Span);
-				}
+				// if (_logger.IsEnabled(LogLevel.Trace))
+				// {
+				// 	LogMessageInfo("RECV", message.Type, message.Data.Span);
+				// }
 
 				_recvBufferReader.AdvanceReadPosition(HeaderLength + messageLength);
 				return message;
@@ -175,6 +176,7 @@ namespace EpicGames.Horde.Compute
 			return new AgentMessage(AgentMessageType.None, ReadOnlyMemory<byte>.Empty);
 		}
 
+		[SuppressMessage("CodeQuality", "IDE0051:Remove unused private members", Justification = "Log calls disabled")]
 		void LogMessageInfo(string verb, AgentMessageType type, ReadOnlySpan<byte> data)
 		{
 			StringBuilder bytes = new StringBuilder();
