Adding WITH_VERSE_BPVM as a valid #define

#rnx

diff --git a/Engine/Source/Programs/Shared/EpicGames.UHT/Parsers/UhtHeaderFileParser.cs b/Engine/Source/Programs/Shared/EpicGames.UHT/Parsers/UhtHeaderFileParser.cs
--- a/Engine/Source/Programs/Shared/EpicGames.UHT/Parsers/UhtHeaderFileParser.cs
+++ b/Engine/Source/Programs/Shared/EpicGames.UHT/Parsers/UhtHeaderFileParser.cs
@@ -107,6 +107,11 @@ namespace EpicGames.UHT.Parsers
 		/// </summary>
 		WithTests = 1 << 11,
 
+		/// <summary>
+		/// This indicates we are in a "#if WITH_VERSE_BPVM" block
+		/// </summary>
+		WithVerseBPVM = 1 << 12,
+
 		/// <summary>
 		/// The following flags are always ignored when keywords test for allowed conditional blocks
 		/// </summary>
@@ -115,7 +120,7 @@ namespace EpicGames.UHT.Parsers
 		/// <summary>
 		/// Default compiler directives to be allowed
 		/// </summary>
-		DefaultAllowedCheck = WithEditor | WithEditorOnlyData | WithVerseVM | WithTests,
+		DefaultAllowedCheck = WithEditor | WithEditorOnlyData | WithVerseVM | WithTests | WithVerseBPVM,
 
 		/// <summary>
 		/// All flags are allowed
@@ -181,6 +186,7 @@ namespace EpicGames.UHT.Parsers
 				case UhtCompilerDirective.WithEngine: return UhtNames.WithEngine;
 				case UhtCompilerDirective.WithCoreUObject: return UhtNames.WithCoreUObject;
 				case UhtCompilerDirective.WithVerseVM: return UhtNames.WithVerseVM;
+				case UhtCompilerDirective.WithVerseBPVM: return UhtNames.WithVerseBPVM;
 				case UhtCompilerDirective.WithTests: return UhtNames.WithTests;
 				default: return "<unrecognized>";
 			}
@@ -986,6 +992,10 @@ namespace EpicGames.UHT.Parsers
 					{
 						return notPresent ? UhtCompilerDirective.Unrecognized : UhtCompilerDirective.WithVerseVM;
 					}
+					else if (define.IsValue(UhtNames.WithVerseBPVM))
+					{
+						return notPresent ? UhtCompilerDirective.Unrecognized : UhtCompilerDirective.WithVerseBPVM;
+					}
 					else if (define.IsValue(UhtNames.WithTests))
 					{
 						return notPresent ? UhtCompilerDirective.Unrecognized : UhtCompilerDirective.WithTests;
@@ -1121,6 +1131,7 @@ namespace EpicGames.UHT.Parsers
 				case UhtCompilerDirective.WithCoreUObject:
 				case UhtCompilerDirective.WithHotReload:
 				case UhtCompilerDirective.WithVerseVM:
+				case UhtCompilerDirective.WithVerseBPVM:
 				case UhtCompilerDirective.WithTests:
 					return true;
 
@@ -1146,6 +1157,7 @@ namespace EpicGames.UHT.Parsers
 				case UhtCompilerDirective.WithCoreUObject:
 				case UhtCompilerDirective.WithHotReload:
 				case UhtCompilerDirective.WithVerseVM:
+				case UhtCompilerDirective.WithVerseBPVM:
 				case UhtCompilerDirective.WithTests:
 					return true;
 
@@ -1163,6 +1175,7 @@ namespace EpicGames.UHT.Parsers
 				compilerDirective == UhtCompilerDirective.WithEngine ||
 				compilerDirective == UhtCompilerDirective.WithCoreUObject ||
 				compilerDirective == UhtCompilerDirective.WithVerseVM ||
+				compilerDirective == UhtCompilerDirective.WithVerseBPVM ||
 				compilerDirective == UhtCompilerDirective.WithTests;
 		}
 
diff --git a/Engine/Source/Programs/Shared/EpicGames.UHT/Parsers/UhtPropertyParser.cs b/Engine/Source/Programs/Shared/EpicGames.UHT/Parsers/UhtPropertyParser.cs
--- a/Engine/Source/Programs/Shared/EpicGames.UHT/Parsers/UhtPropertyParser.cs
+++ b/Engine/Source/Programs/Shared/EpicGames.UHT/Parsers/UhtPropertyParser.cs
@@ -692,6 +692,10 @@ namespace EpicGames.UHT.Parsers
 				{
 					propertySettings.DefineScope |= UhtDefineScope.VerseVM;
 				}
+				if (compilerDirective.HasAnyFlags(UhtCompilerDirective.WithVerseBPVM))
+				{
+					propertySettings.DefineScope |= UhtDefineScope.VerseBPVM;
+				}
 				if (compilerDirective.HasAnyFlags(UhtCompilerDirective.WithTests))
 				{
 					propertySettings.DefineScope |= UhtDefineScope.Tests;
diff --git a/Engine/Source/Programs/Shared/EpicGames.UHT/Types/UhtType.cs b/Engine/Source/Programs/Shared/EpicGames.UHT/Types/UhtType.cs
--- a/Engine/Source/Programs/Shared/EpicGames.UHT/Types/UhtType.cs
+++ b/Engine/Source/Programs/Shared/EpicGames.UHT/Types/UhtType.cs
@@ -572,10 +572,15 @@ namespace EpicGames.UHT.Types
 		/// </summary>
 		Tests = 1 << 2,
 
+		/// <summary>
+		/// Exists inside of a WITH_VERSE_BPVM block
+		/// </summary>
+		VerseBPVM = 1 << 3,
+
 		/// <summary>
 		/// Number of unique scope combinations
 		/// </summary>
-		ScopeCount = 1 << 3,
+		ScopeCount = 1 << 4,
 	}
 
 	/// <summary>
@@ -633,6 +638,8 @@ namespace EpicGames.UHT.Types
 					return defineScopeNames == UhtDefineScopeNames.WithEditor ? UhtCompilerDirective.WithEditor : UhtCompilerDirective.WithEditorOnlyData;
 				case UhtDefineScope.VerseVM:
 					return UhtCompilerDirective.WithVerseVM;
+				case UhtDefineScope.VerseBPVM:
+					return UhtCompilerDirective.WithVerseBPVM;
 				case UhtDefineScope.Tests:
 					return UhtCompilerDirective.WithTests;
 				default:
diff --git a/Engine/Source/Programs/Shared/EpicGames.UHT/Utils/UhtNames.cs b/Engine/Source/Programs/Shared/EpicGames.UHT/Utils/UhtNames.cs
--- a/Engine/Source/Programs/Shared/EpicGames.UHT/Utils/UhtNames.cs
+++ b/Engine/Source/Programs/Shared/EpicGames.UHT/Utils/UhtNames.cs
@@ -77,6 +77,7 @@ namespace EpicGames.UHT.Utils
 		public const string WithEngine = "WITH_ENGINE";
 		public const string WithCoreUObject = "WITH_COREUOBJECT";
 		public const string WithVerseVM = "WITH_VERSE_VM";
+		public const string WithVerseBPVM = "WITH_VERSE_BPVM";
 		public const string WithTests = "WITH_TESTS";
 	}
 }
