Fixing define scope issues where not all possible scopes were being applied to all types and not all generated code was having their code blocked off.

#rnx

diff --git a/Engine/Source/Programs/Shared/EpicGames.UHT/Exporters/CodeGen/UhtHeaderCodeGeneratorCppFile.cs b/Engine/Source/Programs/Shared/EpicGames.UHT/Exporters/CodeGen/UhtHeaderCodeGeneratorCppFile.cs
--- a/Engine/Source/Programs/Shared/EpicGames.UHT/Exporters/CodeGen/UhtHeaderCodeGeneratorCppFile.cs
+++ b/Engine/Source/Programs/Shared/EpicGames.UHT/Exporters/CodeGen/UhtHeaderCodeGeneratorCppFile.cs
@@ -925,133 +925,128 @@ namespace EpicGames.UHT.Exporters.CodeGen
 			// Everything from this point on will be part of the definition hash
 			int hashCodeBlockStart = builder.Length;
 
-			// Note: This is fairly aggressive 
-			using (UhtMacroBlockEmitter blockEmitter = new(builder, UhtDefineScopeNames.WithEditor, function.DefineScope))
-			{
+			// Collect the properties
+			PropertyMemberContextImpl context = new(CodeGenerator, function, eventParameters, staticsName);
+			UhtUsedDefineScopes<UhtProperty> properties = new(function.Properties);
 
-				// Collect the properties
-				PropertyMemberContextImpl context = new(CodeGenerator, function, eventParameters, staticsName);
-				UhtUsedDefineScopes<UhtProperty> properties = new(function.Properties);
+			// Statics declaration
+			{
+				builder.Append("struct ").Append(staticsName).Append("\r\n");
+				builder.Append("{\r\n");
 
-				// Statics declaration
+				if (paramsInStatic)
 				{
-					builder.Append("struct ").Append(staticsName).Append("\r\n");
-					builder.Append("{\r\n");
-
-					if (paramsInStatic)
+					List<UhtScriptStruct> noExportStructs = FindNoExportStructs(function);
+					foreach (UhtScriptStruct noExportStruct in noExportStructs)
 					{
-						List<UhtScriptStruct> noExportStructs = FindNoExportStructs(function);
-						foreach (UhtScriptStruct noExportStruct in noExportStructs)
-						{
-							AppendMirrorsForNoexportStruct(builder, noExportStruct, 1);
-						}
-						AppendEventParameter(builder, function, strippedFunctionName, UhtPropertyTextType.EventParameterFunctionMember, false, 1, "\r\n");
+						AppendMirrorsForNoexportStruct(builder, noExportStruct, 1);
 					}
+					AppendEventParameter(builder, function, strippedFunctionName, UhtPropertyTextType.EventParameterFunctionMember, false, 1, "\r\n");
+				}
 
-					builder.AppendMetaDataDecl(function, context, properties, MetaDataParamsName, 1);
+				builder.AppendMetaDataDecl(function, context, properties, MetaDataParamsName, 1);
 
-					AppendPropertiesDecl(builder, context, properties, 1);
+				AppendPropertiesDecl(builder, context, properties, 1);
 
-					builder.Append("\tstatic const UECodeGen_Private::").AppendFunctionParamsType(function).Append(" FuncParams;\r\n");
+				builder.Append("\tstatic const UECodeGen_Private::").AppendFunctionParamsType(function).Append(" FuncParams;\r\n");
 
-					builder.Append("};\r\n");
-				}
+				builder.Append("};\r\n");
+			}
 
-				// Statics definition
-				{
-					AppendPropertiesDefs(builder, context, properties, 0);
+			// Statics definition
+			{
+				AppendPropertiesDefs(builder, context, properties, 0);
 
-					string functionEngineName = isVerseNativeCallable ? function.GetEncodedVerseScopeAndName(UhtVerseNameMode.Default) : function.EngineName;
-					builder
-						.Append("const UECodeGen_Private::").AppendFunctionParamsType(function).Append(' ')
-						.Append(staticsName).Append("::FuncParams = { { ")
-						.Append("(UObject*(*)())").Append(GetFunctionOuterFunc(function)).Append(", ")
-						.Append(GetSingletonName(function.SuperFunction, true)).Append(", ")
-						.AppendUTF8LiteralString(functionEngineName).Append(", ");
+				string functionEngineName = isVerseNativeCallable ? function.GetEncodedVerseScopeAndName(UhtVerseNameMode.Default) : function.EngineName;
+				builder
+					.Append("const UECodeGen_Private::").AppendFunctionParamsType(function).Append(' ')
+					.Append(staticsName).Append("::FuncParams = { { ")
+					.Append("(UObject*(*)())").Append(GetFunctionOuterFunc(function)).Append(", ")
+					.Append(GetSingletonName(function.SuperFunction, true)).Append(", ")
+					.AppendUTF8LiteralString(functionEngineName).Append(", ");
 
-					builder.AppendArrayPtrLine(properties, UhtDefineScopeNames.Standard, staticsName, "PropPointers", 0, ", ");
-					builder.AppendArrayCountLine(properties, UhtDefineScopeNames.Standard, staticsName, "PropPointers", 0, ", ");
+				builder.AppendArrayPtrLine(properties, UhtDefineScopeNames.Standard, staticsName, "PropPointers", 0, ", ");
+				builder.AppendArrayCountLine(properties, UhtDefineScopeNames.Standard, staticsName, "PropPointers", 0, ", ");
 
-					string sizeOfStatic = "";
-					if (function.Children.Count > 0)
+				string sizeOfStatic = "";
+				if (function.Children.Count > 0)
+				{
+					UhtFunction tempFunction = function;
+					while (tempFunction.SuperFunction != null)
 					{
-						UhtFunction tempFunction = function;
-						while (tempFunction.SuperFunction != null)
-						{
-							tempFunction = tempFunction.SuperFunction;
-						}
-
-						if (paramsInStatic)
-						{
-							sizeOfStatic = "sizeof(" + staticsName + "::" + GetEventStructParametersName(tempFunction.Outer, tempFunction.StrippedFunctionName) + ")";
-						}
-						else
-						{
-							sizeOfStatic = "sizeof(" + GetEventStructParametersName(tempFunction.Outer, tempFunction.StrippedFunctionName) + ")";
-						}
+						tempFunction = tempFunction.SuperFunction;
+					}
 
-						builder.Append(sizeOfStatic);
+					if (paramsInStatic)
+					{
+						sizeOfStatic = "sizeof(" + staticsName + "::" + GetEventStructParametersName(tempFunction.Outer, tempFunction.StrippedFunctionName) + ")";
 					}
 					else
 					{
-						builder.Append('0');
+						sizeOfStatic = "sizeof(" + GetEventStructParametersName(tempFunction.Outer, tempFunction.StrippedFunctionName) + ")";
 					}
-					builder.Append(", ");
 
-					builder
-						.Append("RF_Public|RF_Transient|RF_MarkAsNative, ")
-						.Append($"(EFunctionFlags)0x{(uint)function.FunctionFlags:X8}, ")
-						.Append(isNet ? function.RPCId : 0).Append(", ")
-						.Append(isNet ? function.RPCResponseId : 0).Append(", ")
-						.AppendMetaDataParams(function, staticsName, MetaDataParamsName)
-						.Append("}, ");
-
-					switch (function.FunctionType)
-					{
-						case UhtFunctionType.Function:
-							if (function.IsVerseField)
-							{
-								builder.Append("\r\n");
-								builder.Append("#if WITH_VERSE_VM\r\n");
-								builder.AppendUTF8LiteralString(function.GetVerseScopeAndName(UhtVerseNameMode.Default)).Append(",\r\n");
-								builder.Append("#else\r\n");
-								builder.AppendUTF8LiteralString(function.GetEncodedVerseScopeAndName(UhtVerseNameMode.Default)).Append(",\r\n");
-								builder.Append("#endif\r\n");
-							}
-							break;
+					builder.Append(sizeOfStatic);
+				}
+				else
+				{
+					builder.Append('0');
+				}
+				builder.Append(", ");
 
-						case UhtFunctionType.Delegate:
-							break;
+				builder
+					.Append("RF_Public|RF_Transient|RF_MarkAsNative, ")
+					.Append($"(EFunctionFlags)0x{(uint)function.FunctionFlags:X8}, ")
+					.Append(isNet ? function.RPCId : 0).Append(", ")
+					.Append(isNet ? function.RPCResponseId : 0).Append(", ")
+					.AppendMetaDataParams(function, staticsName, MetaDataParamsName)
+					.Append("}, ");
 
-						case UhtFunctionType.SparseDelegate:
-							builder
-								.AppendUTF8LiteralString(function.SparseOwningClassName).Append(", ")
-								.AppendUTF8LiteralString(function.SparseDelegateName).Append(", ");
-							break;
-					}
+				switch (function.FunctionType)
+				{
+					case UhtFunctionType.Function:
+						if (function.IsVerseField)
+						{
+							builder.Append("\r\n");
+							builder.Append("#if WITH_VERSE_VM\r\n");
+							builder.AppendUTF8LiteralString(function.GetVerseScopeAndName(UhtVerseNameMode.Default)).Append(",\r\n");
+							builder.Append("#else\r\n");
+							builder.AppendUTF8LiteralString(function.GetEncodedVerseScopeAndName(UhtVerseNameMode.Default)).Append(",\r\n");
+							builder.Append("#endif\r\n");
+						}
+						break;
 
-					builder.Append(" };\r\n");
+					case UhtFunctionType.Delegate:
+						break;
 
-					if (!String.IsNullOrEmpty(sizeOfStatic))
-					{
-						builder.Append("static_assert(").Append(sizeOfStatic).Append(" < MAX_uint16);\r\n");
-					}
+					case UhtFunctionType.SparseDelegate:
+						builder
+							.AppendUTF8LiteralString(function.SparseOwningClassName).Append(", ")
+							.AppendUTF8LiteralString(function.SparseDelegateName).Append(", ");
+						break;
 				}
 
-				// Registration function
+				builder.Append(" };\r\n");
+
+				if (!String.IsNullOrEmpty(sizeOfStatic))
 				{
-					builder.Append("UFunction* ").Append(singletonName).Append("()\r\n");
-					builder.Append("{\r\n");
-					builder.Append("\tstatic UFunction* ReturnFunction = nullptr;\r\n");
-					builder.Append("\tif (!ReturnFunction)\r\n");
-					builder.Append("\t{\r\n");
-					builder.Append("\t\t").AppendFunctionConstructorName(function).Append("(&ReturnFunction, ").Append(staticsName).Append("::FuncParams);\r\n");
-					builder.Append("\t}\r\n");
-					builder.Append("\treturn ReturnFunction;\r\n");
-					builder.Append("}\r\n");
+					builder.Append("static_assert(").Append(sizeOfStatic).Append(" < MAX_uint16);\r\n");
 				}
 			}
 
+			// Registration function
+			{
+				builder.Append("UFunction* ").Append(singletonName).Append("()\r\n");
+				builder.Append("{\r\n");
+				builder.Append("\tstatic UFunction* ReturnFunction = nullptr;\r\n");
+				builder.Append("\tif (!ReturnFunction)\r\n");
+				builder.Append("\t{\r\n");
+				builder.Append("\t\t").AppendFunctionConstructorName(function).Append("(&ReturnFunction, ").Append(staticsName).Append("::FuncParams);\r\n");
+				builder.Append("\t}\r\n");
+				builder.Append("\treturn ReturnFunction;\r\n");
+				builder.Append("}\r\n");
+			}
+
 			using (UhtRentedPoolBuffer<char> borrowBuffer = builder.RentPoolBuffer(hashCodeBlockStart, builder.Length - hashCodeBlockStart))
 			{
 				ObjectInfos[function.ObjectTypeIndex].Hash = UhtHash.GenenerateTextHash(borrowBuffer.Buffer.Memory.Span);
@@ -1146,9 +1141,11 @@ namespace EpicGames.UHT.Exporters.CodeGen
 			bool isCallback = IsCallbackFunction(function);
 			bool isVerseNativeCallable = function.FunctionExportFlags.HasAnyFlags(UhtFunctionExportFlags.VerseNativeCallable);
 			bool isVerseCallable = function.FunctionExportFlags.HasAnyFlags(UhtFunctionExportFlags.VerseCallable);
+
 			if (isNotDelegate || isRpc || isCallback)
 			{
-				using UhtCodeBlockComment blockCommand = new(builder, classObj, function);
+				using UhtCodeBlockComment blockComment = new(builder, classObj, function);
+				using UhtMacroBlockEmitter blockEmitter = new(builder, UhtDefineScopeNames.WithEditor, function.DefineScope);
 				if (isCallback)
 				{
 					AppendEventParameter(builder, function, function.StrippedFunctionName, UhtPropertyTextType.EventParameterMember, true, 0, "\r\n");
@@ -1177,7 +1174,6 @@ namespace EpicGames.UHT.Exporters.CodeGen
 				}
 				if (isRpc)
 				{
-					using UhtMacroBlockEmitter blockEmitter = new(builder, UhtDefineScopeNames.WithEditor, function.DefineScope);
 					builder.Append("DEFINE_FUNCTION(").AppendClassSourceNameOrInterfaceName(classObj).Append("::").Append(function.UnMarshalAndCallName).Append(")\r\n");
 					builder.Append("{\r\n");
 					AppendFunctionThunk(builder, function);
diff --git a/Engine/Source/Programs/Shared/EpicGames.UHT/Parsers/UhtClassParser.cs b/Engine/Source/Programs/Shared/EpicGames.UHT/Parsers/UhtClassParser.cs
--- a/Engine/Source/Programs/Shared/EpicGames.UHT/Parsers/UhtClassParser.cs
+++ b/Engine/Source/Programs/Shared/EpicGames.UHT/Parsers/UhtClassParser.cs
@@ -120,10 +120,7 @@ namespace EpicGames.UHT.Parsers
 					}
 
 					UhtCompilerDirective compilerDirective = topScope.HeaderParser.GetCurrentCompositeCompilerDirective();
-					if ((compilerDirective & UhtCompilerDirective.WithTests) != 0)
-					{
-						classObj.DefineScope |= UhtDefineScope.Tests;
-					}
+					classObj.DefineScope |= compilerDirective.GetDefaultDefineScopes();
 
 					classObj.Outer?.AddChild(classObj);
 
diff --git a/Engine/Source/Programs/Shared/EpicGames.UHT/Parsers/UhtEnumParser.cs b/Engine/Source/Programs/Shared/EpicGames.UHT/Parsers/UhtEnumParser.cs
--- a/Engine/Source/Programs/Shared/EpicGames.UHT/Parsers/UhtEnumParser.cs
+++ b/Engine/Source/Programs/Shared/EpicGames.UHT/Parsers/UhtEnumParser.cs
@@ -70,15 +70,11 @@ namespace EpicGames.UHT.Parsers
 					enumObject.Outer?.AddChild(enumObject);
 
 					UhtCompilerDirective compilerDirective = topScope.HeaderParser.GetCurrentCompositeCompilerDirective();
-					if ((compilerDirective & UhtCompilerDirective.WithEditorOnlyData) != 0)
+					if (compilerDirective.HasAnyFlags(UhtCompilerDirective.WithEditorOnlyData))
 					{
 						enumObject.DefineScope |= UhtDefineScope.EditorOnlyData;
 					}
-
-					if ((compilerDirective & UhtCompilerDirective.WithTests) != 0)
-					{
-						enumObject.DefineScope |= UhtDefineScope.Tests;
-					}
+					enumObject.DefineScope |= compilerDirective.GetDefaultDefineScopes();
 
 					// Read base for enum class
 					if (enumObject.CppForm == UhtEnumCppForm.EnumClass)
diff --git a/Engine/Source/Programs/Shared/EpicGames.UHT/Parsers/UhtFunctionParser.cs b/Engine/Source/Programs/Shared/EpicGames.UHT/Parsers/UhtFunctionParser.cs
--- a/Engine/Source/Programs/Shared/EpicGames.UHT/Parsers/UhtFunctionParser.cs
+++ b/Engine/Source/Programs/Shared/EpicGames.UHT/Parsers/UhtFunctionParser.cs
@@ -361,11 +361,7 @@ namespace EpicGames.UHT.Parsers
 						function.FunctionFlags |= EFunctionFlags.EditorOnly;
 						function.DefineScope |= UhtDefineScope.EditorOnlyData;
 					}
-
-					if ((compilerDirective & UhtCompilerDirective.WithTests) != 0)
-					{
-						function.DefineScope |= UhtDefineScope.Tests;
-					}
+					function.DefineScope |= compilerDirective.GetDefaultDefineScopes();
 
 					specifierParser.ParseDeferred();
 					FinalizeFunctionSpecifiers(function);
diff --git a/Engine/Source/Programs/Shared/EpicGames.UHT/Parsers/UhtHeaderFileParser.cs b/Engine/Source/Programs/Shared/EpicGames.UHT/Parsers/UhtHeaderFileParser.cs
--- a/Engine/Source/Programs/Shared/EpicGames.UHT/Parsers/UhtHeaderFileParser.cs
+++ b/Engine/Source/Programs/Shared/EpicGames.UHT/Parsers/UhtHeaderFileParser.cs
@@ -192,6 +192,29 @@ namespace EpicGames.UHT.Parsers
 			}
 		}
 
+		/// <summary>
+		/// Return a define scope based on the given compiler directives
+		/// </summary>
+		/// <param name="compilerDirectives">Directives in question</param>
+		/// <returns>Matching define scope</returns>
+		public static UhtDefineScope GetDefaultDefineScopes(this UhtCompilerDirective compilerDirectives)
+		{
+			UhtDefineScope defineScope = UhtDefineScope.None;
+			if (compilerDirectives.HasAnyFlags(UhtCompilerDirective.WithTests))
+			{
+				defineScope |= UhtDefineScope.Tests;
+			}
+			if (compilerDirectives.HasAnyFlags(UhtCompilerDirective.WithVerseBPVM))
+			{
+				defineScope |= UhtDefineScope.VerseBPVM;
+			}
+			if (compilerDirectives.HasAnyFlags(UhtCompilerDirective.WithVerseVM))
+			{
+				defineScope |= UhtDefineScope.VerseVM;
+			}
+			return defineScope;
+		}
+
 		/// <summary>
 		/// Return a string list of the given compiler directives
 		/// </summary>
diff --git a/Engine/Source/Programs/Shared/EpicGames.UHT/Parsers/UhtPropertyParser.cs b/Engine/Source/Programs/Shared/EpicGames.UHT/Parsers/UhtPropertyParser.cs
--- a/Engine/Source/Programs/Shared/EpicGames.UHT/Parsers/UhtPropertyParser.cs
+++ b/Engine/Source/Programs/Shared/EpicGames.UHT/Parsers/UhtPropertyParser.cs
@@ -688,18 +688,7 @@ namespace EpicGames.UHT.Parsers
 						tokenReader.LogError("UProperties should not be wrapped by WITH_EDITOR, use WITH_EDITORONLY_DATA instead.");
 					}
 				}
-				if (compilerDirective.HasAnyFlags(UhtCompilerDirective.WithVerseVM))
-				{
-					propertySettings.DefineScope |= UhtDefineScope.VerseVM;
-				}
-				if (compilerDirective.HasAnyFlags(UhtCompilerDirective.WithVerseBPVM))
-				{
-					propertySettings.DefineScope |= UhtDefineScope.VerseBPVM;
-				}
-				if (compilerDirective.HasAnyFlags(UhtCompilerDirective.WithTests))
-				{
-					propertySettings.DefineScope |= UhtDefineScope.Tests;
-				}
+				propertySettings.DefineScope |= compilerDirective.GetDefaultDefineScopes();
 			}
 
 			// Parse type information including UPARAM that might appear in template arguments
diff --git a/Engine/Source/Programs/Shared/EpicGames.UHT/Parsers/UhtScriptStructParser.cs b/Engine/Source/Programs/Shared/EpicGames.UHT/Parsers/UhtScriptStructParser.cs
--- a/Engine/Source/Programs/Shared/EpicGames.UHT/Parsers/UhtScriptStructParser.cs
+++ b/Engine/Source/Programs/Shared/EpicGames.UHT/Parsers/UhtScriptStructParser.cs
@@ -142,10 +142,7 @@ namespace EpicGames.UHT.Parsers
 					}
 
 					UhtCompilerDirective compilerDirective = topScope.HeaderParser.GetCurrentCompositeCompilerDirective();
-					if ((compilerDirective & UhtCompilerDirective.WithTests) != 0)
-					{
-						scriptStruct.DefineScope |= UhtDefineScope.Tests;
-					}
+					scriptStruct.DefineScope |= compilerDirective.GetDefaultDefineScopes();
 
 					scriptStruct.Outer?.AddChild(scriptStruct);
 
diff --git a/Engine/Source/Programs/Shared/EpicGames.UHT/Specifiers/UhtPropertyMemberSpecifiers.cs b/Engine/Source/Programs/Shared/EpicGames.UHT/Specifiers/UhtPropertyMemberSpecifiers.cs
--- a/Engine/Source/Programs/Shared/EpicGames.UHT/Specifiers/UhtPropertyMemberSpecifiers.cs
+++ b/Engine/Source/Programs/Shared/EpicGames.UHT/Specifiers/UhtPropertyMemberSpecifiers.cs
@@ -480,12 +480,5 @@ namespace EpicGames.UHT.Parsers
 			UhtPropertySpecifierContext context = (UhtPropertySpecifierContext)specifierContext;
 			context.PropertySettings.PropertyExportFlags |= UhtPropertyExportFlags.FieldNotify;
 		}
-
-		//[UhtSpecifier(Extends = UhtTableNames.PropertyMember, ValueType = UhtSpecifierValueType.String)]
-		//private static void VerseNameSpecifier(UhtSpecifierContext specifierContext, StringView value)
-		//{
-		//	UhtPropertySpecifierContext context = (UhtPropertySpecifierContext)specifierContext;
-		//	context.PropertySettings.VerseName = value.ToString();
-		//}
 	}
 }
diff --git a/Engine/Source/Programs/Shared/EpicGames.UHT/Types/UhtProperty.cs b/Engine/Source/Programs/Shared/EpicGames.UHT/Types/UhtProperty.cs
--- a/Engine/Source/Programs/Shared/EpicGames.UHT/Types/UhtProperty.cs
+++ b/Engine/Source/Programs/Shared/EpicGames.UHT/Types/UhtProperty.cs
@@ -1544,7 +1544,7 @@ namespace EpicGames.UHT.Types
 		}
 
 		/// <summary>
-		/// Apppend the name of a function thunk paramter
+		/// Append the name of a function thunk paramter
 		/// </summary>
 		/// <param name="builder">Output builder</param>
 		/// <returns>Output builder</returns>
