Add validation for metadata which needs to be a path to another type/function to UHT.
This exists in an automation test in C++, but UHT will surface the issue immediately.

diff --git a/Engine/Source/Programs/Shared/EpicGames.UHT/Specifiers/UhtDefaultSpecifiers.cs b/Engine/Source/Programs/Shared/EpicGames.UHT/Specifiers/UhtDefaultSpecifiers.cs
--- a/Engine/Source/Programs/Shared/EpicGames.UHT/Specifiers/UhtDefaultSpecifiers.cs
+++ b/Engine/Source/Programs/Shared/EpicGames.UHT/Specifiers/UhtDefaultSpecifiers.cs
@@ -242,6 +242,38 @@ namespace EpicGames.UHT.Parsers
 				type.LogError($"Unrecognized units '{value}' specified for '{type.FullName}'");
 			}
 		}
+
+  
+		/// <summary>
+		/// Validator for metadata on various types which must be a long path name to another type/function/etc.
+		/// </summary>
+		[UhtSpecifierValidator(Name = "MustImplement", Extends = UhtTableNames.PropertyMember)]
+		[UhtSpecifierValidator(Name = "ObjectMustImplement", Extends = UhtTableNames.PropertyMember)]
+		[UhtSpecifierValidator(Name = "AllowedClasses", Extends = UhtTableNames.PropertyMember)]
+		[UhtSpecifierValidator(Name = "DisallowedClasses", Extends = UhtTableNames.PropertyMember)]
+		[UhtSpecifierValidator(Name = "MetaClass", Extends = UhtTableNames.PropertyMember)]
+		[UhtSpecifierValidator(Name = "MetaStruct", Extends = UhtTableNames.PropertyMember)]
+		[UhtSpecifierValidator(Name = "BaseClass", Extends = UhtTableNames.PropertyMember)]
+		[UhtSpecifierValidator(Name = "BaseStruct", Extends = UhtTableNames.PropertyMember)]
+		[UhtSpecifierValidator(Name = "BitmaskEnum", Extends = UhtTableNames.PropertyMember)]
+		[UhtSpecifierValidator(Name = "HasNativeBreak", Extends = UhtTableNames.ScriptStruct)]
+		[UhtSpecifierValidator(Name = "HasNativeMake", Extends = UhtTableNames.ScriptStruct)]
+		[UhtSpecifierValidator(Name = "ActorClassToPoll", Extends = UhtTableNames.Function)]
+		[UhtSpecifierValidator(Name = "ComponentClass", Extends = UhtTableNames.Function)]
+		private static void LongPathNameMetaDataSpecifierValidator(UhtType type, UhtMetaData metaData, UhtMetaDataKey key, StringView value)
+		{
+			foreach (string element in value.ToString().Split(',', StringSplitOptions.RemoveEmptyEntries | StringSplitOptions.TrimEntries))
+			{
+				if (!UhtUtilities.IsValidLongPathName(element))
+				{
+					UhtType? found = type.FindType(UhtFindOptions.TypesMask | UhtFindOptions.EngineName | UhtFindOptions.NoSelf, element);
+					string message = found is not null 
+						? $"MetaData {key} on {type.EngineName} has value '{element}' which is not a valid long path name. Did you mean '{found.PathName}' ?"
+						: $"MetaData {key} on {type.EngineName} has value '{element}' which is not a valid long path name. Replace it with a full path name specifying the package containing the type.";
+					type.LogError(message);
+				}
+			}
+		}
 		#endregion
 	}
 }
diff --git a/Engine/Source/Programs/Shared/EpicGames.UHT/Utils/UhtUtilities.cs b/Engine/Source/Programs/Shared/EpicGames.UHT/Utils/UhtUtilities.cs
--- a/Engine/Source/Programs/Shared/EpicGames.UHT/Utils/UhtUtilities.cs
+++ b/Engine/Source/Programs/Shared/EpicGames.UHT/Utils/UhtUtilities.cs
@@ -1,6 +1,7 @@
 // Copyright Epic Games, Inc. All Rights Reserved.
 
 using System;
+using System.Buffers;
 using System.Buffers.Binary;
 using System.Collections.Generic;
 using System.Linq;
@@ -197,6 +198,57 @@ namespace EpicGames.UHT.Utils
 					return new UhtEngineNameParts { Prefix = new StringView(String.Empty), EngineName = new StringView(sourceName), IsDeprecated = false };
 			}
 		}
+	
+		private static readonly SearchValues<char> s_pathNameDelimiters = SearchValues.Create(['.', ':']);
+
+	 	/// <summary>
+		/// Returns whether the given string is a valid long path name for an object.
+		/// See FSoftObjectPath::SetPath for reference.
+		/// e.g. MyClass is a short name, but /Script/MyModule.MyClass is a full path name.
+		/// </summary>
+		/// <param name="name"></param>
+		/// <returns></returns>
+		public static bool IsValidLongPathName(StringView name)
+		{
+			if (name.Span.Length < 1 || name.Span[0] != '/')
+			{
+				return false;
+			}
+			if (!name.Span.ContainsAny(s_pathNameDelimiters))
+			{
+				return false;
+			}
+
+			// Reject paths that contain two consecutive delimiters in any position 
+			for (ReadOnlySpan<char> spanIt = name.Span; !spanIt.IsEmpty; )
+			{
+				int index = spanIt.IndexOfAny(s_pathNameDelimiters);
+				if (index < 0)
+				{
+					break;
+				}
+				// Note that we don't reject a trailing delimiter, like FSoftObjectPath
+				if (index + 1 < spanIt.Length && s_pathNameDelimiters.Contains(spanIt[index+1]))
+				{
+					return false;
+				}
+				spanIt = spanIt[(index + 1)..];
+			}
+			// Code in FSoftObjectPath continues to parse after this point but will no longer reject any strings
+			return true;
+		}
+
+	 	/// <summary>
+		/// Returns whether the given string is a valid long path name for an object.
+		/// See FSoftObjectPath::SetPath for reference.
+		/// e.g. MyClass is a short name, but /Script/MyModule.MyClass is a full path name.
+		/// </summary>
+		/// <param name="name"></param>
+		/// <returns></returns>
+		public static bool IsValidLongPathName(string name)
+		{
+			return IsValidLongPathName(new StringView(name));
+		}
 	}
 
 	/// <summary>
