UnrealBuildTool: Fix c# project build commands not being added to vscode task list

#rnx

diff --git a/Engine/Source/Programs/Shared/EpicGames.Core/CsProjectInfo.cs b/Engine/Source/Programs/Shared/EpicGames.Core/CsProjectInfo.cs
--- a/Engine/Source/Programs/Shared/EpicGames.Core/CsProjectInfo.cs
+++ b/Engine/Source/Programs/Shared/EpicGames.Core/CsProjectInfo.cs
@@ -95,13 +95,14 @@ namespace EpicGames.Core
 		}
 
 		/// <summary>
-		/// Returns whether or not this project is a netcore project (as opposed to full .net sdk) project
+		/// Returns whether or not this project is a modern dotnet project
 		/// </summary>
 		/// <returns>True if this is a netcore project</returns>
 		private bool IsNetCoreProject()
 		{
 			string? framework;
-			return Properties.TryGetValue("TargetFramework", out framework) && framework.StartsWith("netcoreapp", StringComparison.Ordinal);
+			return Properties.TryGetValue("TargetFramework", out framework)
+				&& (framework.StartsWith("netcoreapp", StringComparison.Ordinal) || Regex.Match(framework, @"net\d+\.0").Success);
 		}
 
 		/// <summary>
@@ -119,7 +120,9 @@ namespace EpicGames.Core
 			else if (IsNetCoreProject())
 			{
 				string configuration = Properties.TryGetValue("Configuration", out string? value) ? value : "Development";
-				return DirectoryReference.Combine(baseDirectory, "bin", configuration, Properties["TargetFramework"]);
+				return !Properties.TryGetValue("AppendTargetFrameworkToOutputPath", out string? appendFrameworkStr) || appendFrameworkStr.Equals("true", StringComparison.OrdinalIgnoreCase)
+					? DirectoryReference.Combine(baseDirectory, "bin", configuration, Properties["TargetFramework"])
+					: DirectoryReference.Combine(baseDirectory, "bin", configuration);
 			}
 			else
 			{
diff --git a/Engine/Source/Programs/UnrealBuildTool/ProjectFiles/VisualStudioCode/VSCodeProjectFileGenerator.cs b/Engine/Source/Programs/UnrealBuildTool/ProjectFiles/VisualStudioCode/VSCodeProjectFileGenerator.cs
--- a/Engine/Source/Programs/UnrealBuildTool/ProjectFiles/VisualStudioCode/VSCodeProjectFileGenerator.cs
+++ b/Engine/Source/Programs/UnrealBuildTool/ProjectFiles/VisualStudioCode/VSCodeProjectFileGenerator.cs
@@ -542,7 +542,6 @@ namespace UnrealBuildTool
 					{
 						CsProjectInfo? Info = VCSharpProject.GetProjectInfo(Config)!;
 
-						if (Info.Properties.ContainsKey("OutputPath"))
 						{
 							ProjectData.EOutputType OutputType;
 							string? OutputTypeName;
@@ -560,9 +559,10 @@ namespace UnrealBuildTool
 								OutputType = ProjectData.EOutputType.Exe;
 							}
 
+							DirectoryReference OutputPath = Info.GetOutputDir(VCSharpProject.ProjectFilePath.Directory);
 							FileReference? OutputFile = null;
 							HashSet<FileReference> ProjectBuildProducts = new HashSet<FileReference>();
-							Info.FindCompiledBuildProducts(DirectoryReference.Combine(VCSharpProject.ProjectFilePath.Directory, Info.Properties["OutputPath"]), ProjectBuildProducts);
+							Info.FindCompiledBuildProducts(OutputPath, ProjectBuildProducts);
 							foreach (FileReference ProjectBuildProduct in ProjectBuildProducts)
 							{
 								if ((OutputType == ProjectData.EOutputType.Exe && ProjectBuildProduct.GetExtension() == FrameworkExecutableExtension) ||
