EpicGames.Core: Normalize paths in CsProjectInfo to fix pathing issue on Mac

#rnx

diff --git a/Engine/Source/Programs/Shared/EpicGames.Core/CsProjectInfo.cs b/Engine/Source/Programs/Shared/EpicGames.Core/CsProjectInfo.cs
--- a/Engine/Source/Programs/Shared/EpicGames.Core/CsProjectInfo.cs
+++ b/Engine/Source/Programs/Shared/EpicGames.Core/CsProjectInfo.cs
@@ -560,7 +560,7 @@ namespace EpicGames.Core
 						}
 						break;
 					case "Compile":
-						// Reference to another project
+						// Reference to a file
 						if (EvaluateCondition(itemElement, projectInfo))
 						{
 							ParseCompileReference(baseDirectory, itemElement, projectInfo.CompileReferences);
@@ -568,7 +568,7 @@ namespace EpicGames.Core
 						break;
 					case "Content":
 					case "None":
-						// Reference to another project
+						// Reference to a file
 						if (EvaluateCondition(itemElement, projectInfo))
 						{
 							ParseContent(baseDirectory, itemElement, projectInfo.ContentReferences);
@@ -593,7 +593,11 @@ namespace EpicGames.Core
 				bool bEmbedInteropTypes = GetChildElementBoolean(parentElement, "EmbedInteropTypes", false);
 				if(!bEmbedInteropTypes)
 				{
-					FileReference assemblyFile = FileReference.Combine(baseDirectory, hintPath);
+					if (!OperatingSystem.IsWindows() && hintPath.Contains('\\', StringComparison.Ordinal))
+					{
+						hintPath = hintPath.Replace('\\', Path.DirectorySeparatorChar);
+					}
+					FileReference assemblyFile = Path.IsPathFullyQualified(hintPath) ? new FileReference(hintPath) : FileReference.Combine(baseDirectory, hintPath);
 					bool bPrivate = GetChildElementBoolean(parentElement, "Private", true);
 					references.Add(assemblyFile, bPrivate);
 				}
@@ -612,7 +616,12 @@ namespace EpicGames.Core
 			string? includePath = UnescapeString(parentElement.GetAttribute("Include"));
 			if (!String.IsNullOrEmpty(includePath))
 			{
-				FileReference projectFile = FileReference.Combine(baseDirectory, ExpandProperties(includePath, properties));
+				includePath = ExpandProperties(includePath, properties);
+				if (!OperatingSystem.IsWindows() && includePath.Contains('\\', StringComparison.Ordinal))
+				{
+					includePath = includePath.Replace('\\', Path.DirectorySeparatorChar);
+				}
+				FileReference projectFile = Path.IsPathFullyQualified(includePath) ? new FileReference(includePath) : FileReference.Combine(baseDirectory, includePath);
 				bool bPrivate = GetChildElementBoolean(parentElement, "Private", true);
 				projectReferences[projectFile] = bPrivate;
 			}
@@ -719,7 +728,7 @@ namespace EpicGames.Core
 		}
 
 		/// <summary>
-		/// Parses a project reference from a given 'ProjectReference' element
+		/// Parses a a compile file path from a given 'Compile' element
 		/// </summary>
 		/// <param name="baseDirectory">Directory to resolve relative paths against</param>
 		/// <param name="parentElement">The parent 'ProjectReference' element</param>
@@ -729,7 +738,11 @@ namespace EpicGames.Core
 			string? includePath = UnescapeString(parentElement.GetAttribute("Include"));
 			if (!String.IsNullOrEmpty(includePath))
 			{
-				FileReference sourceFile = FileReference.Combine(baseDirectory, includePath);
+				if (!OperatingSystem.IsWindows() && includePath.Contains('\\', StringComparison.Ordinal))
+				{
+					includePath = includePath.Replace('\\', Path.DirectorySeparatorChar);
+				}
+				FileReference sourceFile = Path.IsPathFullyQualified(includePath) ? new FileReference(includePath) : FileReference.Combine(baseDirectory, includePath);
 
 				if (sourceFile.FullName.Contains('*', StringComparison.Ordinal))
 				{
@@ -743,7 +756,7 @@ namespace EpicGames.Core
 		}
 
 		/// <summary>
-		/// Parses an assembly reference from a given 'Content' element
+		/// Parses a file path from a given 'Content' element
 		/// </summary>
 		/// <param name="baseDirectory">Directory to resolve relative paths against</param>
 		/// <param name="parentElement">The parent 'Content' element</param>
@@ -755,7 +768,11 @@ namespace EpicGames.Core
 			{
 				string? copyTo = GetChildElementString(parentElement, "CopyToOutputDirectory", null);
 				bool shouldCopy = !String.IsNullOrEmpty(copyTo) && (copyTo.Equals("Always", StringComparison.OrdinalIgnoreCase) || copyTo.Equals("PreserveNewest", StringComparison.OrdinalIgnoreCase));
-				FileReference contentFile = FileReference.Combine(baseDirectory, includePath);
+				if (!OperatingSystem.IsWindows() && includePath.Contains('\\', StringComparison.Ordinal))
+				{
+					includePath = includePath.Replace('\\', Path.DirectorySeparatorChar);
+				}
+				FileReference contentFile = Path.IsPathFullyQualified(includePath) ? new FileReference(includePath) : FileReference.Combine(baseDirectory, includePath);
 				contents.TryAdd(contentFile, shouldCopy);
 			}
 		}
