Fixed issue in UHT where it was allowing a USTRUCT to have multiple USTRUCTs as base classes.

diff --git a/Engine/Source/Programs/Shared/EpicGames.UHT/Types/UhtScriptStruct.cs b/Engine/Source/Programs/Shared/EpicGames.UHT/Types/UhtScriptStruct.cs
--- a/Engine/Source/Programs/Shared/EpicGames.UHT/Types/UhtScriptStruct.cs
+++ b/Engine/Source/Programs/Shared/EpicGames.UHT/Types/UhtScriptStruct.cs
@@ -757,6 +757,7 @@ namespace EpicGames.UHT.Types
 		public override void BindSuperAndBases()
 		{
 			BindSuper(SuperIdentifier, UhtFindOptions.ScriptStruct);
+			BindBases(BaseIdentifiers, UhtFindOptions.ScriptStruct | UhtFindOptions.Class); // We will generate an error if any are found
 			base.BindSuperAndBases();
 		}
 
@@ -890,6 +891,17 @@ namespace EpicGames.UHT.Types
 				}
 			}
 
+			// Make sure that we don't have too many base structs or classes
+			foreach (UhtStruct baseStruct in Bases)
+			{
+				bool isInterface = baseStruct is UhtClass classObj && classObj.ClassType == UhtClassType.NativeInterface;
+				if (!isInterface)
+				{
+					this.LogError("Structure can only derive from a single USTRUCT.  Native interfaces are not supported by the reflection system.");
+					break;
+				}
+			}
+
 			// Validate the engine name
 			if (!IsVerseField)
 			{
diff --git a/Engine/Source/Programs/Shared/EpicGames.UHT/Utils/UhtSession.cs b/Engine/Source/Programs/Shared/EpicGames.UHT/Utils/UhtSession.cs
--- a/Engine/Source/Programs/Shared/EpicGames.UHT/Utils/UhtSession.cs
+++ b/Engine/Source/Programs/Shared/EpicGames.UHT/Utils/UhtSession.cs
@@ -1497,21 +1497,20 @@ namespace EpicGames.UHT.Utils
 			{
 				throw new UhtIceException("Empty identifier array");
 			}
-			if (identifiers.Length > 2)
+			if (identifiers.Length > 3 || (identifiers.Length > 1 && !identifiers[1].Value.Span.Equals("::", StringComparison.Ordinal)))
 			{
 				if (messageSite != null)
 				{
-					//messageSite.LogError(lineNumber, "UnrealHeaderTool only supports C++ identifiers of two or less identifiers");
-					string fullIdentifier = String.Join("::", identifiers);
+					string fullIdentifier = String.Join(null, identifiers);
 					FindTypeError(messageSite, lineNumber, options, fullIdentifier);
-					return null;
 				}
+				return null;
 			}
 
 			UhtType? type;
-			if (identifiers.Length == 0)
+			if (identifiers.Length == 3)
 			{
-				type = FindTypeTwoNamesInternal(startingType, options, identifiers[0].Value.ToString(), identifiers[1].Value.ToString());
+				type = FindTypeTwoNamesInternal(startingType, options, identifiers[0].Value.ToString(), identifiers[2].Value.ToString());
 			}
 			else
 			{
