Updating UHT to once again generate an error when a USTRUCT contains more than one base.

diff --git a/Engine/Source/Programs/Shared/EpicGames.UHT/Types/UhtScriptStruct.cs b/Engine/Source/Programs/Shared/EpicGames.UHT/Types/UhtScriptStruct.cs
--- a/Engine/Source/Programs/Shared/EpicGames.UHT/Types/UhtScriptStruct.cs
+++ b/Engine/Source/Programs/Shared/EpicGames.UHT/Types/UhtScriptStruct.cs
@@ -757,7 +757,6 @@ namespace EpicGames.UHT.Types
 		public override void BindSuperAndBases()
 		{
 			BindSuper(SuperIdentifier, UhtFindOptions.ScriptStruct);
-			BindBases(BaseIdentifiers, UhtFindOptions.ScriptStruct | UhtFindOptions.Class); // We will generate an error if any are found
 			base.BindSuperAndBases();
 		}
 
@@ -892,14 +891,9 @@ namespace EpicGames.UHT.Types
 			}
 
 			// Make sure that we don't have too many base structs or classes
-			foreach (UhtStruct baseStruct in Bases)
+			if (BaseIdentifiers != null && BaseIdentifiers.Count > 0)
 			{
-				bool isInterface = baseStruct is UhtClass classObj && classObj.ClassType == UhtClassType.NativeInterface;
-				if (!isInterface)
-				{
-					this.LogError("Structure can only derive from a single USTRUCT.  Native interfaces are not supported by the reflection system.");
-					break;
-				}
+				this.LogError("USTRUCTs can only have one USTRUCT base.  Wrap any extra bases in a '#if CPP' block.  Replication will not be supported on any extra bases.");
 			}
 
 			// Validate the engine name
