[UBA]
* Nothing to see here.. gah, did a tiny change just before submitting

diff --git a/Engine/Source/Programs/Shared/EpicGames.UBA/Impl/ConfigImpl.cs b/Engine/Source/Programs/Shared/EpicGames.UBA/Impl/ConfigImpl.cs
--- a/Engine/Source/Programs/Shared/EpicGames.UBA/Impl/ConfigImpl.cs
+++ b/Engine/Source/Programs/Shared/EpicGames.UBA/Impl/ConfigImpl.cs
@@ -43,7 +43,7 @@ namespace EpicGames.UBA.Impl
 			if (_handle != IntPtr.Zero)
 			{
 				//Config_Destroy(_handle); // instance acquired by Config_Load is owned by native side
-				_handle = IntPtr.Zero;
+				//_handle = IntPtr.Zero;
 			}
 		}
 		#endregion
