[UBA]
* Fixed C# static analyzer error in SchedulerImpl.cs

diff --git a/Engine/Source/Programs/Shared/EpicGames.UBA/Impl/SchedulerImpl.cs b/Engine/Source/Programs/Shared/EpicGames.UBA/Impl/SchedulerImpl.cs
--- a/Engine/Source/Programs/Shared/EpicGames.UBA/Impl/SchedulerImpl.cs
+++ b/Engine/Source/Programs/Shared/EpicGames.UBA/Impl/SchedulerImpl.cs
@@ -74,7 +74,7 @@ namespace EpicGames.UBA.Impl
 
 		void RaiseExited(nint userData, nint processHandle, ref byte exitedResponse)
 		{
-			while (_index == uint.MaxValue) { Thread.Sleep(1); } // This to prevent a (most likely impossible) race condition (Index is set after process is enqueued)
+			while (_index == UInt32.MaxValue) { Thread.Sleep(1); } // This to prevent a (most likely impossible) race condition (Index is set after process is enqueued)
 
 			//ProcessImpl process = new(handle);
 			ProcessFinishedInfoImpl info = new(_userData, processHandle);
@@ -87,7 +87,7 @@ namespace EpicGames.UBA.Impl
 		object _userData;
 		event ExitCallback _exitCallbackDelegate;
 		public nint _startInfoHandle;
-		public uint _index = uint.MaxValue;
+		public uint _index = UInt32.MaxValue;
 	}
 
 	internal class SchedulerImpl : IScheduler
@@ -126,12 +126,12 @@ namespace EpicGames.UBA.Impl
 		{
 			nint[] cc = new nint[cacheClients.Count()];
 			uint cacheClientCount = 0;
-			foreach (var client in cacheClients)
+			foreach (ICacheClient client in cacheClients)
 			{
 				cc[cacheClientCount++] = client.GetHandle();
 			}
 
-			IConfig config = IConfig.LoadConfig("");
+			using IConfig config = IConfig.LoadConfig("");
 			config.AddValue("Scheduler", "MaxLocalProcessors", maxLocalProcessors.ToString());
 			if (forceRemote)
 			{
