Refactor the skip declaration code to better handle macros with API declarations before them.

#rnx

diff --git a/Engine/Source/Programs/Shared/EpicGames.UHT/Parsers/UhtFunctionParser.cs b/Engine/Source/Programs/Shared/EpicGames.UHT/Parsers/UhtFunctionParser.cs
--- a/Engine/Source/Programs/Shared/EpicGames.UHT/Parsers/UhtFunctionParser.cs
+++ b/Engine/Source/Programs/Shared/EpicGames.UHT/Parsers/UhtFunctionParser.cs
@@ -539,8 +539,7 @@ namespace EpicGames.UHT.Parsers
 						else if (topScope.TokenReader.TryPeekOptional('{'))
 						{
 							// Skip inline function bodies
-							UhtToken tokenCopy = new();
-							topScope.TokenReader.SkipDeclaration(ref tokenCopy);
+							topScope.SkipDeclaration(new());
 						}
 					}
 				}
diff --git a/Engine/Source/Programs/Shared/EpicGames.UHT/Parsers/UhtHeaderFileParser.cs b/Engine/Source/Programs/Shared/EpicGames.UHT/Parsers/UhtHeaderFileParser.cs
--- a/Engine/Source/Programs/Shared/EpicGames.UHT/Parsers/UhtHeaderFileParser.cs
+++ b/Engine/Source/Programs/Shared/EpicGames.UHT/Parsers/UhtHeaderFileParser.cs
@@ -616,7 +616,7 @@ namespace EpicGames.UHT.Parsers
 				if (topScope.ScopeType is UhtClass classObj)
 				{
 					using UhtTokenRecorder recorder = new(topScope, ref token);
-					topScope.TokenReader.SkipDeclaration(ref token);
+					topScope.SkipDeclaration(token);
 					if (recorder.Stop())
 					{
 						if (classObj.Declarations != null)
@@ -639,7 +639,7 @@ namespace EpicGames.UHT.Parsers
 				}
 				else
 				{
-					topScope.TokenReader.SkipDeclaration(ref token);
+					topScope.SkipDeclaration(token);
 				}
 			}
 			return true;
diff --git a/Engine/Source/Programs/Shared/EpicGames.UHT/Parsers/UhtNativeInterfaceClassParser.cs b/Engine/Source/Programs/Shared/EpicGames.UHT/Parsers/UhtNativeInterfaceClassParser.cs
--- a/Engine/Source/Programs/Shared/EpicGames.UHT/Parsers/UhtNativeInterfaceClassParser.cs
+++ b/Engine/Source/Programs/Shared/EpicGames.UHT/Parsers/UhtNativeInterfaceClassParser.cs
@@ -34,7 +34,7 @@ namespace EpicGames.UHT.Parsers
 					return UhtParseResult.Handled;
 				}
 			}
-			return topScope.TokenReader.SkipDeclaration(ref token) ? UhtParseResult.Handled : UhtParseResult.Invalid;
+			return UhtParseResult.Unhandled;
 		}
 
 		[UhtKeyword(Extends = UhtTableNames.NativeInterface)]
diff --git a/Engine/Source/Programs/Shared/EpicGames.UHT/Tokenizer/UhtTokenReaderUtilityExtensions.cs b/Engine/Source/Programs/Shared/EpicGames.UHT/Tokenizer/UhtTokenReaderUtilityExtensions.cs
--- a/Engine/Source/Programs/Shared/EpicGames.UHT/Tokenizer/UhtTokenReaderUtilityExtensions.cs
+++ b/Engine/Source/Programs/Shared/EpicGames.UHT/Tokenizer/UhtTokenReaderUtilityExtensions.cs
@@ -3,6 +3,7 @@
 using System;
 using System.Collections.Generic;
 using EpicGames.Core;
+using EpicGames.UHT.Parsers;
 using EpicGames.UHT.Utils;
 
 namespace EpicGames.UHT.Tokenizer
@@ -49,146 +50,133 @@ namespace EpicGames.UHT.Tokenizer
 			return false;
 		}
 
+		/// <summary>
+		/// Check to see if we have reached an identifier that is a UE macro
+		/// </summary>
+		/// <param name="messageSite">Used to generate error message</param>
+		/// <param name="token">Token being tested</param>
+		/// <param name="generateError">If true, generate an error.</param>
+		/// <returns>True if a UE macro was found</returns>
+		private static bool CheckForUEMacro(IUhtMessageSite messageSite, ref UhtToken token, bool generateError)
+		{
+			ReadOnlySpan<char> span = token.Value.Span;
+			if (span[0] == 'G' || span[0] == 'R' || span[0] == 'U')
+			{
+				if (s_skipDeclarationWarningStrings.Contains(token.Value))
+				{
+					if (generateError)
+					{
+						messageSite.LogWarning($"The identifier \'{token.Value}\' was detected in a block being skipped. Was this intentional?");
+					}
+					return true;
+				}
+			}
+			return false;
+		}
+
 		/// <summary>
 		/// Given a declaration/statement that starts with the given token, skip the declaration in the header.
 		/// </summary>
-		/// <param name="tokenReader">Token reader</param>
+		/// <param name="topScope">Current scope being processed</param>
 		/// <param name="token">Token that started the process</param>
 		/// <returns>true if there could be more header to process, false if the end was reached.</returns>
-		public static bool SkipDeclaration(this IUhtTokenReader tokenReader, ref UhtToken token)
+		public static void SkipDeclaration(this UhtParsingScope topScope, UhtToken token)
 		{
-			// Consume all tokens until the end of declaration/definition has been found.
-			int nestedScopes = 0;
-			bool endOfDeclarationFound = false;
+			IUhtTokenReader tokenReader = topScope.TokenReader;
 
 			// Store the current value of PrevComment so it can be restored after we parsed everything.
+			using UhtTokenDisableComments disableComments = new(tokenReader);
+
+			// Processing must be adjusted if we think we are processing a class or struct.
+			// Macros are tricky because they usually don't end in a ";".  So we have to terminate 
+			// when we find the ')'.
+			// Type definitions might have a trailing variable declaration.
+			bool mightBeMacro = false;
+			bool mightBeClassOrStruct = false;
+			if (token.IsIdentifier())
 			{
-				using UhtTokenDisableComments disableComments = new(tokenReader);
-
-				// Check if this is a class/struct declaration in which case it can be followed by member variable declaration.	
-				bool possiblyClassDeclaration = token.IsIdentifier() && (token.IsValue("class") || token.IsValue("struct"));
+				mightBeMacro = true;
+				if (token.IsValue("class") || token.IsValue("struct"))
+				{
+					mightBeClassOrStruct = true;
+				}
+			}
 
-				// (known) macros can end without ; or } so use () to find the end of the declaration.
-				// However, we don't want to use it with DECLARE_FUNCTION, because we need it to be treated like a function.
-				bool macroDeclaration = IsProbablyAMacro(token.Value) && !token.IsIdentifier("DECLARE_FUNCTION");
+			bool peekedToken = false;
+			while (!token.TokenType.IsEndType())
+			{
+				if (peekedToken)
+				{
 
-				bool definitionFound = false;
-				char openingBracket = macroDeclaration ? '(' : '{';
-				char closingBracket = macroDeclaration ? ')' : '}';
+					// For the moment, warn on anything skipped.  
+					CheckForUEMacro(tokenReader, ref token, true);
+					//if (CheckForUEMacro(tokenReader, ref token, false))
+					//{
+					//	break;
+					//}
 
-				bool retestCurrentToken = false;
-				while (true)
-				{
-					if (!retestCurrentToken)
+					// Check for poor termination of a macro
+					if (token.IsSymbol('}'))
 					{
-						token = tokenReader.GetToken();
-						if (token.TokenType.IsEndType())
-						{
-							break;
-						}
+						break;
 					}
-					else
+					tokenReader.ConsumeToken();
+				}
+
+				if (token.IsIdentifier())
+				{
+					if (token.IsValue("PURE_VIRTUAL"))
 					{
-						retestCurrentToken = false;
+						tokenReader.SkipBrackets('(', ')', 0, (x) => { CheckForUEMacro(tokenReader, ref x, true); });
+						break;
 					}
-					ReadOnlySpan<char> span = token.Value.Span;
-
-					// If this is a macro, consume it
-					// If we find parentheses at top-level and we think it's a class declaration then it's more likely
-					// to be something like: class UThing* GetThing();
-					if (possiblyClassDeclaration && nestedScopes == 0 && token.IsSymbol() && span.Length == 1 && span[0] == '(')
+					else if (mightBeMacro && (!IsProbablyAMacro(token.Value) || token.IsValue("DECLARE_FUNCTION")))
 					{
-						possiblyClassDeclaration = false;
+						mightBeMacro = false;
 					}
-
-					if (token.IsSymbol() && span.Length == 1 && span[0] == ';' && nestedScopes == 0)
+				}
+				else if (token.IsSymbol())
+				{
+					if (token.IsValue(';'))
 					{
-						endOfDeclarationFound = true;
 						break;
 					}
-
-					if (token.IsIdentifier())
+					else if (token.IsValue('('))
 					{
-						// Use a trivial pre-filter to avoid doing the search on things that aren't UE keywords we care about
-						if (span[0] == 'G' || span[0] == 'R' || span[0] == 'U')
+						mightBeClassOrStruct = false; // some type of function declaration
+						tokenReader.SkipBrackets('(', ')', 1, (x) => { CheckForUEMacro(tokenReader, ref x, true); });
+						// Could be a class declaration in all capitals, and not a macro
+						if (mightBeMacro && !tokenReader.TryPeekOptional('{'))
 						{
-							if (s_skipDeclarationWarningStrings.Contains(token.Value))
-							{
-								tokenReader.LogWarning($"The identifier \'{token.Value}\' was detected in a block being skipped. Was this intentional?");
-							}
+							break;
 						}
 					}
-
-					if (!macroDeclaration && token.IsIdentifier() && span.Equals("PURE_VIRTUAL", StringComparison.Ordinal) && nestedScopes == 0)
-					{
-						openingBracket = '(';
-						closingBracket = ')';
-					}
-
-					if (token.IsSymbol() && span.Length == 1 && span[0] == openingBracket)
-					{
-						// This is a function definition or class declaration.
-						definitionFound = true;
-						nestedScopes++;
-					}
-					else if (token.IsSymbol() && span.Length == 1 && span[0] == closingBracket)
+					else if (token.IsValue('{'))
 					{
-						nestedScopes--;
-						if (nestedScopes == 0)
+						tokenReader.SkipBrackets('{', '}', 1, (x) => { CheckForUEMacro(tokenReader, ref x, true); });
+						if (mightBeClassOrStruct)
 						{
-							// Could be a class declaration in all capitals, and not a macro
-							bool reallyEndDeclaration = true;
-							if (macroDeclaration)
-							{
-								reallyEndDeclaration = !tokenReader.TryPeekOptional('{');
-							}
-
-							if (reallyEndDeclaration)
+							if (tokenReader.TryOptionalIdentifier())
 							{
-								endOfDeclarationFound = true;
-								break;
+								tokenReader.Require(';');
 							}
 						}
-
-						if (nestedScopes < 0)
-						{
-							throw new UhtException(tokenReader, token.InputLine, $"Unexpected '{closingBracket}'. Did you miss a semi-colon?");
-						}
-					}
-					else if (macroDeclaration && nestedScopes == 0)
-					{
-						macroDeclaration = false;
-						openingBracket = '{';
-						closingBracket = '}';
-						retestCurrentToken = true;
+						break;
 					}
+					mightBeMacro = false;
 				}
-				if (endOfDeclarationFound)
+				else
 				{
-					// Member variable declaration after class declaration (see bPossiblyClassDeclaration).
-					if (possiblyClassDeclaration && definitionFound)
-					{
-						// Should syntax errors be also handled when someone declares a variable after function definition?
-						// Consume the variable name.
-						if (tokenReader.IsEOF)
-						{
-							return true;
-						}
-						if (tokenReader.TryOptionalIdentifier())
-						{
-							tokenReader.Require(';');
-						}
-					}
-
-					// C++ allows any number of ';' after member declaration/definition.
-					while (tokenReader.TryOptional(';'))
-					{
-					}
+					mightBeMacro = false;
 				}
+				peekedToken = true;
+				token = tokenReader.PeekToken();
 			}
 
-			// Successfully consumed C++ declaration unless mismatched pair of brackets has been found.
-			return nestedScopes == 0 && endOfDeclarationFound;
+			// C++ allows any number of ';' after member declaration/definition.
+			while (tokenReader.TryOptional(';'))
+			{
+			}
 		}
 
 		private static bool IsProbablyAMacro(StringView identifier)
