Horde : Fix using directive is unnecessary warning

#rnx

diff --git a/Engine/Source/Programs/Shared/EpicGames.Horde/Issues/Handlers/LocalizationIssueHandler.cs b/Engine/Source/Programs/Shared/EpicGames.Horde/Issues/Handlers/LocalizationIssueHandler.cs
--- a/Engine/Source/Programs/Shared/EpicGames.Horde/Issues/Handlers/LocalizationIssueHandler.cs
+++ b/Engine/Source/Programs/Shared/EpicGames.Horde/Issues/Handlers/LocalizationIssueHandler.cs
@@ -1,6 +1,5 @@
 // Copyright Epic Games, Inc. All Rights Reserved.
 
-using System;
 using System.Collections.Generic;
 using System.Text.Json;
 using EpicGames.Core;
