Removing some of the magic name generation in UHT Verse support.  Also removed the _RetVal in the VNI code generation since it wasn't being used.

#rnx
#okforversepublic

diff --git a/Engine/Source/Programs/Shared/EpicGames.UHT/Specifiers/UhtFieldSpecifiers.cs b/Engine/Source/Programs/Shared/EpicGames.UHT/Specifiers/UhtFieldSpecifiers.cs
--- a/Engine/Source/Programs/Shared/EpicGames.UHT/Specifiers/UhtFieldSpecifiers.cs
+++ b/Engine/Source/Programs/Shared/EpicGames.UHT/Specifiers/UhtFieldSpecifiers.cs
@@ -85,29 +85,7 @@ namespace EpicGames.UHT.Parsers
 			// If no name was specified, generate one using the following algorithm FVerseXxYyZz -> xx_yy_zz
 			if (String.IsNullOrEmpty(fieldObj.VerseName))
 			{
-				ReadOnlySpan<char> span = StripVersePrefix(fieldObj);
-				if (!span.IsEmpty)
-				{
-					using BorrowStringBuilder borrowStringBuilder = new BorrowStringBuilder(StringBuilderCache.Small);
-					bool first = true;
-					foreach (char c in span)
-					{
-						if (c >= 'A' && c <= 'Z')
-						{
-							if (!first)
-							{
-								borrowStringBuilder.StringBuilder.Append('_');
-							}
-							borrowStringBuilder.StringBuilder.Append(Char.ToLower(c));
-						}
-						else
-						{
-							borrowStringBuilder.StringBuilder.Append(c);
-						}
-						first = false;
-					}
-					fieldObj.VerseName = borrowStringBuilder.StringBuilder.ToString();
-				}
+				fieldObj.LogError($"A verse name must be specified containing the name of the verse type");
 			}
 
 			// Final validation
