UnrealBuildTool: Ensure clang and clang++ errors are matched

diff --git a/Engine/Source/Programs/Shared/UnrealBuildTool.Tests/LogEventMatcherTests.cs b/Engine/Source/Programs/Shared/UnrealBuildTool.Tests/LogEventMatcherTests.cs
--- a/Engine/Source/Programs/Shared/UnrealBuildTool.Tests/LogEventMatcherTests.cs
+++ b/Engine/Source/Programs/Shared/UnrealBuildTool.Tests/LogEventMatcherTests.cs
@@ -727,6 +727,19 @@ namespace UnrealBuildTool.Tests
 			CheckEventGroup(logEvents, 0, 1, LogLevel.Error, KnownLogEvents.Linker);
 		}
 
+		[TestMethod]
+		public void LinkEventMatcher6()
+		{
+			string[] lines =
+			[
+				@"clang++: error: clang frontend command failed due to signal (use -v to see invocation)"
+			];
+
+			List<LogEvent> logEvents = Parse(String.Join("\n", lines));
+			Assert.AreEqual(1, logEvents.Count);
+			CheckEventGroup(logEvents, 0, 1, LogLevel.Error, KnownLogEvents.Linker);
+		}
+
 		[TestMethod]
 		public void MacLinkErrorMatcher()
 		{
diff --git a/Engine/Source/Programs/UnrealBuildTool/Matchers/LinkEventMatcher.cs b/Engine/Source/Programs/UnrealBuildTool/Matchers/LinkEventMatcher.cs
--- a/Engine/Source/Programs/UnrealBuildTool/Matchers/LinkEventMatcher.cs
+++ b/Engine/Source/Programs/UnrealBuildTool/Matchers/LinkEventMatcher.cs
@@ -19,10 +19,10 @@ namespace UnrealBuildTool.Matchers
 			@"^(\s*)Undefined symbols for architecture");
 
 		static readonly Regex s_ldClangPattern = new Regex(
-			@"^\s*(ld|ld.lld|lld-link|clang):");
+			@"^\s*(ld|ld.lld|lld-link|clang|clang\+\+):");
 
 		static readonly Regex s_ldClangErrorPattern = new Regex(
-			@"^\s*(ld|ld.lld|lld-link):\s+error:\s+(?<error>.+)$");
+			@"^\s*(ld|ld.lld|lld-link|clang|clang\+\+):\s+error:\s+(?<error>.+)$");
 
 		static readonly Regex s_ldMultiplyDefinedPattern = new Regex(
 			@"(?<severity>error): L\d+: symbol `(?<symbol>[^`]+)' multiply defined$");
