UHT will flag U macros (i.e. UCLASS) found in skipped blocks of code as deprecated.  This will be made  a warning in a future version of UE.

diff --git a/Engine/Source/Programs/Shared/EpicGames.UHT/Tokenizer/UhtTokenReaderUtilityExtensions.cs b/Engine/Source/Programs/Shared/EpicGames.UHT/Tokenizer/UhtTokenReaderUtilityExtensions.cs
--- a/Engine/Source/Programs/Shared/EpicGames.UHT/Tokenizer/UhtTokenReaderUtilityExtensions.cs
+++ b/Engine/Source/Programs/Shared/EpicGames.UHT/Tokenizer/UhtTokenReaderUtilityExtensions.cs
@@ -13,23 +13,27 @@ namespace EpicGames.UHT.Tokenizer
 	/// </summary>
 	public static class UhtTokenReaderUtilityExtensions
 	{
-		private static readonly HashSet<StringView> s_skipDeclarationWarningStrings = new()
-		{
+		private static readonly HashSet<StringView> s_skipDeclarationWarningStrings =
+		[
 			"GENERATED_BODY",
 			"GENERATED_IINTERFACE_BODY",
 			"GENERATED_UCLASS_BODY",
 			"GENERATED_UINTERFACE_BODY",
 			"GENERATED_USTRUCT_BODY",
-			// Leaving these disabled ATM since they can exist in the code without causing compile issues
-			//"RIGVM_METHOD",
-			//"UCLASS",
-			//"UDELEGATE",
-			//"UENUM",
-			//"UFUNCTION",
-			//"UINTERFACE",
-			//"UPROPERTY",
-			//"USTRUCT",
-		};
+		];
+
+		private static readonly HashSet<StringView> s_skipTypeWarningStrings =
+		[
+			"RIGVM_METHOD",
+			"UCLASS",
+			"UDELEGATE",
+			"UENUM",
+			"UFUNCTION",
+			"UINTERFACE",
+			"UMETA",
+			"UPROPERTY",
+			"USTRUCT",
+		];
 
 		/// <summary>
 		/// Try to parse an optional _API macro
@@ -70,6 +74,14 @@ namespace EpicGames.UHT.Tokenizer
 					}
 					return true;
 				}
+				if (s_skipTypeWarningStrings.Contains(token.Value))
+				{
+					if (generateError)
+					{
+						messageSite.LogDeprecation($"The identifier \'{token.Value}\' was detected in a block being skipped and is now being flagged. This will be a warning in a future version of Unreal.");
+					}
+					return true;
+				}
 			}
 			return false;
 		}
