Adding support for namespace in UHT.

To enabled namespace support in a module, set the "bAllowUETypesInNamespaces" value to true in the module's build file.
The default is disabled but may default to enabled in a future version of UE.

Placing a UE type inside of a namespace does not change the UE type name at runtime.  Two UE types that exist in different
C++ namespaces but share the same engine name will be considered an error by UHT.   Changing the C++ namespace of a UE type
will not impact saved data since the engine name will not change.

When enabled, any UE type can be parsed inside of a namespace with the exception of the "DECLARE_DYNAMIC_*" macros.
Inside of a namespace, a "DECLARE_DYNAMIC_*" macro must maked with a UDELEGATE() macro.  This was done to prevent
any "DECLARE_DYNAMIC_*" already appearing in a C++ namespace from unexpectedly becoming a UE type when
namespace are enabled.

Modules that do not have namespaces enabled can reference modules that do have namespaces enabled.

Along with namespace support, UHT will now warn on UE type macros that appear in skipped code.

diff --git a/Engine/Source/Programs/Shared/EpicGames.UHT/Exporters/CodeGen/UhtHeaderCodeGeneratorCppFile.cs b/Engine/Source/Programs/Shared/EpicGames.UHT/Exporters/CodeGen/UhtHeaderCodeGeneratorCppFile.cs
--- a/Engine/Source/Programs/Shared/EpicGames.UHT/Exporters/CodeGen/UhtHeaderCodeGeneratorCppFile.cs
+++ b/Engine/Source/Programs/Shared/EpicGames.UHT/Exporters/CodeGen/UhtHeaderCodeGeneratorCppFile.cs
@@ -273,6 +273,7 @@ namespace EpicGames.UHT.Exporters.CodeGen
 							uint hash = ObjectInfos[scriptStruct.ObjectTypeIndex].Hash;
 							builder
 								.Append("\t\t{ ")
+								.Append(scriptStruct.Namespace.FullSourceName)
 								.Append(scriptStruct.SourceName)
 								.Append("::StaticStruct, Z_Construct_UScriptStruct_")
 								.Append(scriptStruct.SourceName)
@@ -281,6 +282,7 @@ namespace EpicGames.UHT.Exporters.CodeGen
 								.Append("\"), &Z_Registration_Info_UScriptStruct_")
 								.Append(scriptStruct.SourceName)
 								.Append(", CONSTRUCT_RELOAD_VERSION_INFO(FStructReloadVersionInfo, sizeof(")
+								.Append(scriptStruct.NamespaceExportName)
 								.Append(scriptStruct.SourceName)
 								.Append($"), {hash}U) }},\r\n");
 							combinedHash = HashCombine(combinedHash, hash);
@@ -296,12 +298,14 @@ namespace EpicGames.UHT.Exporters.CodeGen
 								.Append("\t\t{ Z_Construct_UClass_")
 								.Append(classObj.SourceName)
 								.Append(", ")
+								.Append(classObj.Namespace.FullSourceName)
 								.Append(classObj.SourceName)
 								.Append("::StaticClass, TEXT(\"")
 								.Append(classObj.IsVerseField ? classObj.EngineName : classObj.SourceName)
 								.Append("\"), &Z_Registration_Info_UClass_")
 								.Append(classObj.SourceName)
 								.Append(", CONSTRUCT_RELOAD_VERSION_INFO(FClassReloadVersionInfo, sizeof(")
+								.Append(classObj.Namespace.FullSourceName)
 								.Append(classObj.SourceName)
 								.Append($"), {hash}U) }},\r\n");
 							combinedHash = HashCombine(combinedHash, hash);
@@ -477,7 +481,7 @@ namespace EpicGames.UHT.Exporters.CodeGen
 						{
 							keyName = value.Name.ToString();
 						}
-						builder.Append("\t\t{ ").AppendUTF8LiteralString(keyName).Append(", (int64)").Append(value.Name).Append(" },\r\n");
+						builder.Append("\t\t{ ").AppendUTF8LiteralString(keyName).Append(", (int64)").Append(enumObj.Namespace.FullSourceName).Append(value.Name).Append(" },\r\n");
 						++enumIndex;
 					}
 					builder.Append("\t};\r\n");
@@ -570,7 +574,7 @@ namespace EpicGames.UHT.Exporters.CodeGen
 
 				// Outer singleton
 				builder.Append("static FStructRegistrationInfo ").Append(registrationName).Append(";\r\n");
-				builder.Append("class UScriptStruct* ").Append(scriptStruct.SourceName).Append("::StaticStruct()\r\n");
+				builder.Append("class UScriptStruct* ").Append(scriptStruct.Namespace.FullSourceName).Append(scriptStruct.SourceName).Append("::StaticStruct()\r\n");
 				builder.Append("{\r\n");
 				builder.Append("\tif (!").Append(registrationName).Append(".OuterSingleton)\r\n");
 				builder.Append("\t{\r\n");
@@ -675,7 +679,7 @@ namespace EpicGames.UHT.Exporters.CodeGen
 			int hashCodeBlockStart = builder.Length;
 
 			// Collect the properties
-			PropertyMemberContextImpl context = new(CodeGenerator, scriptStruct, scriptStruct.SourceName, staticsName);
+			PropertyMemberContextImpl context = new(CodeGenerator, scriptStruct, staticsName);
 			UhtUsedDefineScopes<UhtProperty> properties = new(scriptStruct.Properties);
 
 			// Declare the statics structure
@@ -686,8 +690,8 @@ namespace EpicGames.UHT.Exporters.CodeGen
 				foreach (UhtScriptStruct noExportStruct in noExportStructs)
 				{
 					AppendMirrorsForNoexportStruct(builder, noExportStruct, 1);
-					builder.Append("\tstatic_assert(sizeof(").Append(noExportStruct.SourceName).Append(") < MAX_uint16);\r\n");
-					builder.Append("\tstatic_assert(alignof(").Append(noExportStruct.SourceName).Append(") < MAX_uint8);\r\n");
+					builder.Append("\tstatic_assert(sizeof(").Append(noExportStruct.NamespaceExportName).Append(noExportStruct.SourceName).Append(") < MAX_uint16);\r\n");
+					builder.Append("\tstatic_assert(alignof(").Append(noExportStruct.NamespaceExportName).Append(noExportStruct.SourceName).Append(") < MAX_uint8);\r\n");
 				}
 
 				// Meta data
@@ -701,7 +705,7 @@ namespace EpicGames.UHT.Exporters.CodeGen
 				{
 					builder.Append("\tstatic void* NewStructOps()\r\n");
 					builder.Append("\t{\r\n");
-					builder.Append("\t\treturn (UScriptStruct::ICppStructOps*)new UScriptStruct::TCppStructOps<").Append(scriptStruct.SourceName).Append(">();\r\n");
+					builder.Append("\t\treturn (UScriptStruct::ICppStructOps*)new UScriptStruct::TCppStructOps<").Append(scriptStruct.NamespaceExportName).Append(scriptStruct.SourceName).Append(">();\r\n");
 					builder.Append("\t}\r\n");
 				}
 				builder.Append("\tstatic const UECodeGen_Private::").AppendStructParamsType(scriptStruct).Append(" StructParams;\r\n");
@@ -730,8 +734,8 @@ namespace EpicGames.UHT.Exporters.CodeGen
 				builder.Append('\t').AppendUTF8LiteralString(scriptStruct.EngineName).Append(",\r\n");
 				builder.AppendArrayPtrLine(properties, UhtDefineScopeNames.Standard, staticsName, "PropPointers", 1, ",\r\n");
 				builder.AppendArrayCountLine(properties, UhtDefineScopeNames.Standard, staticsName, "PropPointers", 1, ",\r\n");
-				builder.Append("\tsizeof(").Append(scriptStruct.SourceName).Append("),\r\n");
-				builder.Append("\talignof(").Append(scriptStruct.SourceName).Append("),\r\n");
+				builder.Append("\tsizeof(").Append(scriptStruct.NamespaceExportName).Append(scriptStruct.SourceName).Append("),\r\n");
+				builder.Append("\talignof(").Append(scriptStruct.NamespaceExportName).Append(scriptStruct.SourceName).Append("),\r\n");
 				builder.Append("\tRF_Public|RF_Transient|RF_MarkAsNative,\r\n");
 				builder.Append($"\tEStructFlags(0x{(uint)(scriptStruct.ScriptStructFlags & ~EStructFlags.ComputedFlags):X8}),\r\n");
 				builder.Append('\t').AppendMetaDataParams(scriptStruct, staticsName, MetaDataParamsName).Append("\r\n");
@@ -813,6 +817,7 @@ namespace EpicGames.UHT.Exporters.CodeGen
 					builder
 						.Append(methodInfo.ReturnType)
 						.Append(' ')
+						.Append(scriptStruct.Namespace.FullSourceName)
 						.Append(scriptStruct.SourceName)
 						.Append("::")
 						.Append(methodInfo.Name)
@@ -841,6 +846,7 @@ namespace EpicGames.UHT.Exporters.CodeGen
 					builder
 						.Append(methodInfo.ReturnType)
 						.Append(' ')
+						.Append(scriptStruct.Namespace.FullSourceName)
 						.Append(scriptStruct.SourceName)
 						.Append("::")
 						.Append(methodInfo.Name)
@@ -926,7 +932,7 @@ namespace EpicGames.UHT.Exporters.CodeGen
 			int hashCodeBlockStart = builder.Length;
 
 			// Collect the properties
-			PropertyMemberContextImpl context = new(CodeGenerator, function, eventParameters, staticsName);
+			PropertyMemberContextImpl context = new(CodeGenerator, function, "", eventParameters, staticsName);
 			UhtUsedDefineScopes<UhtProperty> properties = new(function.Properties);
 
 			// Statics declaration
@@ -1174,7 +1180,7 @@ namespace EpicGames.UHT.Exporters.CodeGen
 				}
 				if (isRpc)
 				{
-					builder.Append("DEFINE_FUNCTION(").AppendClassSourceNameOrInterfaceName(classObj).Append("::").Append(function.UnMarshalAndCallName).Append(")\r\n");
+					builder.Append("DEFINE_FUNCTION(").Append(classObj.Namespace.FullSourceName).AppendClassSourceNameOrInterfaceName(classObj).Append("::").Append(function.UnMarshalAndCallName).Append(")\r\n");
 					builder.Append("{\r\n");
 					AppendFunctionThunk(builder, function);
 					builder.Append("}\r\n");
@@ -1218,7 +1224,7 @@ namespace EpicGames.UHT.Exporters.CodeGen
 				builder.Append(isSuspends ? "TVerseTask<void> " : "void ");
 			}
 
-			builder.AppendClassSourceNameOrInterfaceName(classObj).Append("::").Append(function.SourceName).Append('(');
+			builder.Append(classObj.Namespace.FullSourceName).AppendClassSourceNameOrInterfaceName(classObj).Append("::").Append(function.SourceName).Append('(');
 			bool needsComma = false;
 			foreach (UhtType type in function.ParameterProperties.Span)
 			{
@@ -1800,11 +1806,11 @@ namespace EpicGames.UHT.Exporters.CodeGen
 					switch (GetConstructorType(classObj))
 					{
 						case ConstructorType.ObjectInitializer:
-							builder.Append(classObj.SourceName).Append("::").Append(classObj.SourceName).Append("(const FObjectInitializer& ObjectInitializer) : Super(ObjectInitializer) {}\r\n");
+							builder.Append(classObj.Namespace.FullSourceName).Append(classObj.SourceName).Append("::").Append(classObj.SourceName).Append("(const FObjectInitializer& ObjectInitializer) : Super(ObjectInitializer) {}\r\n");
 							break;
 
 						case ConstructorType.Default:
-							builder.Append(classObj.SourceName).Append("::").Append(classObj.SourceName).Append("() {}\r\n");
+							builder.Append(classObj.Namespace.FullSourceName).Append(classObj.SourceName).Append("::").Append(classObj.SourceName).Append("() {}\r\n");
 							break;
 					}
 				}
@@ -1812,12 +1818,12 @@ namespace EpicGames.UHT.Exporters.CodeGen
 
 			if (!classObj.ClassExportFlags.HasAnyFlags(UhtClassExportFlags.HasCustomVTableHelperConstructor))
 			{
-				builder.Append("DEFINE_VTABLE_PTR_HELPER_CTOR(").Append(classObj.SourceName).Append(");\r\n");
+				builder.Append($"DEFINE_VTABLE_PTR_HELPER_CTOR_NS({classObj.Namespace.FullSourceName}, {classObj.SourceName});\r\n");
 			}
 
 			if (!classObj.ClassExportFlags.HasAnyFlags(UhtClassExportFlags.HasDestructor) && !classObj.ClassFlags.HasAnyFlags(EClassFlags.Interface))
 			{
-				builder.Append(classObj.SourceName).Append("::~").Append(classObj.SourceName).Append("() {}\r\n");
+				builder.Append($"{classObj.Namespace.FullSourceName}{classObj.SourceName}::~{classObj.SourceName}() {{}}\r\n");
 			}
 
 			AppendFieldNotify(builder, classObj);
@@ -1873,7 +1879,7 @@ namespace EpicGames.UHT.Exporters.CodeGen
 					if (property.PropertyExportFlags.HasAnyFlags(UhtPropertyExportFlags.GetterSpecifiedAuto))
 					{
 						string getterCallText = property.Getter ?? "Get" + property.SourceName;
-						builder.AppendPropertyText(property, UhtPropertyTextType.GetterRetVal).Append(classObj.SourceName).Append("::").Append(getterCallText).Append("() const\r\n");
+						builder.AppendPropertyText(property, UhtPropertyTextType.GetterRetVal).Append(classObj.Namespace.FullSourceName).Append(classObj.SourceName).Append("::").Append(getterCallText).Append("() const\r\n");
 						builder.Append("{\r\n");
 						builder.Append("\treturn ").Append(property.SourceName).Append(";\r\n");
 						builder.Append("}\r\n");
@@ -1881,7 +1887,7 @@ namespace EpicGames.UHT.Exporters.CodeGen
 					if (property.PropertyExportFlags.HasAnyFlags(UhtPropertyExportFlags.SetterSpecifiedAuto))
 					{
 						string setterCallText = property.Setter ?? "Set" + property.SourceName;
-						builder.Append("void ").Append(classObj.SourceName).Append("::").Append(setterCallText).Append('(').AppendPropertyText(property, UhtPropertyTextType.SetterParameterArgType).Append("InValue").Append(")\r\n");
+						builder.Append("void ").Append(classObj.Namespace.FullSourceName).Append(classObj.SourceName).Append("::").Append(setterCallText).Append('(').AppendPropertyText(property, UhtPropertyTextType.SetterParameterArgType).Append("InValue").Append(")\r\n");
 						builder.Append("{\r\n");
 						// @todo: setter defn
 						builder.Append("}\r\n");
@@ -1897,7 +1903,7 @@ namespace EpicGames.UHT.Exporters.CodeGen
 				{
 					if (property.PropertyExportFlags.HasAnyFlags(UhtPropertyExportFlags.GetterFound))
 					{
-						builder.Append("void ").Append(classObj.SourceName).Append("::").AppendPropertyGetterWrapperName(property).Append("(const void* Object, void* OutValue)\r\n");
+						builder.Append("void ").Append(classObj.Namespace.FullSourceName).Append(classObj.SourceName).Append("::").AppendPropertyGetterWrapperName(property).Append("(const void* Object, void* OutValue)\r\n");
 						builder.Append("{\r\n");
 						using (UhtMacroBlockEmitter blockEmitter = new(builder, UhtDefineScopeNames.Standard, property.DefineScope))
 						{
@@ -1961,7 +1967,7 @@ namespace EpicGames.UHT.Exporters.CodeGen
 					}
 					if (property.PropertyExportFlags.HasAnyFlags(UhtPropertyExportFlags.SetterFound))
 					{
-						builder.Append("void ").Append(classObj.SourceName).Append("::").AppendPropertySetterWrapperName(property).Append("(void* Object, const void* InValue)\r\n");
+						builder.Append("void ").Append(classObj.Namespace.FullSourceName).Append(classObj.SourceName).Append("::").AppendPropertySetterWrapperName(property).Append("(void* Object, const void* InValue)\r\n");
 						builder.Append("{\r\n");
 						using (UhtMacroBlockEmitter blockEmitter = new(builder, UhtDefineScopeNames.Standard, property.DefineScope))
 						{
@@ -2042,7 +2048,7 @@ namespace EpicGames.UHT.Exporters.CodeGen
 			if (!classObj.SerializerArchiveType.HasAnyFlags(serializerType))
 			{
 				builder.AppendScoped(UhtDefineScopeNames.Standard, classObj.SerializerDefineScope,
-					builder => builder.Append(macroText).Append('(').Append(classObj.SourceName).Append(")\r\n"));
+					builder => builder.Append(macroText).Append('(').Append(classObj.Namespace.FullSourceName).Append(classObj.SourceName).Append(")\r\n"));
 			}
 			return builder;
 		}
@@ -2055,19 +2061,19 @@ namespace EpicGames.UHT.Exporters.CodeGen
 			string registrationName = $"Z_Registration_Info_UClass_{classObj.SourceName}";
 			string[]? sparseDataTypes = classObj.MetaData.GetStringArray(UhtNames.SparseClassDataTypes);
 
-			PropertyMemberContextImpl context = new(CodeGenerator, classObj, classObj.SourceName, staticsName);
+			PropertyMemberContextImpl context = new(CodeGenerator, classObj, staticsName);
 
 			bool hasInterfaces = classObj.Bases.Any(x => x is UhtClass baseClass && baseClass.ClassFlags.HasAnyFlags(EClassFlags.Interface)) || classObj.FlattenedVerseInterfaces.Count > 0;
 
 			// This block of code is the contents of IMPLEMENT_CLASS_NO_AUTO_REGISTRATION 
 			builder.Append("FClassRegistrationInfo ").Append(registrationName).Append(";\r\n");
-			builder.Append("UClass* ").Append(classObj.SourceName).Append("::GetPrivateStaticClass()\r\n");
+			builder.Append("UClass* ").Append(classObj.Namespace.FullSourceName).Append(classObj.SourceName).Append("::GetPrivateStaticClass()\r\n");
 			builder.Append("{\r\n");
-			builder.Append("\tusing TClass = ").Append(classObj.SourceName).Append(";\r\n");
+			builder.Append("\tusing TClass = ").Append(classObj.Namespace.FullSourceName).Append(classObj.SourceName).Append(";\r\n");
 			builder.Append("\tif (!").Append(registrationName).Append(".InnerSingleton)\r\n");
 			builder.Append("\t{\r\n");
 			builder.Append("\t\t").Append(classObj.IsVerseField ? "Verse::CodeGen::Private::ConstructUVerseClassNoInit" : "GetPrivateStaticClassBody").Append("(\r\n");
-			builder.Append("\t\t\tStaticPackage(),\r\n");
+			builder.Append("\t\t\tTClass::StaticPackage(),\r\n");
 			builder.Append("\t\t\tTEXT(\"").Append(classObj.EngineName).Append("\"),\r\n");
 			builder.Append("\t\t\t").Append(registrationName).Append(".InnerSingleton,\r\n");
 			builder.Append("\t\t\tStaticRegisterNatives").Append(classObj.SourceName).Append(",\r\n");
@@ -2093,7 +2099,7 @@ namespace EpicGames.UHT.Exporters.CodeGen
 			{
 				builder.Append("UClass* ").Append(GetSingletonName(classObj, false)).Append("()\r\n");
 				builder.Append("{\r\n");
-				builder.Append("\treturn ").Append(classObj.SourceName).Append("::GetPrivateStaticClass();\r\n");
+				builder.Append("\treturn ").Append(classObj.Namespace.FullSourceName).Append(classObj.SourceName).Append("::GetPrivateStaticClass();\r\n");
 				builder.Append("}\r\n");
 			}
 
@@ -2141,7 +2147,7 @@ namespace EpicGames.UHT.Exporters.CodeGen
 				}
 
 				builder.Append("\tstatic constexpr FCppClassTypeInfoStatic StaticCppClassTypeInfo = {\r\n");
-				builder.Append("\t\tTCppClassTypeTraits<").AppendClassSourceNameOrInterfaceName(classObj).Append(">::IsAbstract,\r\n");
+				builder.Append("\t\tTCppClassTypeTraits<").Append(classObj.Namespace.FullSourceName).AppendClassSourceNameOrInterfaceName(classObj).Append(">::IsAbstract,\r\n");
 				builder.Append("\t};\r\n");
 
 				builder.Append("\tstatic const UECodeGen_Private::").AppendClassParamsType(classObj).Append(" ClassParams;\r\n");
@@ -2191,7 +2197,7 @@ namespace EpicGames.UHT.Exporters.CodeGen
 				{
 					builder.Append("\t{\r\n");
 				}
-				builder.Append("\t&").Append(classObj.SourceName).Append("::StaticClass,\r\n");
+				builder.Append("\t&").Append(classObj.Namespace.FullSourceName).Append(classObj.SourceName).Append("::StaticClass,\r\n");
 				if (classObj.Config.Length > 0)
 				{
 					builder.Append('\t').AppendUTF8LiteralString(classObj.Config).Append(",\r\n");
@@ -2346,7 +2352,7 @@ namespace EpicGames.UHT.Exporters.CodeGen
 
 		private static StringBuilder AppendNatives(StringBuilder builder, UhtClass classObj)
 		{
-			builder.Append("void ").Append(classObj.SourceName).Append("::StaticRegisterNatives").Append(classObj.SourceName).Append("()\r\n");
+			builder.Append("void ").Append(classObj.Namespace.FullSourceName).Append(classObj.SourceName).Append("::StaticRegisterNatives").Append(classObj.SourceName).Append("()\r\n");
 			builder.Append("{\r\n");
 
 			UhtUsedDefineScopes<UhtFunction> functions = new(classObj.Functions.Where(
@@ -2579,18 +2585,30 @@ namespace EpicGames.UHT.Exporters.CodeGen
 		{
 			private readonly UhtCodeGenerator _codeGenerator;
 			private readonly UhtStruct _outerStruct;
+			private readonly string _outerStructNamespaceName;
 			private readonly string _outerStructSourceName;
 			private readonly string _staticsName;
 
-			public PropertyMemberContextImpl(UhtCodeGenerator codeGenerator, UhtStruct outerStruct, string outerStructSourceName, string staticsName)
+			public PropertyMemberContextImpl(UhtCodeGenerator codeGenerator, UhtStruct outerStruct, string outerStructNamespaceName, string outerStructSourceName, string staticsName)
 			{
 				_codeGenerator = codeGenerator;
 				_outerStruct = outerStruct;
 				_staticsName = staticsName;
-				_outerStructSourceName = outerStructSourceName.Length == 0 ? outerStruct.SourceName : outerStructSourceName;
+				_outerStructNamespaceName = outerStructNamespaceName;
+				_outerStructSourceName = outerStructSourceName;
+			}
+
+			public PropertyMemberContextImpl(UhtCodeGenerator codeGenerator, UhtStruct outerStruct, string staticsName)
+			{
+				_codeGenerator = codeGenerator;
+				_outerStruct = outerStruct;
+				_staticsName = staticsName;
+				_outerStructNamespaceName = _outerStruct is UhtScriptStruct scriptStruct ? scriptStruct.NamespaceExportName : _outerStruct.Namespace.FullSourceName;
+				_outerStructSourceName = _outerStruct.SourceName;
 			}
 
 			public UhtStruct OuterStruct => _outerStruct;
+			public string OuterStructNamespaceName => _outerStructNamespaceName;
 			public string OuterStructSourceName => _outerStructSourceName;
 			public string StaticsName => _staticsName;
 			public string NamePrefix => "NewProp_";
diff --git a/Engine/Source/Programs/Shared/EpicGames.UHT/Exporters/CodeGen/UhtHeaderCodeGeneratorHFile.cs b/Engine/Source/Programs/Shared/EpicGames.UHT/Exporters/CodeGen/UhtHeaderCodeGeneratorHFile.cs
--- a/Engine/Source/Programs/Shared/EpicGames.UHT/Exporters/CodeGen/UhtHeaderCodeGeneratorHFile.cs
+++ b/Engine/Source/Programs/Shared/EpicGames.UHT/Exporters/CodeGen/UhtHeaderCodeGeneratorHFile.cs
@@ -98,14 +98,44 @@ namespace EpicGames.UHT.Exporters.CodeGen
 				builder.Append("\r\n");
 				builder.Append(DisableDeprecationWarnings).Append("\r\n");
 
+				StringComparerUE compareUE = StringComparerUE.OrdinalIgnoreCase;
 				if (HeaderFile.References.ForwardDeclarations.Count > 0)
 				{
-					string[] sorted = new string[HeaderFile.References.ForwardDeclarations.Count];
-					int index = 0;
-					foreach (string forwardDeclaration in HeaderFile.References.ForwardDeclarations)
+					UhtField[] sorted = [.. HeaderFile.References.ForwardDeclarations];
+					Array.Sort(sorted, (x, y) =>
 					{
-						sorted[index++] = forwardDeclaration;
+					
+						int compare;
+						compare = compareUE.Compare(x.Namespace.FullSourceName, y.Namespace.FullSourceName);
+						if (compare != 0)
+						{
+							return compare;
+						}
+
+						compare = compareUE.Compare(x.EngineClassName, y.EngineClassName);
+						if (compare != 0)
+						{
+							return compare;
+						}
+
+						return compareUE.Compare(x.SourceName, y.SourceName);
+					});
+
+					for (int index = 0; index < sorted.Length;)
+					{
+						UhtNamespace namespaceObj = sorted[index].Namespace;
+						namespaceObj.AppendMultipleLines(builder, (builder) =>
+						{
+							for (; index < sorted.Length && sorted[index].Namespace == namespaceObj; index++)
+							{
+								builder.AppendForwardDeclaration(sorted[index]).Append("\r\n");
+							}
+						});
 					}
+				}
+				if (HeaderFile.References.ForwardDeclarationStrings.Count > 0)
+				{
+					string[] sorted = [.. HeaderFile.References.ForwardDeclarationStrings];
 					Array.Sort(sorted, StringComparerUE.OrdinalIgnoreCase);
 					builder.Append("\r\n");
 					foreach (string forwardDeclaration in sorted)
@@ -448,6 +478,8 @@ namespace EpicGames.UHT.Exporters.CodeGen
 
 		private StringBuilder AppendScriptStruct(StringBuilder builder, UhtScriptStruct scriptStruct)
 		{
+			builder.Append("struct Z_Construct_UScriptStruct_").Append(scriptStruct.SourceName).Append("_Statics;\r\n");
+
 			AppendRigVmStructInfoDefine(builder, scriptStruct);
 			
 			if (scriptStruct.ScriptStructFlags.HasAnyFlags(EStructFlags.Native))
@@ -649,7 +681,7 @@ namespace EpicGames.UHT.Exporters.CodeGen
 				}
 
 				// For back compat we still need to forward declare the struct here
-				builder.Append("struct ").Append(scriptStruct.SourceName).Append(";\r\n");
+				AppendForwardDeclarationWithNamespace(builder, scriptStruct);
 
 				// Verse declarations
 				if (scriptStruct.IsVerseField)
@@ -828,7 +860,11 @@ namespace EpicGames.UHT.Exporters.CodeGen
 				AppendSerializer(builder, classObj, api, UhtSerializerArchiveType.StructuredArchiveRecord, "DECLARE_FSTRUCTUREDARCHIVE_SERIALIZER");
 			}
 
-			// Declaration of globla function fetching UClass of type
+			// Declaration of global function fetching UClass of type
+			if (!classObj.ClassExportFlags.HasAnyFlags(UhtClassExportFlags.NoExport))
+			{
+				builder.Append("struct ").Append(ObjectInfos[classObj.ObjectTypeIndex].RegisteredSingletonName).Append("_Statics;\r\n");
+			}
 			builder.Append(classObj.Module.Api).Append("UClass* ").Append(ObjectInfos[classObj.ObjectTypeIndex].UnregisteredSingletonName).Append("();\r\n\r\n");
 
 			if (classObj.ClassFlags.HasAnyFlags(EClassFlags.Interface))
@@ -917,7 +953,7 @@ namespace EpicGames.UHT.Exporters.CodeGen
 			}
 
 			// For back compat we still need to forward declare the class here
-			builder.Append("class ").Append(classObj.SourceName).Append(";\r\n");
+			AppendForwardDeclarationWithNamespace(builder, classObj);
 			
 			// Verse declarations
 			if (classObj.IsVerseField)
@@ -1552,7 +1588,7 @@ namespace EpicGames.UHT.Exporters.CodeGen
 			}
 			
 			builder.Append('\t').Append("static UClass* GetPrivateStaticClass(); \\\r\n");
-			builder.Append("\tfriend ").Append(classObj.Module.Api).Append("UClass* ").Append(info.UnregisteredSingletonName).Append("(); \\\r\n");
+			builder.Append("\tfriend ").Append(classObj.Module.Api).Append("UClass* ::").Append(info.UnregisteredSingletonName).Append("(); \\\r\n");
 
 			builder.Append("public: \\\r\n");
 
@@ -1920,6 +1956,11 @@ namespace EpicGames.UHT.Exporters.CodeGen
 			return builder;
 		}
 
+		private static StringBuilder AppendForwardDeclarationWithNamespace(StringBuilder builder, UhtField field)
+		{
+			return field.Namespace.AppendSingleLine(builder, (builder) => builder.AppendForwardDeclaration(field));
+		}
+
 		#region Enum helper methods
 		private static bool IsFullEnumName(string inEnumName)
 		{
diff --git a/Engine/Source/Programs/Shared/EpicGames.UHT/Parsers/UhtClassParser.cs b/Engine/Source/Programs/Shared/EpicGames.UHT/Parsers/UhtClassParser.cs
--- a/Engine/Source/Programs/Shared/EpicGames.UHT/Parsers/UhtClassParser.cs
+++ b/Engine/Source/Programs/Shared/EpicGames.UHT/Parsers/UhtClassParser.cs
@@ -56,7 +56,7 @@ namespace EpicGames.UHT.Parsers
 
 		private static UhtParseResult ParseUClass(UhtParsingScope parentScope, ref UhtToken token)
 		{
-			UhtClass classObj = new(parentScope.HeaderFile, parentScope.ScopeType, token.InputLine);
+			UhtClass classObj = new(parentScope.HeaderFile, parentScope.HeaderParser.GetNamespace(), parentScope.ScopeType, token.InputLine);
 			{
 				using UhtParsingScope topScope = new(parentScope, classObj, parentScope.Session.GetKeywordTable(UhtTableNames.Class), UhtAccessSpecifier.Private);
 				const string ScopeName = "class";
diff --git a/Engine/Source/Programs/Shared/EpicGames.UHT/Parsers/UhtEnumParser.cs b/Engine/Source/Programs/Shared/EpicGames.UHT/Parsers/UhtEnumParser.cs
--- a/Engine/Source/Programs/Shared/EpicGames.UHT/Parsers/UhtEnumParser.cs
+++ b/Engine/Source/Programs/Shared/EpicGames.UHT/Parsers/UhtEnumParser.cs
@@ -27,7 +27,7 @@ namespace EpicGames.UHT.Parsers
 
 		private static UhtParseResult ParseUEnum(UhtParsingScope parentScope, UhtToken keywordToken)
 		{
-			UhtEnum enumObject = new(parentScope.HeaderFile, parentScope.ScopeType, keywordToken.InputLine);
+			UhtEnum enumObject = new(parentScope.HeaderFile, parentScope.HeaderParser.GetNamespace(), parentScope.ScopeType, keywordToken.InputLine);
 			{
 				using UhtParsingScope topScope = new(parentScope, enumObject, parentScope.Session.GetKeywordTable(UhtTableNames.Enum), UhtAccessSpecifier.Public);
 				const string ScopeName = "UENUM";
diff --git a/Engine/Source/Programs/Shared/EpicGames.UHT/Parsers/UhtFunctionParser.cs b/Engine/Source/Programs/Shared/EpicGames.UHT/Parsers/UhtFunctionParser.cs
--- a/Engine/Source/Programs/Shared/EpicGames.UHT/Parsers/UhtFunctionParser.cs
+++ b/Engine/Source/Programs/Shared/EpicGames.UHT/Parsers/UhtFunctionParser.cs
@@ -119,7 +119,7 @@ namespace EpicGames.UHT.Parsers
 		[SuppressMessage("CodeQuality", "IDE0051:Remove unused private members", Justification = "Attribute accessed method")]
 		private static UhtParseResult ParseCatchAllKeyword(UhtParsingScope topScope, ref UhtToken token)
 		{
-			if (UhtFunctionParser.IsValidateDelegateDeclaration(token))
+			if (topScope.HeaderParser.NamespaceDepth == 1 && UhtFunctionParser.IsValidateDelegateDeclaration(token))
 			{
 				return ParseUDelegate(topScope, token, false);
 			}
@@ -129,7 +129,7 @@ namespace EpicGames.UHT.Parsers
 
 		private static UhtParseResult ParseUDelegate(UhtParsingScope parentScope, UhtToken token, bool hasSpecifiers)
 		{
-			UhtFunction function = new(parentScope.HeaderFile, parentScope.ScopeType, token.InputLine);
+			UhtFunction function = new(parentScope.HeaderFile, parentScope.HeaderParser.GetNamespace(), parentScope.ScopeType, token.InputLine);
 
 			{
 				using UhtParsingScope topScope = new(parentScope, function, parentScope.Session.GetKeywordTable(UhtTableNames.Function), UhtAccessSpecifier.Public);
@@ -299,7 +299,7 @@ namespace EpicGames.UHT.Parsers
 
 		private static UhtParseResult ParseUFunction(UhtParsingScope parentScope, UhtToken token)
 		{
-			UhtFunction function = new(parentScope.HeaderFile, parentScope.ScopeType, token.InputLine);
+			UhtFunction function = new(parentScope.HeaderFile, parentScope.HeaderParser.GetNamespace(), parentScope.ScopeType, token.InputLine);
 
 			{
 				using UhtParsingScope topScope = new(parentScope, function, parentScope.Session.GetKeywordTable(UhtTableNames.Function), UhtAccessSpecifier.Public);
diff --git a/Engine/Source/Programs/Shared/EpicGames.UHT/Parsers/UhtHeaderFileParser.cs b/Engine/Source/Programs/Shared/EpicGames.UHT/Parsers/UhtHeaderFileParser.cs
--- a/Engine/Source/Programs/Shared/EpicGames.UHT/Parsers/UhtHeaderFileParser.cs
+++ b/Engine/Source/Programs/Shared/EpicGames.UHT/Parsers/UhtHeaderFileParser.cs
@@ -355,29 +355,45 @@ namespace EpicGames.UHT.Parsers
 		/// </summary>
 		public string? RestrictedPreprocessorContext { get; set; } = null;
 
+		/// <summary>
+		/// Depth of the namespaces
+		/// </summary>
+		public int NamespaceDepth => _namespaces.Count;
+
 		/// <summary>
 		/// Stack of current #if states
 		/// </summary>
-		private readonly List<CompilerDirective> _compilerDirectives = new();
+		private readonly List<CompilerDirective> _compilerDirectives = [];
 
 		/// <summary>
 		/// Stack of current #if states saved as part of the preprocessor state
 		/// </summary>
-		private readonly List<CompilerDirective> _savedCompilerDirectives = new();
+		private readonly List<CompilerDirective> _savedCompilerDirectives = [];
 
 		/// <summary>
 		/// Current top of the parsing scopes.  Classes, structures and functions all allocate scopes.
 		/// </summary>
 		private UhtParsingScope? _topScope = null;
 
+		/// <summary>
+		/// Current namespaces being parsed
+		/// </summary>
+		private readonly List<UhtNamespaceElement> _namespaces = [];
+
+		/// <summary>
+		/// Lookup cache for any namespaces defined in this header.
+		/// </summary>
+		private readonly Dictionary<string, UhtNamespace> _namespaceLookupCache = [];
+
 		/// <summary>
 		/// Parse the given header file
 		/// </summary>
 		/// <param name="headerFile">Header file to parse</param>
+		/// <param name="globalNamespace">Global namespace</param>
 		/// <returns>Parser</returns>
-		public static UhtHeaderFileParser Parse(UhtHeaderFile headerFile)
+		public static UhtHeaderFileParser Parse(UhtHeaderFile headerFile, UhtNamespace globalNamespace)
 		{
-			UhtHeaderFileParser headerParser = new(headerFile);
+			UhtHeaderFileParser headerParser = new(headerFile, globalNamespace);
 			using UhtParsingScope topScope = new(headerParser, headerFile.Session.GetKeywordTable(UhtTableNames.Global));
 			headerParser.ParseStatements();
 
@@ -427,11 +443,13 @@ namespace EpicGames.UHT.Parsers
 			return headerParser;
 		}
 
-		private UhtHeaderFileParser(UhtHeaderFile headerFile)
+		private UhtHeaderFileParser(UhtHeaderFile headerFile, UhtNamespace globalNamespace)
 		{
 			TokenReader = new UhtTokenBufferReader(headerFile, headerFile.Data.Memory);
 			HeaderFile = headerFile;
 			TokenReader.TokenPreprocessor = this;
+			_namespaces.Add(new("", globalNamespace));
+			_namespaceLookupCache.Add("", globalNamespace);
 		}
 
 		/// <summary>
@@ -462,6 +480,41 @@ namespace EpicGames.UHT.Parsers
 			_topScope = scope.ParentScope;
 		}
 
+		/// <summary>
+		/// Add a namespace to the token list
+		/// </summary>
+		/// <param name="token">Namespace being added</param>
+		public void PushNamespace(UhtToken token)
+		{
+			_namespaces.Add(new(token.Value, null));
+		}
+
+		/// <summary>
+		/// Remove the given number of namespaces
+		/// </summary>
+		/// <param name="count"></param>
+		public void PopNamespaces(int count)
+		{
+			if (count > _namespaces.Count)
+			{
+				throw new UhtException(TokenReader, "Mismatch when looking for closing brace");
+			}
+			_namespaces.RemoveRange(_namespaces.Count - count, count);
+		}
+
+		/// <summary>
+		/// Returns the currently active namespace
+		/// </summary>
+		/// <returns></returns>
+		public UhtNamespace GetNamespace()
+		{
+			if (_topScope == null)
+			{
+				throw new UhtIceException("Expected a top scope");
+			}
+			return _topScope.Session.GetNamespace(_namespaces, _namespaceLookupCache);
+		}
+
 		/// <summary>
 		/// Return the current compiler directive
 		/// </summary>
diff --git a/Engine/Source/Programs/Shared/EpicGames.UHT/Parsers/UhtInterfaceClassParser.cs b/Engine/Source/Programs/Shared/EpicGames.UHT/Parsers/UhtInterfaceClassParser.cs
--- a/Engine/Source/Programs/Shared/EpicGames.UHT/Parsers/UhtInterfaceClassParser.cs
+++ b/Engine/Source/Programs/Shared/EpicGames.UHT/Parsers/UhtInterfaceClassParser.cs
@@ -50,7 +50,7 @@ namespace EpicGames.UHT.Parsers
 
 		private static UhtParseResult ParseUInterface(UhtParsingScope parentScope, ref UhtToken token)
 		{
-			UhtClass classObj = new(parentScope.HeaderFile, parentScope.ScopeType, token.InputLine);
+			UhtClass classObj = new(parentScope.HeaderFile, parentScope.HeaderParser.GetNamespace(), parentScope.ScopeType, token.InputLine);
 			classObj.ClassType = UhtClassType.Interface;
 
 			{
diff --git a/Engine/Source/Programs/Shared/EpicGames.UHT/Parsers/UhtNamespaceParser.cs b/Engine/Source/Programs/Shared/EpicGames.UHT/Parsers/UhtNamespaceParser.cs
new file mode 100644
--- /dev/null
+++ b/Engine/Source/Programs/Shared/EpicGames.UHT/Parsers/UhtNamespaceParser.cs
@@ -0,0 +1,76 @@
+// Copyright Epic Games, Inc. All Rights Reserved.
+
+using System.Diagnostics.CodeAnalysis;
+using EpicGames.UHT.Tables;
+using EpicGames.UHT.Tokenizer;
+using EpicGames.UHT.Types;
+using EpicGames.UHT.Utils;
+
+namespace EpicGames.UHT.Parsers
+{
+
+	/// <summary>
+	/// namespace parser
+	/// </summary>
+	[UnrealHeaderTool]
+	class UhtNamespaceParser
+	{
+		[UhtKeyword(Extends = UhtTableNames.Global, Keyword = "namespace", DisableUsageError = true)]
+		[SuppressMessage("CodeQuality", "IDE0051:Remove unused private members", Justification = "Attribute accessed method")]
+		[SuppressMessage("Style", "IDE0060:Remove unused parameter", Justification = "Attribute accessed method")]
+		private static UhtParseResult NamespaceKeyword(UhtParsingScope topScope, UhtParsingScope actionScope, ref UhtToken token)
+		{
+			return topScope.Module.Module.AllowUETypesInNamespaces ? ParseNamespace(topScope) : UhtParseResult.Unhandled;
+		}
+
+		private static UhtParseResult ParseNamespace(UhtParsingScope parentScope)
+		{
+			using UhtParsingScope topScope = new(parentScope, parentScope.ScopeType, parentScope.Session.GetKeywordTable(UhtTableNames.Global), UhtAccessSpecifier.Public);
+			const string ScopeName = "namespace";
+			using UhtMessageContext tokenContext = new(ScopeName);
+
+			// Handle any attributes
+			topScope.TokenReader.OptionalAttributes(false);
+
+			// Parse the namespace identifiers
+			int namespaceCount = 0;
+			if (topScope.TokenReader.TryOptionalIdentifier(out UhtToken token))
+			{
+				namespaceCount++;
+				topScope.HeaderParser.PushNamespace(token);
+				while (topScope.TokenReader.TryOptional("::"))
+				{
+					namespaceCount++;
+					topScope.TokenReader.Optional("inline");
+					topScope.HeaderParser.PushNamespace(topScope.TokenReader.GetIdentifier());
+				}
+			}
+
+			// if this is a namespace id = ... statement
+			if (topScope.TokenReader.TryOptional('='))
+			{
+				while (true)
+				{
+					UhtToken skipped = topScope.TokenReader.GetToken();
+					if (skipped.IsSymbol(';'))
+					{
+						break;
+					}
+				}
+			}
+			else
+			{
+				// Parse the namespace body
+				topScope.HeaderParser.ParseStatements('{', '}', true);
+
+				// TODO: The parse statement code expects to have any trailing ';' removed.  Remove any here
+				topScope.TokenReader.Optional(';');
+			}
+
+
+			// Remove the namespaces from the active namespace
+			topScope.HeaderParser.PopNamespaces(namespaceCount);
+			return UhtParseResult.Handled;
+		}
+	}
+}
diff --git a/Engine/Source/Programs/Shared/EpicGames.UHT/Parsers/UhtNativeInterfaceClassParser.cs b/Engine/Source/Programs/Shared/EpicGames.UHT/Parsers/UhtNativeInterfaceClassParser.cs
--- a/Engine/Source/Programs/Shared/EpicGames.UHT/Parsers/UhtNativeInterfaceClassParser.cs
+++ b/Engine/Source/Programs/Shared/EpicGames.UHT/Parsers/UhtNativeInterfaceClassParser.cs
@@ -15,7 +15,7 @@ namespace EpicGames.UHT.Parsers
 	/// </summary>
 	[UnrealHeaderTool]
 	public static class UhtNativeInterfaceClassParser
-	{ 
+	{
 
 		#region Keywords
 		[UhtKeyword(Extends = UhtTableNames.Global, Keyword = "class", DisableUsageError = true)]
@@ -26,7 +26,7 @@ namespace EpicGames.UHT.Parsers
 			{
 				using UhtTokenSaveState saveState = new(topScope.TokenReader);
 				UhtToken sourceName = new();
-				UhtToken superName = new();
+				UhtToken[] superName = [];
 				if (TryParseIInterface(topScope, out sourceName, out superName))
 				{
 					saveState.AbandonState();
@@ -60,7 +60,7 @@ namespace EpicGames.UHT.Parsers
 		}
 		#endregion
 
-		private static bool TryParseIInterface(UhtParsingScope parentScope, out UhtToken sourceName, out UhtToken superName)
+		private static bool TryParseIInterface(UhtParsingScope parentScope, out UhtToken sourceName, out UhtToken[] superName)
 		{
 			IUhtTokenReader tokenReader = parentScope.TokenReader;
 
@@ -70,15 +70,15 @@ namespace EpicGames.UHT.Parsers
 			// Get the name of the interface
 			sourceName = tokenReader.GetIdentifier();
 
-			// Old UHT would still parse the inheritance 
-			superName = new UhtToken();
+			// Old UHT would still parse the inheritance
+			superName = [];
 			if (tokenReader.TryOptional(':'))
 			{
 				if (!tokenReader.TryOptional("public"))
 				{
 					return false;
 				}
-				superName = tokenReader.GetIdentifier();
+				superName = tokenReader.GetCppIdentifier(UhtCppIdentifierOptions.None).ToArray();
 			}
 
 			// Only process classes starting with 'I'
@@ -101,11 +101,11 @@ namespace EpicGames.UHT.Parsers
 			return true;
 		}
 
-		private static void ParseIInterface(UhtParsingScope parentScope, ref UhtToken token, UhtToken sourceName, UhtToken superName)
+		private static void ParseIInterface(UhtParsingScope parentScope, ref UhtToken token, UhtToken sourceName, UhtToken[] superName)
 		{
 			IUhtTokenReader tokenReader = parentScope.TokenReader;
 
-			UhtClass classObj = new(parentScope.HeaderFile, parentScope.ScopeType, token.InputLine);
+			UhtClass classObj = new(parentScope.HeaderFile, parentScope.HeaderParser.GetNamespace(), parentScope.ScopeType, token.InputLine);
 			classObj.ClassType = UhtClassType.NativeInterface;
 			classObj.SourceName = sourceName.Value.ToString();
 			classObj.ClassFlags |= EClassFlags.Native | EClassFlags.Interface;
diff --git a/Engine/Source/Programs/Shared/EpicGames.UHT/Parsers/UhtParserHelpers.cs b/Engine/Source/Programs/Shared/EpicGames.UHT/Parsers/UhtParserHelpers.cs
--- a/Engine/Source/Programs/Shared/EpicGames.UHT/Parsers/UhtParserHelpers.cs
+++ b/Engine/Source/Programs/Shared/EpicGames.UHT/Parsers/UhtParserHelpers.cs
@@ -19,17 +19,17 @@ namespace EpicGames.UHT.Parsers
 		/// <summary>
 		/// Super class identifier
 		/// </summary>
-		public UhtToken SuperIdentifier { get; set; } = new();
+		public UhtToken[] SuperIdentifier { get; set; } = [];
 
 		/// <summary>
 		/// Collection of other base identifiers
 		/// </summary>
-		public List<UhtToken[]> BaseIdentifiers { get; } = new();
+		public List<UhtToken[]> BaseIdentifiers { get; } = [];
 
 		/// <summary>
 		/// Collection of verse interface identifiers
 		/// </summary>
-		public List<UhtToken> VerseInterfaceIdentifiers { get; } = new();
+		public List<UhtToken[]> VerseInterfaceIdentifiers { get; } = [];
 
 		/// <summary>
 		/// Line number that the VIDENTIFIER macro appeared
@@ -70,10 +70,10 @@ namespace EpicGames.UHT.Parsers
 				{
 					tokenReader
 						.Require("public", "public access modifier")
-						.RequireIdentifier((ref UhtToken identifier) =>
+						.RequireCppIdentifier(UhtCppIdentifierOptions.None, (UhtTokenList identifier) =>
 						{
-							config.RedirectTypeIdentifier(ref identifier);
-							scratch.SuperIdentifier = identifier;
+							identifier.RedirectTypeIdentifier(config);
+							scratch.SuperIdentifier = identifier.ToArray();
 						})
 						.While(',', () =>
 						{
@@ -89,9 +89,9 @@ namespace EpicGames.UHT.Parsers
 							tokenReader.Require('(');
 							tokenReader.RequireList(')', ',', false, () =>
 							{
-								tokenReader.RequireIdentifier((ref UhtToken identifier) =>
+								tokenReader.RequireCppIdentifier(UhtCppIdentifierOptions.None, (UhtTokenList identifier) =>
 								{
-									scratch.VerseInterfaceIdentifiers.Add(identifier);
+									scratch.VerseInterfaceIdentifiers.Add(identifier.ToArray());
 								});
 							});
 							scratch.VerseInterfacesLine = tokenReader.InputLine;
diff --git a/Engine/Source/Programs/Shared/EpicGames.UHT/Parsers/UhtScriptStructParser.cs b/Engine/Source/Programs/Shared/EpicGames.UHT/Parsers/UhtScriptStructParser.cs
--- a/Engine/Source/Programs/Shared/EpicGames.UHT/Parsers/UhtScriptStructParser.cs
+++ b/Engine/Source/Programs/Shared/EpicGames.UHT/Parsers/UhtScriptStructParser.cs
@@ -65,7 +65,7 @@ namespace EpicGames.UHT.Parsers
 
 		private static UhtParseResult ParseUScriptStruct(UhtParsingScope parentScope, UhtToken keywordToken)
 		{
-			UhtScriptStruct scriptStruct = new(parentScope.HeaderFile, parentScope.ScopeType, keywordToken.InputLine);
+			UhtScriptStruct scriptStruct = new(parentScope.HeaderFile, parentScope.HeaderParser.GetNamespace(), parentScope.ScopeType, keywordToken.InputLine);
 			{
 				using UhtParsingScope topScope = new(parentScope, scriptStruct, parentScope.Session.GetKeywordTable(UhtTableNames.ScriptStruct), UhtAccessSpecifier.Public);
 
diff --git a/Engine/Source/Programs/Shared/EpicGames.UHT/Tokenizer/UhtTokenReader.cs b/Engine/Source/Programs/Shared/EpicGames.UHT/Tokenizer/UhtTokenReader.cs
--- a/Engine/Source/Programs/Shared/EpicGames.UHT/Tokenizer/UhtTokenReader.cs
+++ b/Engine/Source/Programs/Shared/EpicGames.UHT/Tokenizer/UhtTokenReader.cs
@@ -268,6 +268,22 @@ namespace EpicGames.UHT.Tokenizer
 		/// </summary>
 		public UhtTokenList? Next { get; set; }
 
+		/// <summary>
+		/// Return the last token list entry
+		/// </summary>
+		public UhtTokenList Last
+		{
+			get 
+			{
+				UhtTokenList last = this;
+				while (last.Next != null)
+				{
+					last = last.Next;
+				}
+				return last;
+			}
+		}
+
 		/// <summary>
 		/// Join the tokens in the list
 		/// </summary>
@@ -302,6 +318,18 @@ namespace EpicGames.UHT.Tokenizer
 			return builder.ToString();
 		}
 
+		/// <summary>
+		/// Remap the last identifier
+		/// </summary>
+		/// <param name="config">Configuration containing remapping information</param>
+		public void RedirectTypeIdentifier(IUhtConfig config)
+		{
+			UhtTokenList last = Last;
+			UhtToken scratch = last.Token;
+			config.RedirectTypeIdentifier(ref scratch);
+			last.Token = scratch;
+		}
+
 		/// <summary>
 		/// Return the token list as an array
 		/// </summary>
diff --git a/Engine/Source/Programs/Shared/EpicGames.UHT/Tokenizer/UhtTokenReaderIdentifierExtensions.cs b/Engine/Source/Programs/Shared/EpicGames.UHT/Tokenizer/UhtTokenReaderIdentifierExtensions.cs
--- a/Engine/Source/Programs/Shared/EpicGames.UHT/Tokenizer/UhtTokenReaderIdentifierExtensions.cs
+++ b/Engine/Source/Programs/Shared/EpicGames.UHT/Tokenizer/UhtTokenReaderIdentifierExtensions.cs
@@ -247,6 +247,7 @@ namespace EpicGames.UHT.Tokenizer
 		/// <returns>Token list</returns>
 		public static UhtTokenList GetCppIdentifier(this IUhtTokenReader tokenReader, UhtCppIdentifierOptions options = UhtCppIdentifierOptions.None)
 		{
+			tokenReader.Optional("::");
 			UhtToken token = tokenReader.GetIdentifier();
 			return tokenReader.GetCppIdentifier(ref token, options);
 		}
diff --git a/Engine/Source/Programs/Shared/EpicGames.UHT/Types/Properties/UhtArrayProperty.cs b/Engine/Source/Programs/Shared/EpicGames.UHT/Types/Properties/UhtArrayProperty.cs
--- a/Engine/Source/Programs/Shared/EpicGames.UHT/Types/Properties/UhtArrayProperty.cs
+++ b/Engine/Source/Programs/Shared/EpicGames.UHT/Types/Properties/UhtArrayProperty.cs
@@ -86,18 +86,6 @@ namespace EpicGames.UHT.Types
 			}
 		}
 
-		/// <inheritdoc/>
-		public override void CollectReferencesInternal(IUhtReferenceCollector collector, bool templateProperty)
-		{
-			ValueProperty.CollectReferencesInternal(collector, true);
-		}
-
-		/// <inheritdoc/>
-		public override string? GetForwardDeclarations()
-		{
-			return ValueProperty.GetForwardDeclarations();
-		}
-
 		/// <inheritdoc/>
 		public override IEnumerable<UhtType> EnumerateReferencedTypes()
 		{
diff --git a/Engine/Source/Programs/Shared/EpicGames.UHT/Types/Properties/UhtBoolProperty.cs b/Engine/Source/Programs/Shared/EpicGames.UHT/Types/Properties/UhtBoolProperty.cs
--- a/Engine/Source/Programs/Shared/EpicGames.UHT/Types/Properties/UhtBoolProperty.cs
+++ b/Engine/Source/Programs/Shared/EpicGames.UHT/Types/Properties/UhtBoolProperty.cs
@@ -185,7 +185,7 @@ namespace EpicGames.UHT.Types
 			{
 				builder.AppendTabs(tabs).Append("void ").AppendNameDef(context, name, nameSuffix).Append("_SetBit(void* Obj)\r\n");
 				builder.AppendTabs(tabs).Append("{\r\n");
-				builder.AppendTabs(tabs + 1).Append("((").Append(context.OuterStructSourceName).Append("*)Obj)->").Append(SourceName).Append(" = 1;\r\n");
+				builder.AppendTabs(tabs + 1).Append("((").Append(context.OuterStructNamespaceName).Append(context.OuterStructSourceName).Append("*)Obj)->").Append(SourceName).Append(" = 1;\r\n");
 				builder.AppendTabs(tabs).Append("}\r\n");
 			}
 
@@ -200,7 +200,7 @@ namespace EpicGames.UHT.Types
 			if (Outer == context.OuterStruct)
 			{
 				builder
-					.Append("sizeof(").Append(context.OuterStructSourceName).Append("), ")
+					.Append("sizeof(").Append(context.OuterStructNamespaceName).Append(context.OuterStructSourceName).Append("), ")
 					.Append('&').AppendNameDef(context, name, nameSuffix).Append("_SetBit, ");
 			}
 			else
diff --git a/Engine/Source/Programs/Shared/EpicGames.UHT/Types/Properties/UhtByteProperty.cs b/Engine/Source/Programs/Shared/EpicGames.UHT/Types/Properties/UhtByteProperty.cs
--- a/Engine/Source/Programs/Shared/EpicGames.UHT/Types/Properties/UhtByteProperty.cs
+++ b/Engine/Source/Programs/Shared/EpicGames.UHT/Types/Properties/UhtByteProperty.cs
@@ -110,8 +110,9 @@ namespace EpicGames.UHT.Types
 		}
 
 		/// <inheritdoc/>
-		public override void CollectReferencesInternal(IUhtReferenceCollector collector, bool templateProperty)
+		public override void CollectReferencesInternal(IUhtReferenceCollector collector, bool addForwardDeclarations, bool isTemplateProperty)
 		{
+			base.CollectReferencesInternal(collector, addForwardDeclarations, isTemplateProperty);
 			collector.AddCrossModuleReference(Enum, true);
 		}
 
diff --git a/Engine/Source/Programs/Shared/EpicGames.UHT/Types/Properties/UhtClassProperty.cs b/Engine/Source/Programs/Shared/EpicGames.UHT/Types/Properties/UhtClassProperty.cs
--- a/Engine/Source/Programs/Shared/EpicGames.UHT/Types/Properties/UhtClassProperty.cs
+++ b/Engine/Source/Programs/Shared/EpicGames.UHT/Types/Properties/UhtClassProperty.cs
@@ -37,9 +37,13 @@ namespace EpicGames.UHT.Types
 		}
 
 		/// <inheritdoc/>
-		public override string? GetForwardDeclarations()
+		public override void CollectReferencesInternal(IUhtReferenceCollector collector, bool addForwardDeclarations, bool isTemplateProperty)
 		{
-			return MetaClass != null ? $"class {MetaClass.SourceName};" : null;
+			base.CollectReferencesInternal(collector, addForwardDeclarations, isTemplateProperty);
+			if (addForwardDeclarations && MetaClass != null)
+			{
+				collector.AddForwardDeclaration(MetaClass);
+			}
 		}
 
 		/// <inheritdoc/>
@@ -96,7 +100,7 @@ namespace EpicGames.UHT.Types
 					}
 					else
 					{
-						builder.Append(Class.SourceName);
+						builder.Append(Class.Namespace.FullSourceName).Append(Class.SourceName);
 					}
 					break;
 
diff --git a/Engine/Source/Programs/Shared/EpicGames.UHT/Types/Properties/UhtContainerProperty.cs b/Engine/Source/Programs/Shared/EpicGames.UHT/Types/Properties/UhtContainerProperty.cs
--- a/Engine/Source/Programs/Shared/EpicGames.UHT/Types/Properties/UhtContainerProperty.cs
+++ b/Engine/Source/Programs/Shared/EpicGames.UHT/Types/Properties/UhtContainerProperty.cs
@@ -2,6 +2,7 @@
 
 using System.Diagnostics.CodeAnalysis;
 using EpicGames.Core;
+using EpicGames.UHT.Utils;
 
 namespace EpicGames.UHT.Types
 {
@@ -27,6 +28,13 @@ namespace EpicGames.UHT.Types
 				(ValueProperty.PropertyCaps & (UhtPropertyCaps.CanBeInstanced | UhtPropertyCaps.CanHaveConfig));
 		}
 
+		/// <inheritdoc/>
+		public override void CollectReferencesInternal(IUhtReferenceCollector collector, bool addForwardDeclarations, bool isTemplateProperty)
+		{
+			base.CollectReferencesInternal(collector, addForwardDeclarations, isTemplateProperty);
+			ValueProperty.CollectReferencesInternal(collector, addForwardDeclarations, true);
+		}
+
 		/// <inheritdoc/>
 		protected override bool NeedsGCBarrierWhenPassedToFunctionImpl(UhtFunction function)
 		{
diff --git a/Engine/Source/Programs/Shared/EpicGames.UHT/Types/Properties/UhtDelegateProperty.cs b/Engine/Source/Programs/Shared/EpicGames.UHT/Types/Properties/UhtDelegateProperty.cs
--- a/Engine/Source/Programs/Shared/EpicGames.UHT/Types/Properties/UhtDelegateProperty.cs
+++ b/Engine/Source/Programs/Shared/EpicGames.UHT/Types/Properties/UhtDelegateProperty.cs
@@ -66,8 +66,9 @@ namespace EpicGames.UHT.Types
 		}
 
 		/// <inheritdoc/>
-		public override void CollectReferencesInternal(IUhtReferenceCollector collector, bool templateProperty)
+		public override void CollectReferencesInternal(IUhtReferenceCollector collector, bool addForwardDeclarations, bool isTemplateProperty)
 		{
+			base.CollectReferencesInternal(collector, addForwardDeclarations, isTemplateProperty);
 			collector.AddCrossModuleReference(Function, true);
 		}
 
diff --git a/Engine/Source/Programs/Shared/EpicGames.UHT/Types/Properties/UhtEnumProperty.cs b/Engine/Source/Programs/Shared/EpicGames.UHT/Types/Properties/UhtEnumProperty.cs
--- a/Engine/Source/Programs/Shared/EpicGames.UHT/Types/Properties/UhtEnumProperty.cs
+++ b/Engine/Source/Programs/Shared/EpicGames.UHT/Types/Properties/UhtEnumProperty.cs
@@ -71,19 +71,14 @@ namespace EpicGames.UHT.Types
 		}
 
 		/// <inheritdoc/>
-		public override void CollectReferencesInternal(IUhtReferenceCollector collector, bool templateProperty)
+		public override void CollectReferencesInternal(IUhtReferenceCollector collector, bool addForwardDeclarations, bool isTemplateProperty)
 		{
+			base.CollectReferencesInternal(collector, addForwardDeclarations, isTemplateProperty);
 			collector.AddCrossModuleReference(Enum, true);
-		}
-
-		/// <inheritdoc/>
-		public override string? GetForwardDeclarations()
-		{
-			if (Enum.CppForm != UhtEnumCppForm.EnumClass)
+			if (addForwardDeclarations && Enum.CppForm == UhtEnumCppForm.EnumClass)
 			{
-				return null;
+				collector.AddForwardDeclaration(Enum);
 			}
-			return System.String.Format("enum class {0} : {1};", Enum.SourceName, UnderlyingType.ToString().ToLower());
 		}
 
 		/// <inheritdoc/>
diff --git a/Engine/Source/Programs/Shared/EpicGames.UHT/Types/Properties/UhtFieldPathProperty.cs b/Engine/Source/Programs/Shared/EpicGames.UHT/Types/Properties/UhtFieldPathProperty.cs
--- a/Engine/Source/Programs/Shared/EpicGames.UHT/Types/Properties/UhtFieldPathProperty.cs
+++ b/Engine/Source/Programs/Shared/EpicGames.UHT/Types/Properties/UhtFieldPathProperty.cs
@@ -45,9 +45,13 @@ namespace EpicGames.UHT.Types
 		}
 
 		/// <inheritdoc/>
-		public override string? GetForwardDeclarations()
+		public override void CollectReferencesInternal(IUhtReferenceCollector collector, bool addForwardDeclarations, bool isTemplateProperty)
 		{
-			return $"class F{FieldClassName};";
+			base.CollectReferencesInternal(collector, addForwardDeclarations, isTemplateProperty);
+			if (addForwardDeclarations)
+			{
+				collector.AddForwardDeclaration(null, $"class F{FieldClassName};");
+			}
 		}
 
 		/// <inheritdoc/>
diff --git a/Engine/Source/Programs/Shared/EpicGames.UHT/Types/Properties/UhtInterfaceProperty.cs b/Engine/Source/Programs/Shared/EpicGames.UHT/Types/Properties/UhtInterfaceProperty.cs
--- a/Engine/Source/Programs/Shared/EpicGames.UHT/Types/Properties/UhtInterfaceProperty.cs
+++ b/Engine/Source/Programs/Shared/EpicGames.UHT/Types/Properties/UhtInterfaceProperty.cs
@@ -76,20 +76,22 @@ namespace EpicGames.UHT.Types
 		}
 
 		/// <inheritdoc/>
-		public override void CollectReferencesInternal(IUhtReferenceCollector collector, bool templateProperty)
+		public override void CollectReferencesInternal(IUhtReferenceCollector collector, bool addForwardDeclarations, bool isTemplateProperty)
 		{
+			base.CollectReferencesInternal(collector, addForwardDeclarations, isTemplateProperty);
 			collector.AddCrossModuleReference(InterfaceClass, false);
-		}
-
-		/// <inheritdoc/>
-		public override string? GetForwardDeclarations()
-		{
-			UhtClass? exportClass = InterfaceClass;
-			while (exportClass != null && !exportClass.ClassFlags.HasAnyFlags(EClassFlags.Native))
+			if (addForwardDeclarations)
 			{
-				exportClass = exportClass.SuperClass;
+				UhtClass? exportClass = InterfaceClass;
+				while (exportClass != null && !exportClass.ClassFlags.HasAnyFlags(EClassFlags.Native))
+				{
+					exportClass = exportClass.SuperClass;
+				}
+				if (exportClass != null)
+				{
+					collector.AddForwardDeclaration(exportClass);
+				}
 			}
-			return exportClass != null ? $"class {exportClass.SourceName};" : null;
 		}
 
 		/// <inheritdoc/>
diff --git a/Engine/Source/Programs/Shared/EpicGames.UHT/Types/Properties/UhtMapProperty.cs b/Engine/Source/Programs/Shared/EpicGames.UHT/Types/Properties/UhtMapProperty.cs
--- a/Engine/Source/Programs/Shared/EpicGames.UHT/Types/Properties/UhtMapProperty.cs
+++ b/Engine/Source/Programs/Shared/EpicGames.UHT/Types/Properties/UhtMapProperty.cs
@@ -111,16 +111,10 @@ namespace EpicGames.UHT.Types
 		}
 
 		/// <inheritdoc/>
-		public override void CollectReferencesInternal(IUhtReferenceCollector collector, bool templateProperty)
+		public override void CollectReferencesInternal(IUhtReferenceCollector collector, bool addForwardDeclarations, bool isTemplateProperty)
 		{
-			ValueProperty.CollectReferencesInternal(collector, true);
-			KeyProperty.CollectReferencesInternal(collector, true);
-		}
-
-		/// <inheritdoc/>
-		public override string? GetForwardDeclarations()
-		{
-			return $"{KeyProperty.GetForwardDeclarations()} {ValueProperty.GetForwardDeclarations()}";
+			base.CollectReferencesInternal(collector, addForwardDeclarations, isTemplateProperty);
+			KeyProperty.CollectReferencesInternal(collector, addForwardDeclarations, true);
 		}
 
 		/// <inheritdoc/>
diff --git a/Engine/Source/Programs/Shared/EpicGames.UHT/Types/Properties/UhtMulticastDelegateProperty.cs b/Engine/Source/Programs/Shared/EpicGames.UHT/Types/Properties/UhtMulticastDelegateProperty.cs
--- a/Engine/Source/Programs/Shared/EpicGames.UHT/Types/Properties/UhtMulticastDelegateProperty.cs
+++ b/Engine/Source/Programs/Shared/EpicGames.UHT/Types/Properties/UhtMulticastDelegateProperty.cs
@@ -67,8 +67,9 @@ namespace EpicGames.UHT.Types
 		}
 
 		/// <inheritdoc/>
-		public override void CollectReferencesInternal(IUhtReferenceCollector collector, bool templateProperty)
+		public override void CollectReferencesInternal(IUhtReferenceCollector collector, bool addForwardDeclarations, bool isTemplateProperty)
 		{
+			base.CollectReferencesInternal(collector, addForwardDeclarations, isTemplateProperty);
 			collector.AddCrossModuleReference(Function, true);
 		}
 
diff --git a/Engine/Source/Programs/Shared/EpicGames.UHT/Types/Properties/UhtObjectProperty.cs b/Engine/Source/Programs/Shared/EpicGames.UHT/Types/Properties/UhtObjectProperty.cs
--- a/Engine/Source/Programs/Shared/EpicGames.UHT/Types/Properties/UhtObjectProperty.cs
+++ b/Engine/Source/Programs/Shared/EpicGames.UHT/Types/Properties/UhtObjectProperty.cs
@@ -84,7 +84,7 @@ namespace EpicGames.UHT.Types
 					}
 					else
 					{
-						builder.Append(Class.SourceName).Append('*');
+						builder.Append(Class.Namespace.FullSourceName).Append(Class.SourceName).Append('*');
 					}
 					break;
 
@@ -114,11 +114,11 @@ namespace EpicGames.UHT.Types
 					}
 					else if (CppForm == UhtObjectCppForm.TInterfaceInstance)
 					{
-						builder.AppendClassSourceNameOrInterfaceProxyName(ReferencedClass);
+						builder.Append(Class.Namespace.FullSourceName).AppendClassSourceNameOrInterfaceProxyName(ReferencedClass);
 					}
 					else
 					{
-						builder.Append(Class.SourceName);
+						builder.Append(Class.Namespace.FullSourceName).Append(Class.SourceName);
 					}
 					break;
 
@@ -194,7 +194,7 @@ namespace EpicGames.UHT.Types
 			base.ValidateMember(structObj, options);
 			if (Class.NativeInterface != null && (CppForm == UhtObjectCppForm.TObjectPtrObject || CppForm == UhtObjectCppForm.TObjectPtrClass))
 			{
-				this.LogError($"UPROPERTY pointers cannot be interfaces - did you mean TScriptInterface<{Class.SourceName}>?");
+				this.LogError($"UPROPERTY pointers cannot be interfaces - did you mean TScriptInterface<{Class.Namespace.FullSourceName}{Class.SourceName}>?");
 			}
 		}
 	}
diff --git a/Engine/Source/Programs/Shared/EpicGames.UHT/Types/Properties/UhtObjectPropertyBase.cs b/Engine/Source/Programs/Shared/EpicGames.UHT/Types/Properties/UhtObjectPropertyBase.cs
--- a/Engine/Source/Programs/Shared/EpicGames.UHT/Types/Properties/UhtObjectPropertyBase.cs
+++ b/Engine/Source/Programs/Shared/EpicGames.UHT/Types/Properties/UhtObjectPropertyBase.cs
@@ -487,16 +487,15 @@ namespace EpicGames.UHT.Types
 		}
 
 		/// <inheritdoc/>
-		public override void CollectReferencesInternal(IUhtReferenceCollector collector, bool templateProperty)
+		public override void CollectReferencesInternal(IUhtReferenceCollector collector, bool addForwardDeclarations, bool isTemplateProperty)
 		{
+			base.CollectReferencesInternal(collector, addForwardDeclarations, isTemplateProperty);
 			collector.AddCrossModuleReference(Class, false);
 			collector.AddCrossModuleReference(MetaClass, false);
-		}
-
-		/// <inheritdoc/>
-		public override string? GetForwardDeclarations()
-		{
-			return $"class {Class.SourceName};";
+			if (addForwardDeclarations)
+			{
+				collector.AddForwardDeclaration(Class);
+			}
 		}
 
 		/// <inheritdoc/>
@@ -596,7 +595,7 @@ namespace EpicGames.UHT.Types
 			{
 				case UhtObjectCppForm.NativeObject:
 				case UhtObjectCppForm.NativeClass:
-					builder.Append(Class.SourceName);
+					builder.Append(Class.Namespace.FullSourceName).Append(Class.SourceName);
 					if (!isIsWrappedType)
 					{
 						builder.Append('*');
@@ -605,19 +604,28 @@ namespace EpicGames.UHT.Types
 
 				case UhtObjectCppForm.TObjectPtrObject:
 				case UhtObjectCppForm.TObjectPtrClass:
-					builder.Append("TObjectPtr<").Append(Class.SourceName).Append('>');
+					builder.Append("TObjectPtr<").Append(Class.Namespace.FullSourceName).Append(Class.SourceName).Append('>');
 					break;
 
 				case UhtObjectCppForm.TSoftObjectPtr:
-					builder.Append("TSoftObjectPtr<").Append(Class.SourceName).Append('>');
+					builder.Append("TSoftObjectPtr<").Append(Class.Namespace.FullSourceName).Append(Class.SourceName).Append('>');
 					break;
 
 				case UhtObjectCppForm.TSoftClassPtr:
-					builder.Append("TSoftObjectPtr<").Append(Class.SourceName).Append('>');
+					builder.Append("TSoftObjectPtr<").Append(Class.Namespace.FullSourceName).Append(Class.SourceName).Append('>');
 					break;
 
 				case UhtObjectCppForm.TSubclassOf:
-					builder.Append("TSubclassOf<").Append(MetaClass != null ? MetaClass.SourceName : "UClass").Append('>');
+					builder.Append("TSubclassOf<");
+					if (MetaClass != null)
+					{
+						builder.Append(MetaClass.Namespace.FullSourceName).Append(MetaClass.SourceName);
+					}
+					else
+					{
+						builder.Append("UClass");
+					}
+					builder.Append('>');
 					break;
 
 				case UhtObjectCppForm.VerseType:
@@ -629,15 +637,15 @@ namespace EpicGames.UHT.Types
 					break;
 
 				case UhtObjectCppForm.VerseSubtype:
-					builder.Append("verse::subtype<").Append(Class.SourceName).Append('>');
+					builder.Append("verse::subtype<").Append(Class.Namespace.FullSourceName).Append(Class.SourceName).Append('>');
 					break;
 
 				case UhtObjectCppForm.VerseCastableSubtype:
-					builder.Append("verse::castable_subtype<").Append(Class.SourceName).Append('>');
+					builder.Append("verse::castable_subtype<").Append(Class.Namespace.FullSourceName).Append(Class.SourceName).Append('>');
 					break;
 
 				case UhtObjectCppForm.TInterfaceInstance:
-					builder.Append("TInterfaceInstance<").AppendClassSourceNameOrInterfaceProxyName(ReferencedClass).Append('>');
+					builder.Append("TInterfaceInstance<").Append(ReferencedClass.Namespace.FullSourceName).AppendClassSourceNameOrInterfaceProxyName(ReferencedClass).Append('>');
 					break;
 
 				default:
diff --git a/Engine/Source/Programs/Shared/EpicGames.UHT/Types/Properties/UhtOptionalProperty.cs b/Engine/Source/Programs/Shared/EpicGames.UHT/Types/Properties/UhtOptionalProperty.cs
--- a/Engine/Source/Programs/Shared/EpicGames.UHT/Types/Properties/UhtOptionalProperty.cs
+++ b/Engine/Source/Programs/Shared/EpicGames.UHT/Types/Properties/UhtOptionalProperty.cs
@@ -86,18 +86,6 @@ namespace EpicGames.UHT.Types
 			}
 		}
 
-		/// <inheritdoc/>
-		public override void CollectReferencesInternal(IUhtReferenceCollector collector, bool templateProperty)
-		{
-			ValueProperty.CollectReferencesInternal(collector, true);
-		}
-
-		/// <inheritdoc/>
-		public override string? GetForwardDeclarations()
-		{
-			return ValueProperty.GetForwardDeclarations();
-		}
-
 		/// <inheritdoc/>
 		public override IEnumerable<UhtType> EnumerateReferencedTypes()
 		{
diff --git a/Engine/Source/Programs/Shared/EpicGames.UHT/Types/Properties/UhtSetProperty.cs b/Engine/Source/Programs/Shared/EpicGames.UHT/Types/Properties/UhtSetProperty.cs
--- a/Engine/Source/Programs/Shared/EpicGames.UHT/Types/Properties/UhtSetProperty.cs
+++ b/Engine/Source/Programs/Shared/EpicGames.UHT/Types/Properties/UhtSetProperty.cs
@@ -79,18 +79,6 @@ namespace EpicGames.UHT.Types
 			}
 		}
 
-		/// <inheritdoc/>
-		public override void CollectReferencesInternal(IUhtReferenceCollector collector, bool templateProperty)
-		{
-			ValueProperty.CollectReferencesInternal(collector, true);
-		}
-
-		/// <inheritdoc/>
-		public override string? GetForwardDeclarations()
-		{
-			return ValueProperty.GetForwardDeclarations();
-		}
-
 		/// <inheritdoc/>
 		public override IEnumerable<UhtType> EnumerateReferencedTypes()
 		{
diff --git a/Engine/Source/Programs/Shared/EpicGames.UHT/Types/Properties/UhtSoftClassProperty.cs b/Engine/Source/Programs/Shared/EpicGames.UHT/Types/Properties/UhtSoftClassProperty.cs
--- a/Engine/Source/Programs/Shared/EpicGames.UHT/Types/Properties/UhtSoftClassProperty.cs
+++ b/Engine/Source/Programs/Shared/EpicGames.UHT/Types/Properties/UhtSoftClassProperty.cs
@@ -33,15 +33,14 @@ namespace EpicGames.UHT.Types
 		}
 
 		/// <inheritdoc/>
-		public override string? GetForwardDeclarations()
-		{
-			return MetaClass != null ? $"class {MetaClass.SourceName};" : null;
-		}
-
-		/// <inheritdoc/>
-		public override void CollectReferencesInternal(IUhtReferenceCollector collector, bool templateProperty)
+		public override void CollectReferencesInternal(IUhtReferenceCollector collector, bool addForwardDeclarations, bool isTemplateProperty)
 		{
+			base.CollectReferencesInternal(collector, addForwardDeclarations, isTemplateProperty);
 			collector.AddCrossModuleReference(MetaClass, false);
+			if (addForwardDeclarations && MetaClass != null)
+			{
+				collector.AddForwardDeclaration(MetaClass);
+			}
 		}
 
 		/// <inheritdoc/>
diff --git a/Engine/Source/Programs/Shared/EpicGames.UHT/Types/Properties/UhtStructProperty.cs b/Engine/Source/Programs/Shared/EpicGames.UHT/Types/Properties/UhtStructProperty.cs
--- a/Engine/Source/Programs/Shared/EpicGames.UHT/Types/Properties/UhtStructProperty.cs
+++ b/Engine/Source/Programs/Shared/EpicGames.UHT/Types/Properties/UhtStructProperty.cs
@@ -98,19 +98,14 @@ namespace EpicGames.UHT.Types
 		}
 
 		/// <inheritdoc/>
-		public override void CollectReferencesInternal(IUhtReferenceCollector collector, bool templateProperty)
+		public override void CollectReferencesInternal(IUhtReferenceCollector collector, bool addForwardDeclarations, bool isTemplateProperty)
 		{
+			base.CollectReferencesInternal(collector, addForwardDeclarations, isTemplateProperty);
 			collector.AddCrossModuleReference(ScriptStruct, true);
-		}
-
-		/// <inheritdoc/>
-		public override string? GetForwardDeclarations()
-		{
-			if (ScriptStruct.IsCoreType)
+			if (addForwardDeclarations && !ScriptStruct.IsCoreType)
 			{
-				return null;
+				collector.AddForwardDeclaration(ScriptStruct);
 			}
-			return $"struct {ScriptStruct.SourceName};";
 		}
 
 		/// <inheritdoc/>
@@ -524,13 +519,17 @@ namespace EpicGames.UHT.Types
 		}
 
 		/// <inheritdoc/>
-		public override string? GetForwardDeclarations()
+		public override void CollectReferencesInternal(IUhtReferenceCollector collector, bool addForwardDeclarations, bool isTemplateProperty)
 		{
-			if (TemplateArgumentStruct.IsCoreType)
+			base.CollectReferencesInternal(collector, addForwardDeclarations, isTemplateProperty);
+			if (addForwardDeclarations && !TemplateArgumentStruct.IsCoreType)
 			{
-				return null;
+				if (RootProperty.Outer is UhtField outerField)
+				{
+					collector.AddForwardDeclaration(TemplateArgumentStruct);
+					collector.AddForwardDeclaration(outerField.Namespace, $"template struct {TemplateWrapperName}<{TemplateArgumentStruct.SourceName}>;");
+				}
 			}
-			return $"struct {TemplateArgumentStruct.SourceName};\r\ntemplate struct {TemplateWrapperName}<{TemplateArgumentStruct.SourceName}>;";
 		}
 
 		/// <inheritdoc/>
@@ -628,18 +627,21 @@ namespace EpicGames.UHT.Types
 				tokenReader.Require('<');
 			}
 
-			UhtToken identifier = new();
+			UhtTokenList? identifier = null;
 			tokenReader
 				.Optional("struct")
 				.Optional("class")
-				.RequireIdentifier((ref UhtToken token) => { identifier = token; });
-
-			args.Config.RedirectTypeIdentifier(ref identifier);
+				.RequireCppIdentifier(UhtCppIdentifierOptions.None, (UhtTokenList token) => { identifier = token; });
+			if (identifier == null)
+			{
+				throw new UhtIceException("Expected an identifier list");
+			}
+			identifier.RedirectTypeIdentifier(args.Config);
 
 			UhtStructProperty instancedStructProperty = new UhtStructProperty(propertySettings, args.Session.FStateTreePropertyRef!);
 
 			// TStateTreePropertyRef supports UStructs, UClasses, enums and primitive types.
-			UhtType? foundType = propertySettings.Outer.FindType(UhtFindOptions.SourceName | UhtFindOptions.TypesMask, ref identifier);
+			UhtType? foundType = propertySettings.Outer.FindType(UhtFindOptions.SourceName | UhtFindOptions.TypesMask, identifier);
 			if (foundType is UhtStruct foundStruct)
 			{
 				if (foundStruct.IsChildOf(args.Session.UObject))
@@ -658,7 +660,7 @@ namespace EpicGames.UHT.Types
 			else
 			{
 				// It's a primitive or unknown type.
-				rootSettings.MetaData.Add(RefTypeName, identifier.ToString());
+				rootSettings.MetaData.Add(RefTypeName, identifier.Join("::"));
 			}
 
 			if (isRefToArray)
diff --git a/Engine/Source/Programs/Shared/EpicGames.UHT/Types/UhtClass.cs b/Engine/Source/Programs/Shared/EpicGames.UHT/Types/UhtClass.cs
--- a/Engine/Source/Programs/Shared/EpicGames.UHT/Types/UhtClass.cs
+++ b/Engine/Source/Programs/Shared/EpicGames.UHT/Types/UhtClass.cs
@@ -3,6 +3,7 @@
 using System;
 using System.Collections.Generic;
 using System.Diagnostics.CodeAnalysis;
+using System.Linq;
 using System.Text;
 using System.Text.Json.Serialization;
 using EpicGames.Core;
@@ -515,9 +516,10 @@ namespace EpicGames.UHT.Types
 		/// Construct a new instance of the class
 		/// </summary>
 		/// <param name="headerFile">Header being parsed</param>
+		/// <param name="namespaceObj">Namespace where the field was defined</param>
 		/// <param name="outer">The outer type</param>
 		/// <param name="lineNumber">Line number where class begins</param>
-		public UhtClass(UhtHeaderFile headerFile, UhtType outer, int lineNumber) : base(headerFile, outer, lineNumber)
+		public UhtClass(UhtHeaderFile headerFile, UhtNamespace namespaceObj, UhtType outer, int lineNumber) : base(headerFile, namespaceObj, outer, lineNumber)
 		{
 			ClassWithin = this;
 		}
@@ -584,7 +586,7 @@ namespace EpicGames.UHT.Types
 		/// </summary>
 		[JsonIgnore]
 		[SuppressMessage("Usage", "CA2227:Collection properties should be read only", Justification = "<Pending>")]
-		public List<UhtToken>? VerseInterfaceIdentifiers { get; set; } = null;
+		public List<UhtToken[]>? VerseInterfaceIdentifiers { get; set; } = null;
 
 		private static readonly List<UhtClass> s_emptyVerseInterfaces = new();
 		private List<UhtClass>? _verseInterfaces = null;
@@ -1351,6 +1353,12 @@ namespace EpicGames.UHT.Types
 			}
 		}
 
+		/// <inheritdoc/>
+		public override StringBuilder AppendForwardDeclaration(StringBuilder builder)
+		{
+			return builder.Append($"class {SourceName};");
+		}
+
 		/// <summary>
 		/// Given a token stream, verify that it matches the expected signature for a getter/setter
 		/// </summary>
@@ -2178,11 +2186,10 @@ namespace EpicGames.UHT.Types
 		{
 			if (VerseInterfaceIdentifiers != null)
 			{
-				foreach (UhtToken identifier in VerseInterfaceIdentifiers)
+				foreach (UhtToken[] identifier in VerseInterfaceIdentifiers)
 				{
 					// We really only case about interfaces, but we can also handle structs
-					UhtToken scratch = identifier;
-					UhtClass? verseInterface = (UhtClass?)FindType(UhtFindOptions.Class | UhtFindOptions.SourceName | UhtFindOptions.NoSelf, ref scratch);
+					UhtClass? verseInterface = (UhtClass?)FindType(UhtFindOptions.Class | UhtFindOptions.SourceName | UhtFindOptions.NoSelf, identifier);
 					if (verseInterface != null)
 					{
 						_verseInterfaces ??= new List<UhtClass>();
@@ -2190,7 +2197,8 @@ namespace EpicGames.UHT.Types
 					}
 					else
 					{
-						this.LogError($"Verse interface {identifier.Value} not found");
+						string fullIdentifier = String.Join("::", identifier.Select(x => x.Value));
+						this.LogError($"Verse interface {fullIdentifier} not found");
 					}
 				}
 			}
diff --git a/Engine/Source/Programs/Shared/EpicGames.UHT/Types/UhtEnum.cs b/Engine/Source/Programs/Shared/EpicGames.UHT/Types/UhtEnum.cs
--- a/Engine/Source/Programs/Shared/EpicGames.UHT/Types/UhtEnum.cs
+++ b/Engine/Source/Programs/Shared/EpicGames.UHT/Types/UhtEnum.cs
@@ -2,6 +2,7 @@
 
 using System;
 using System.Collections.Generic;
+using System.Text;
 using System.Text.Json.Serialization;
 using EpicGames.Core;
 using EpicGames.UHT.Tables;
@@ -219,14 +220,31 @@ namespace EpicGames.UHT.Types
 		/// Construct a new enumeration
 		/// </summary>
 		/// <param name="headerFile">Header being parsed</param>
+		/// <param name="namespaceObj">Namespace where the field was defined</param>
 		/// <param name="outer">Outer type</param>
 		/// <param name="lineNumber">Line number of declaration</param>
-		public UhtEnum(UhtHeaderFile headerFile, UhtType outer, int lineNumber) : base(headerFile, outer, lineNumber)
+		public UhtEnum(UhtHeaderFile headerFile, UhtNamespace namespaceObj, UhtType outer, int lineNumber) : base(headerFile, namespaceObj, outer, lineNumber)
 		{
 			MetaData.KeyConversion = this;
 			EnumValues = new List<UhtEnumValue>();
 		}
 
+		/// <inheritdoc/>
+		public override StringBuilder AppendForwardDeclaration(StringBuilder builder)
+		{
+			switch (CppForm)
+			{
+				case UhtEnumCppForm.EnumClass:
+					{
+						UhtEnumUnderlyingType underlyingType = UnderlyingType != UhtEnumUnderlyingType.Unspecified ? UnderlyingType : UhtEnumUnderlyingType.Int32;
+						builder.Append($"enum class {SourceName} : {underlyingType.ToString().ToLower()};");
+					}
+					break;
+				// no support for other types ATM
+			}
+			return builder;
+		}
+
 		/// <summary>
 		/// Test to see if the value is a known enum value
 		/// </summary>
diff --git a/Engine/Source/Programs/Shared/EpicGames.UHT/Types/UhtField.cs b/Engine/Source/Programs/Shared/EpicGames.UHT/Types/UhtField.cs
--- a/Engine/Source/Programs/Shared/EpicGames.UHT/Types/UhtField.cs
+++ b/Engine/Source/Programs/Shared/EpicGames.UHT/Types/UhtField.cs
@@ -120,14 +120,21 @@ namespace EpicGames.UHT.Types
 		/// </summary>
 		public bool IsVerseField => VerseName != null;
 
+		/// <summary>
+		/// Namespace where the field was defined
+		/// </summary>
+		public UhtNamespace Namespace { get; init; }
+
 		/// <summary>
 		/// Construct a new field
 		/// </summary>
 		/// <param name="headerFile">Header file being parsed</param>
+		/// <param name="namespaceObj">Namespace where the field was defined</param>
 		/// <param name="outer">Outer object</param>
 		/// <param name="lineNumber">Line number of declaration</param>
-		protected UhtField(UhtHeaderFile headerFile, UhtType outer, int lineNumber) : base(headerFile, outer, lineNumber)
+		protected UhtField(UhtHeaderFile headerFile, UhtNamespace namespaceObj, UhtType outer, int lineNumber) : base(headerFile, outer, lineNumber)
 		{
+			Namespace = namespaceObj;
 		}
 
 		/// <summary>
@@ -278,6 +285,16 @@ namespace EpicGames.UHT.Types
 		{
 			builder.Append(VerseName);
 		}
+
+		/// <summary>
+		/// Append the field's forward declaration
+		/// </summary>
+		/// <param name="builder">Destination builder</param>
+		/// <returns>Builder</returns>
+		public virtual StringBuilder AppendForwardDeclaration(StringBuilder builder)
+		{
+			throw new NotImplementedException();
+		}
 	}
 
 	/// <summary>
@@ -330,5 +347,16 @@ namespace EpicGames.UHT.Types
 			fieldObj.AppendVerseScopeAndName(builder, mode);
 			return builder;
 		}
+
+		/// <summary>
+		/// Append the fields forward declaration
+		/// </summary>
+		/// <param name="builder">Destination builder</param>
+		/// <param name="fieldObj">Field in question</param>
+		/// <returns>Builder</returns>
+		public static StringBuilder AppendForwardDeclaration(this StringBuilder builder, UhtField fieldObj)
+		{
+			return fieldObj.AppendForwardDeclaration(builder);
+		}
 	}
 }
diff --git a/Engine/Source/Programs/Shared/EpicGames.UHT/Types/UhtFunction.cs b/Engine/Source/Programs/Shared/EpicGames.UHT/Types/UhtFunction.cs
--- a/Engine/Source/Programs/Shared/EpicGames.UHT/Types/UhtFunction.cs
+++ b/Engine/Source/Programs/Shared/EpicGames.UHT/Types/UhtFunction.cs
@@ -406,9 +406,10 @@ namespace EpicGames.UHT.Types
 		/// Construct a new instance of a function
 		/// </summary>
 		/// <param name="headerFile">Header being parsed</param>
+		/// <param name="namespaceObj">Namespace where the field was defined</param>
 		/// <param name="outer">The parent object</param>
 		/// <param name="lineNumber">The line number where the function is defined</param>
-		public UhtFunction(UhtHeaderFile headerFile, UhtType outer, int lineNumber) : base(headerFile, outer, lineNumber)
+		public UhtFunction(UhtHeaderFile headerFile, UhtNamespace namespaceObj, UhtType outer, int lineNumber) : base(headerFile, namespaceObj, outer, lineNumber)
 		{
 		}
 
diff --git a/Engine/Source/Programs/Shared/EpicGames.UHT/Types/UhtNamespace.cs b/Engine/Source/Programs/Shared/EpicGames.UHT/Types/UhtNamespace.cs
new file mode 100644
--- /dev/null
+++ b/Engine/Source/Programs/Shared/EpicGames.UHT/Types/UhtNamespace.cs
@@ -0,0 +1,112 @@
+// Copyright Epic Games, Inc. All Rights Reserved.
+
+using System;
+using System.Text;
+using System.Web;
+using EpicGames.Core;
+
+namespace EpicGames.UHT.Types
+{
+
+	/// <summary>
+	/// Represents an element that makes up a namespace.  This includes a previously looked up namespace
+	/// </summary>
+	/// <param name="Name">Name of the namespace</param>
+	/// <param name="Namespace">If already looked up, the namespace object</param>
+	public record struct UhtNamespaceElement(StringView Name, UhtNamespace? Namespace)
+	{
+	}
+
+	/// <summary>
+	/// Represents a C++ namespace
+	/// </summary>
+	public class UhtNamespace
+	{
+
+		/// <summary>
+		/// Parent namespace.  If null, then this is the global namespace
+		/// </summary>
+		public UhtNamespace? Parent { get; init; }
+
+		/// <summary>
+		/// Name of this namespace
+		/// </summary>
+		public string SourceName { get; init; }
+
+		/// <summary>
+		/// The full name of the namespace
+		/// </summary>
+		public string FullSourceName { get; init; }
+
+		/// <summary>
+		/// True if this is the global namespace
+		/// </summary>
+		public bool IsGlobal { get; init; }
+
+		/// <summary>
+		/// Construct a namespace instance
+		/// </summary>
+		/// <param name="parent">Parent namespace</param>
+		/// <param name="sourceName">Name (will be empty if global)</param>
+		/// <param name="fullSourceName">Full name of the namespace</param>
+		/// <param name="isGlobal">True if this is the global namespace</param>
+		public UhtNamespace(UhtNamespace? parent, string sourceName, string fullSourceName, bool isGlobal)
+		{
+			Parent = parent;
+			SourceName = sourceName;
+			FullSourceName = fullSourceName;
+			IsGlobal = isGlobal;
+		}
+
+		/// <summary>
+		/// Append the body inside the namespace on a single line
+		/// </summary>
+		/// <param name="builder">Destination builder</param>
+		/// <param name="body">Body of the namespace</param>
+		/// <returns>Destination builder</returns>
+		public StringBuilder AppendSingleLine(StringBuilder builder, Action<StringBuilder> body)
+		{
+			ReadOnlySpan<char> namespaceName = FullSourceName.AsSpan();
+			if (namespaceName.EndsWith("::"))
+			{
+				namespaceName = namespaceName[..^2];
+			}
+			if (namespaceName.Length > 0)
+			{
+				builder.Append($"namespace {namespaceName} {{ ");
+			}
+			body(builder);
+			if (namespaceName.Length > 0)
+			{
+				builder.Append(" }");
+			}
+			builder.Append("\r\n");
+			return builder;
+		}
+
+		/// <summary>
+		/// Append the body inside the namespace on multiple lines
+		/// </summary>
+		/// <param name="builder">Destination builder</param>
+		/// <param name="body">Body of the namespace</param>
+		/// <returns>Destination builder</returns>
+		public StringBuilder AppendMultipleLines(StringBuilder builder, Action<StringBuilder> body)
+		{
+			ReadOnlySpan<char> namespaceName = FullSourceName.AsSpan();
+			if (namespaceName.EndsWith("::"))
+			{
+				namespaceName = namespaceName[..^2];
+			}
+			if (namespaceName.Length > 0)
+			{
+				builder.Append($"namespace {namespaceName}\r\n{{\r\n");
+			}
+			body(builder);
+			if (namespaceName.Length > 0)
+			{
+				builder.Append($"}} // {namespaceName}\r\n");
+			}
+			return builder;
+		}
+	}
+}
diff --git a/Engine/Source/Programs/Shared/EpicGames.UHT/Types/UhtProperty.cs b/Engine/Source/Programs/Shared/EpicGames.UHT/Types/UhtProperty.cs
--- a/Engine/Source/Programs/Shared/EpicGames.UHT/Types/UhtProperty.cs
+++ b/Engine/Source/Programs/Shared/EpicGames.UHT/Types/UhtProperty.cs
@@ -639,6 +639,11 @@ namespace EpicGames.UHT.Types
 		/// </summary>
 		public UhtStruct OuterStruct { get; }
 
+		/// <summary>
+		/// The outer/owning structure namespace name.  In some cases, this can differ from the OuterStruct.Namespace.FullName
+		/// </summary>
+		public string OuterStructNamespaceName { get; }
+
 		/// <summary>
 		/// The outer/owning structure source name.  In some cases, this can differ from the OuterStruct.SourceName
 		/// </summary>
@@ -1420,7 +1425,7 @@ namespace EpicGames.UHT.Types
 				}
 				else
 				{
-					builder.Append("STRUCT_OFFSET(").Append(context.OuterStructSourceName).Append(", ").Append(SourceName).Append("), ");
+					builder.Append("STRUCT_OFFSET(").Append(context.OuterStructNamespaceName).Append(context.OuterStructSourceName).Append(", ").Append(SourceName).Append("), ");
 				}
 			}
 			return builder;
@@ -2122,11 +2127,7 @@ namespace EpicGames.UHT.Types
 		/// <inheritdoc/>
 		public override void CollectReferences(IUhtReferenceCollector collector)
 		{
-			CollectReferencesInternal(collector, false);
-			if (PropertyFlags.HasAnyFlags(EPropertyFlags.ParmFlags))
-			{
-				collector.AddForwardDeclaration(GetForwardDeclarations());
-			}
+			CollectReferencesInternal(collector, PropertyFlags.HasAnyFlags(EPropertyFlags.ParmFlags), false);
 		}
 
 		/// <summary>
@@ -2134,20 +2135,10 @@ namespace EpicGames.UHT.Types
 		/// collect the contained property's references.
 		/// </summary>
 		/// <param name="collector">Reference collector</param>
+		/// <param name="addForwardDeclarations">If true, forward declarations should be added</param>
 		/// <param name="isTemplateProperty">If true, this is a property in a container</param>
-		public virtual void CollectReferencesInternal(IUhtReferenceCollector collector, bool isTemplateProperty)
-		{
-		}
-
-		/// <summary>
-		/// Return a string of all the forward declarations.
-		/// This method should be removed once C++ UHT is removed.
-		/// This information can be collected via CollectReferences.
-		/// </summary>
-		/// <returns></returns>
-		public virtual string? GetForwardDeclarations()
+		public virtual void CollectReferencesInternal(IUhtReferenceCollector collector, bool addForwardDeclarations, bool isTemplateProperty)
 		{
-			return null;
 		}
 
 		/// <summary>
@@ -2156,7 +2147,7 @@ namespace EpicGames.UHT.Types
 		/// <returns>Type enumerator</returns>
 		public virtual IEnumerable<UhtType> EnumerateReferencedTypes()
 		{
-			return Enumerable.Empty<UhtType>();
+			return [];
 		}
 		#endregion
 
diff --git a/Engine/Source/Programs/Shared/EpicGames.UHT/Types/UhtScriptStruct.cs b/Engine/Source/Programs/Shared/EpicGames.UHT/Types/UhtScriptStruct.cs
--- a/Engine/Source/Programs/Shared/EpicGames.UHT/Types/UhtScriptStruct.cs
+++ b/Engine/Source/Programs/Shared/EpicGames.UHT/Types/UhtScriptStruct.cs
@@ -731,6 +731,12 @@ namespace EpicGames.UHT.Types
 		[JsonIgnore]
 		public override string EngineNamePrefix => IsVerseField ? "" : (Session.Config!.IsStructWithTPrefix(EngineName) ? "T" : "F");
 
+		/// <summary>
+		/// In certain cases (usually with NoExportTypes), a locally generated version of the structures is created and we do code 
+		/// generation based on it.  In such cases, we can't use the namespace.
+		/// </summary>
+		public string NamespaceExportName => !ScriptStructFlags.HasAnyFlags(EStructFlags.Native) && !IsAlwaysAccessible ? "" : Namespace.FullSourceName;
+
 		///<inheritdoc/>
 		[JsonIgnore]
 		protected override UhtSpecifierValidatorTable? SpecifierValidatorTable => Session.GetSpecifierValidatorTable(UhtTableNames.ScriptStruct);
@@ -745,9 +751,10 @@ namespace EpicGames.UHT.Types
 		/// Construct a new script struct
 		/// </summary>
 		/// <param name="headerFile">Header being parsed</param>
+		/// <param name="namespaceObj">Namespace where the field was defined</param>
 		/// <param name="outer">Outer type</param>
 		/// <param name="lineNumber">Line number of the definition</param>
-		public UhtScriptStruct(UhtHeaderFile headerFile, UhtType outer, int lineNumber) : base(headerFile, outer, lineNumber)
+		public UhtScriptStruct(UhtHeaderFile headerFile, UhtNamespace namespaceObj, UhtType outer, int lineNumber) : base(headerFile, namespaceObj, outer, lineNumber)
 		{
 		}
 
@@ -976,6 +983,12 @@ namespace EpicGames.UHT.Types
 			}
 		}
 
+		/// <inheritdoc/>
+		public override StringBuilder AppendForwardDeclaration(StringBuilder builder)
+		{
+			return builder.Append($"struct {SourceName};");
+		}
+
 		/// <inheritdoc/>
 		public override void CollectReferences(IUhtReferenceCollector collector)
 		{
diff --git a/Engine/Source/Programs/Shared/EpicGames.UHT/Types/UhtStruct.cs b/Engine/Source/Programs/Shared/EpicGames.UHT/Types/UhtStruct.cs
--- a/Engine/Source/Programs/Shared/EpicGames.UHT/Types/UhtStruct.cs
+++ b/Engine/Source/Programs/Shared/EpicGames.UHT/Types/UhtStruct.cs
@@ -3,6 +3,7 @@
 using System;
 using System.Collections.Generic;
 using System.Diagnostics.CodeAnalysis;
+using System.Linq;
 using System.Text.Json.Serialization;
 using System.Threading;
 using EpicGames.Core;
@@ -96,7 +97,7 @@ namespace EpicGames.UHT.Types
 		/// Super identifier
 		/// </summary>
 		[JsonIgnore]
-		public UhtToken SuperIdentifier { get; set; }
+		public UhtToken[] SuperIdentifier { get; set; } = [];
 
 		/// <summary>
 		/// Base identifiers
@@ -123,9 +124,10 @@ namespace EpicGames.UHT.Types
 		/// Construct a new instance
 		/// </summary>
 		/// <param name="headerFile">Header file being parsed</param>
+		/// <param name="namespaceObj">Namespace where the field was defined</param>
 		/// <param name="outer">Outer type</param>
 		/// <param name="lineNumber">Line number where definition begins</param>
-		protected UhtStruct(UhtHeaderFile headerFile, UhtType outer, int lineNumber) : base(headerFile, outer, lineNumber)
+		protected UhtStruct(UhtHeaderFile headerFile, UhtNamespace namespaceObj, UhtType outer, int lineNumber) : base(headerFile, namespaceObj, outer, lineNumber)
 		{
 		}
 
@@ -225,14 +227,15 @@ namespace EpicGames.UHT.Types
 		/// <param name="superIdentifier">Token that represent the super</param>
 		/// <param name="findOptions">Find options to restrict types</param>
 		/// <exception cref="UhtException">Thrown if super can not be found</exception>
-		public void BindSuper(UhtToken superIdentifier, UhtFindOptions findOptions)
+		public void BindSuper(UhtToken[] superIdentifier, UhtFindOptions findOptions)
 		{
-			if (superIdentifier)
+			if (superIdentifier.Length > 0)
 			{
-				Super = (UhtStruct?)FindType(findOptions | UhtFindOptions.SourceName | UhtFindOptions.NoSelf, ref superIdentifier);
+				Super = (UhtStruct?)FindType(findOptions | UhtFindOptions.SourceName | UhtFindOptions.NoSelf, superIdentifier);
 				if (Super == null)
 				{
-					throw new UhtException(this, $"Unable to find parent {EngineType.ShortLowercaseText()} type for '{SourceName}' named '{superIdentifier.Value}'");
+					string fullIdentifier = String.Join("::", superIdentifier.Select(x => x.Value));
+					throw new UhtException(this, $"Unable to find parent {EngineType.ShortLowercaseText()} type for '{SourceName}' named '{fullIdentifier}'");
 				}
 				HeaderFile.AddReferencedHeader(Super);
 				MetaData.Parent = Super.MetaData;
diff --git a/Engine/Source/Programs/Shared/EpicGames.UHT/Utils/UhtPropertyResolveArgs.cs b/Engine/Source/Programs/Shared/EpicGames.UHT/Utils/UhtPropertyResolveArgs.cs
--- a/Engine/Source/Programs/Shared/EpicGames.UHT/Utils/UhtPropertyResolveArgs.cs
+++ b/Engine/Source/Programs/Shared/EpicGames.UHT/Utils/UhtPropertyResolveArgs.cs
@@ -223,36 +223,42 @@ namespace EpicGames.UHT.Utils
 			}
 
 			bool isNativeConstTemplateArg = false;
-			UhtToken identifier = new();
+			UhtTokenList? identifier = null;
 			TokenReader
 				.Require('<')
 				.Optional("const", () => { isNativeConstTemplateArg = true; })
 				.Optional("class")
-				.RequireIdentifier((ref UhtToken token) => { identifier = token; })
-				.ExceptionIf("::", "Template arguments to object types do not currently support namespaces")
+				.RequireCppIdentifier(UhtCppIdentifierOptions.None, (UhtTokenList tokenList) => { identifier = tokenList; })
 				.Optional("const", () => { isNativeConstTemplateArg = true; })
 				.Require('>');
 
+			if (identifier == null)
+			{
+				throw new UhtIceException("Expected an identifier list");
+			}
+			identifier.RedirectTypeIdentifier(Config);
+
 			if (isNativeConstTemplateArg)
 			{
 				PropertySettings.MetaData.Add(UhtNames.NativeConstTemplateArg, "");
 			}
 
-			Config.RedirectTypeIdentifier(ref identifier);
-
 			string scratchString;
 			if (mode == UhtTemplateObjectMode.NativeInterfaceProxy)
 			{
-				if (!identifier.Value.Span.StartsWith("I") || !identifier.Value.Span.EndsWith(UhtNames.VerseProxySuffix))
+				UhtTokenList last = identifier.Last;
+				UhtToken token = last.Token;
+				if (!token.Value.Span.StartsWith("I") || !token.Value.Span.EndsWith(UhtNames.VerseProxySuffix))
 				{
-					TokenReader.LogError(identifier.InputLine, $"Expected a Verse interface proxy.  Must begin with 'I' and end with '{UhtNames.VerseProxySuffix}'");
+					TokenReader.LogError(token.InputLine, $"Expected a Verse interface proxy.  Must begin with 'I' and end with '{UhtNames.VerseProxySuffix}'");
 					return null;
 				}
-				scratchString = $"U{identifier.Value.Span[1..(identifier.Value.Length - UhtNames.VerseProxySuffix.Length)]}";
-				identifier.Value = scratchString;
+				scratchString = $"U{token.Value.Span[1..(token.Value.Length - UhtNames.VerseProxySuffix.Length)]}";
+				token.Value = scratchString;
+				last.Token = token;
 			}
 
-			UhtClass? returnClass = PropertySettings.Outer.FindType(UhtFindOptions.SourceName | UhtFindOptions.Class, ref identifier, TokenReader) as UhtClass;
+			UhtClass? returnClass = PropertySettings.Outer.FindType(UhtFindOptions.SourceName | UhtFindOptions.Class, identifier, TokenReader) as UhtClass;
 			if (returnClass != null && returnClass.AlternateObject != null && mode == UhtTemplateObjectMode.Normal)
 			{
 				returnClass = returnClass.AlternateObject as UhtClass;
@@ -284,8 +290,6 @@ namespace EpicGames.UHT.Utils
 		/// <returns>Referenced class</returns>
 		public readonly UhtClass? ParseTemplateClass()
 		{
-			UhtToken identifier = new();
-
 			if (TokenReader.TryOptional("const"))
 			{
 				PropertySettings.MetaData.Add(UhtNames.NativeConst, "");
@@ -296,15 +300,20 @@ namespace EpicGames.UHT.Utils
 				return null;
 			}
 
+			UhtTokenList? identifier = null;
 			TokenReader
 				.Require('<')
 				.Optional("class")
-				.RequireIdentifier((ref UhtToken token) => { identifier = token; })
-				.ExceptionIf("::", "Template arguments to object types do not currently support namespaces")
+				.RequireCppIdentifier(UhtCppIdentifierOptions.None, (UhtTokenList tokenList) => { identifier = tokenList; })
 				.Require('>');
 
-			Config.RedirectTypeIdentifier(ref identifier);
-			UhtClass? returnClass = PropertySettings.Outer.FindType(UhtFindOptions.SourceName | UhtFindOptions.Class, ref identifier, TokenReader) as UhtClass;
+			if (identifier == null)
+			{
+				throw new UhtIceException("Expected an identifier list");
+			}
+			identifier.RedirectTypeIdentifier(Config);
+
+			UhtClass? returnClass = PropertySettings.Outer.FindType(UhtFindOptions.SourceName | UhtFindOptions.Class, identifier, TokenReader) as UhtClass;
 			if (returnClass != null && returnClass.AlternateObject != null)
 			{
 				returnClass = returnClass.AlternateObject as UhtClass;
@@ -318,8 +327,6 @@ namespace EpicGames.UHT.Utils
 		/// <returns>Referenced class</returns>
 		public readonly UhtScriptStruct? ParseTemplateScriptStruct()
 		{
-			UhtToken identifier = new();
-
 			if (TokenReader.TryOptional("const"))
 			{
 				PropertySettings.MetaData.Add(UhtNames.NativeConst, "");
@@ -330,15 +337,20 @@ namespace EpicGames.UHT.Utils
 				return null;
 			}
 
+			UhtTokenList? identifier = null;
 			TokenReader
 				.Require('<')
 				.Optional("struct")
-				.RequireIdentifier((ref UhtToken token) => { identifier = token; })
-				.ExceptionIf("::", "Template arguments to object types do not currently support namespaces")
+				.RequireCppIdentifier(UhtCppIdentifierOptions.None, (UhtTokenList tokenList) => { identifier = tokenList; })
 				.Require('>');
 
-			Config.RedirectTypeIdentifier(ref identifier);
-			return PropertySettings.Outer.FindType(UhtFindOptions.SourceName | UhtFindOptions.ScriptStruct, ref identifier, TokenReader) as UhtScriptStruct;
+			if (identifier == null)
+			{
+				throw new UhtIceException("Expected an identifier list");
+			}
+			identifier.RedirectTypeIdentifier(Config);
+
+			return PropertySettings.Outer.FindType(UhtFindOptions.SourceName | UhtFindOptions.ScriptStruct, identifier, TokenReader) as UhtScriptStruct;
 		}
 
 		/// <summary>
diff --git a/Engine/Source/Programs/Shared/EpicGames.UHT/Utils/UhtReferenceCollector.cs b/Engine/Source/Programs/Shared/EpicGames.UHT/Utils/UhtReferenceCollector.cs
--- a/Engine/Source/Programs/Shared/EpicGames.UHT/Utils/UhtReferenceCollector.cs
+++ b/Engine/Source/Programs/Shared/EpicGames.UHT/Utils/UhtReferenceCollector.cs
@@ -2,6 +2,7 @@
 
 using System;
 using System.Collections.Generic;
+using System.Text;
 using EpicGames.Core;
 using EpicGames.UHT.Types;
 
@@ -41,11 +42,19 @@ namespace EpicGames.UHT.Utils
 		/// <param name="field">Field to be added</param>
 		void AddExportType(UhtField field);
 
+		/// <summary>
+		/// Add a forward declaration.  This is the preferred way of adding a forward declaration since the exporter will
+		/// handle all the formatting requirements.
+		/// </summary>
+		/// <param name="field"></param>
+		void AddForwardDeclaration(UhtField field);
+
 		/// <summary>
 		/// Add a forward declaration.  The string can contain multiple declarations but must only exist on one line.
 		/// </summary>
+		/// <param name="namespaceObj">The namespace to place the string</param>
 		/// <param name="declaration">The declarations to add</param>
-		void AddForwardDeclaration(string? declaration);
+		void AddForwardDeclaration(UhtNamespace? namespaceObj, string? declaration);
 	}
 
 	/// <summary>
@@ -163,38 +172,39 @@ namespace EpicGames.UHT.Utils
 		/// <summary>
 		/// Collection of unique cross module references
 		/// </summary>
-		public UhtUniqueReferenceCollection CrossModule { get; set; } = new UhtUniqueReferenceCollection();
+		public UhtUniqueReferenceCollection CrossModule { get; set; } = new();
 
 		/// <summary>
 		/// Collection of unique declarations
 		/// </summary>
-		public UhtUniqueReferenceCollection Declaration { get; set; } = new UhtUniqueReferenceCollection();
+		public UhtUniqueReferenceCollection Declaration { get; set; } = new();
 
 		/// <summary>
 		/// Collection of singletons
 		/// </summary>
-		public List<UhtField> Singletons { get; } = new List<UhtField>();
+		public List<UhtField> Singletons { get; } = [];
 
 		/// <summary>
 		/// Collection of types to export
 		/// </summary>
-		public List<UhtField> ExportTypes { get; } = new List<UhtField>();
+		public List<UhtField> ExportTypes { get; } = [];
 
 		/// <summary>
-		/// Collection of forward declarations
+		/// Collection of fields needing forward declarations
 		/// </summary>
-		public HashSet<string> ForwardDeclarations { get; } = new HashSet<string>();
+		public HashSet<UhtField> ForwardDeclarations { get; } = [];
 
 		/// <summary>
-		/// Collection of referenced headers
+		/// Collection of forward declarations in text form
 		/// </summary>
-		public HashSet<UhtHeaderFile> ReferencedHeaders { get; } = new HashSet<UhtHeaderFile>();
+		public HashSet<string> ForwardDeclarationStrings { get; } = [];
 
 		/// <summary>
-		/// Add a cross module reference
+		/// Collection of referenced headers
 		/// </summary>
-		/// <param name="obj">Object being referenced</param>
-		/// <param name="registered">True if the object being referenced must be registered</param>
+		public HashSet<UhtHeaderFile> ReferencedHeaders { get; } = [];
+
+		/// <inheritdoc/>
 		public void AddCrossModuleReference(UhtObject? obj, bool registered)
 		{
 			CrossModule.Add(obj, registered);
@@ -204,43 +214,46 @@ namespace EpicGames.UHT.Utils
 			}
 		}
 
-		/// <summary>
-		/// Add a declaration
-		/// </summary>
-		/// <param name="obj">Object being declared</param>
-		/// <param name="registered">True if the object being declared must be registered</param>
+		/// <inheritdoc/>
 		public void AddDeclaration(UhtObject obj, bool registered)
 		{
 			Declaration.Add(obj, registered);
 		}
 
-		/// <summary>
-		/// Add a singleton.  These are added as forward declared functions in the package file.
-		/// </summary>
-		/// <param name="field">Field being added</param>
+		/// <inheritdoc/>
 		public void AddSingleton(UhtField field)
 		{
 			Singletons.Add(field);
 		}
 
-		/// <summary>
-		/// Add a type to be exported.
-		/// </summary>
-		/// <param name="field">Type to be exported</param>
+		/// <inheritdoc/>
 		public void AddExportType(UhtField field)
 		{
 			ExportTypes.Add(field);
 		}
 
-		/// <summary>
-		/// Add a symbol that must be forward declared
-		/// </summary>
-		/// <param name="declaration">Symbol to be forward declared.</param>
-		public void AddForwardDeclaration(string? declaration)
+		/// <inheritdoc/>
+		public void AddForwardDeclaration(UhtField field)
+		{
+			ForwardDeclarations.Add(field);
+		}
+
+		/// <inheritdoc/>
+		public void AddForwardDeclaration(UhtNamespace? namespaceObj, string? declaration)
 		{
 			if (!String.IsNullOrEmpty(declaration))
 			{
-				ForwardDeclarations.Add(declaration);
+				if (namespaceObj == null || namespaceObj.IsGlobal)
+				{
+					ForwardDeclarationStrings.Add(declaration);
+				}
+				else
+				{
+					// This is hardly used so just wrap the declaration in the namespace instead of 
+					// trying to group them all together in a multiple line namespace declaration.
+					StringBuilder builder = new();
+					namespaceObj.AppendSingleLine(builder, (builder) => builder.Append(declaration));
+				}
 			}
 		}
 	}
diff --git a/Engine/Source/Programs/Shared/EpicGames.UHT/Utils/UhtSession.cs b/Engine/Source/Programs/Shared/EpicGames.UHT/Utils/UhtSession.cs
--- a/Engine/Source/Programs/Shared/EpicGames.UHT/Utils/UhtSession.cs
+++ b/Engine/Source/Programs/Shared/EpicGames.UHT/Utils/UhtSession.cs
@@ -1,6 +1,7 @@
 // Copyright Epic Games, Inc. All Rights Reserved.
 
 using System;
+using System.Collections.Concurrent;
 using System.Collections.Generic;
 using System.IO;
 using System.IO.Enumeration;
@@ -907,14 +908,14 @@ namespace EpicGames.UHT.Utils
 		/// </summary>
 		public UhtEnum? EVerseNativeCallResult { get; set; } = null;
 
-		private readonly List<UhtModule> _modules = new();
-		private readonly List<UhtHeaderFile> _headerFiles = new();
-		private readonly List<UhtHeaderFile> _sortedHeaderFiles = new();
+		private readonly List<UhtModule> _modules = [];
+		private readonly List<UhtHeaderFile> _headerFiles = [];
+		private readonly List<UhtHeaderFile> _sortedHeaderFiles = [];
 		private readonly Dictionary<string, UhtHeaderFile> _headerFileDictionary = new(StringComparer.OrdinalIgnoreCase);
 		private readonly Lazy<EngineVersion> _engineVersion;
 		private long _errorCount = 0;
 		private long _warningCount = 0;
-		private readonly List<UhtMessage> _messages = new();
+		private readonly List<UhtMessage> _messages = [];
 		private Task? _messageTask = null;
 		private UhtClass? _uobject = null;
 		private UhtClass? _uclass = null;
@@ -929,9 +930,10 @@ namespace EpicGames.UHT.Utils
 		private bool _symbolTablePopulated = false;
 		private Task? _referenceDeleteTask = null;
 		private readonly Dictionary<string, bool> _exporterStates = new(StringComparer.OrdinalIgnoreCase);
-		private readonly Dictionary<string, EnumAndValue> _fullEnumValueLookup = new();
-		private readonly Dictionary<string, UhtEnum> _shortEnumValueLookup = new();
+		private readonly Dictionary<string, EnumAndValue> _fullEnumValueLookup = [];
+		private readonly Dictionary<string, UhtEnum> _shortEnumValueLookup = [];
 		private JsonDocument? _projectJson = null;
+		private readonly ConcurrentDictionary<string, UhtNamespace> _namespaces = [];
 
 		/// <summary>
 		/// The number of errors
@@ -1451,25 +1453,21 @@ namespace EpicGames.UHT.Utils
 		{
 			ValidateFindOptions(options);
 
-			if (identifiers.Next != null && identifiers.Next.Next != null)
+			// UHT uses a grossly primitive method of resolving namespaces.  We just look at the last two tokens and assume the type will match
+			UhtTokenList lookupIdentifiers = identifiers;
+			while (lookupIdentifiers.Next != null && lookupIdentifiers.Next.Next != null)
 			{
-				if (messageSite != null)
-				{
-					//messageSite.LogError(lineNumber, "UnrealHeaderTool only supports C++ identifiers of two or less identifiers");
-					string fullIdentifier = identifiers.Join("::");
-					FindTypeError(messageSite, lineNumber, options, fullIdentifier);
-					return null;
-				}
+				lookupIdentifiers = lookupIdentifiers.Next;
 			}
 
 			UhtType? type;
-			if (identifiers.Next != null)
+			if (lookupIdentifiers.Next != null)
 			{
-				type = FindTypeTwoNamesInternal(startingType, options, identifiers.Token.Value.ToString(), identifiers.Next.Token.Value.ToString());
+				type = FindTypeTwoNamesInternal(startingType, options, lookupIdentifiers.Token.Value.ToString(), lookupIdentifiers.Next.Token.Value.ToString());
 			}
 			else
 			{
-				type = FindTypeInternal(startingType, options, identifiers.Token.Value.ToString());
+				type = FindTypeInternal(startingType, options, lookupIdentifiers.Token.Value.ToString());
 			}
 
 			if (type == null && messageSite != null)
@@ -1497,29 +1495,20 @@ namespace EpicGames.UHT.Utils
 			{
 				throw new UhtIceException("Empty identifier array");
 			}
-			if (identifiers.Length > 3 || (identifiers.Length > 1 && !identifiers[1].Value.Span.Equals("::", StringComparison.Ordinal)))
-			{
-				if (messageSite != null)
-				{
-					string fullIdentifier = String.Join(null, identifiers);
-					FindTypeError(messageSite, lineNumber, options, fullIdentifier);
-				}
-				return null;
-			}
 
 			UhtType? type;
-			if (identifiers.Length == 3)
+			if (identifiers.Length > 1)
 			{
-				type = FindTypeTwoNamesInternal(startingType, options, identifiers[0].Value.ToString(), identifiers[2].Value.ToString());
+				type = FindTypeTwoNamesInternal(startingType, options, identifiers[^2].Value.ToString(), identifiers[^1].Value.ToString());
 			}
 			else
 			{
-				type = FindTypeInternal(startingType, options, identifiers[0].Value.ToString());
+				type = FindTypeInternal(startingType, options, identifiers[^1].Value.ToString());
 			}
 
 			if (type == null && messageSite != null)
 			{
-				string fullIdentifier = String.Join("::", identifiers);
+				string fullIdentifier = String.Join("::", identifiers.Select(x => x.Value));
 				FindTypeError(messageSite, lineNumber, options, fullIdentifier);
 			}
 			return type;
@@ -1548,7 +1537,7 @@ namespace EpicGames.UHT.Utils
 					UhtType? type = FindTypeInternal(startingType, subOptions, firstName);
 					if (type == null)
 					{
-						return null;
+						return FindTypeInternal(startingType, options, secondName);
 					}
 					if (type is UhtEnum)
 					{
@@ -1567,8 +1556,8 @@ namespace EpicGames.UHT.Utils
 					}
 				}
 
-				// We can't match anything at this point
-				return null;
+				// Just try to match the second name normally
+				return FindTypeInternal(startingType, options, secondName);
 			}
 
 			// Perform the lookup for just a single name
@@ -1692,6 +1681,62 @@ namespace EpicGames.UHT.Utils
 			return null;
 		}
 
+		/// <summary>
+		/// Get the given namespace
+		/// </summary>
+		/// <param name="namespaceElements">Elements of the namespace</param>
+		/// <param name="headerNamespaceLookupCache">Header file local version of the namespace cache (to avoid locking)</param>
+		/// <returns></returns>
+		/// <exception cref="UhtIceException"></exception>
+		public UhtNamespace GetNamespace(List<UhtNamespaceElement> namespaceElements, Dictionary<string, UhtNamespace> headerNamespaceLookupCache)
+		{
+			if (namespaceElements.Count == 0)
+			{
+				throw new UhtIceException("Namespace stack must have at least one element");
+			}
+
+			int index = namespaceElements.Count - 1;
+			if (namespaceElements[index].Namespace != null)
+			{
+				return namespaceElements[index].Namespace!;
+			}
+
+			int firstUnresolved = namespaceElements.FindIndex(x => x.Namespace == null);
+			StringBuilder builder = new();
+			builder.Append(namespaceElements[firstUnresolved - 1].Namespace!.FullSourceName);
+			foreach (UhtNamespaceElement namespaceElement in namespaceElements[firstUnresolved..])
+			{
+				builder.Append(namespaceElement.Name);
+				builder.Append("::");
+			}
+
+			string fullName = builder.ToString();
+			if (headerNamespaceLookupCache.TryGetValue(fullName, out UhtNamespace? namespaceObj))
+			{
+				return namespaceObj;
+			}
+
+			return GetNamespace(namespaceElements, index, builder.ToString(), headerNamespaceLookupCache);
+		}
+
+		private UhtNamespace GetNamespace(List<UhtNamespaceElement> namespaceElements, int index, string fullName, Dictionary<string, UhtNamespace> headerNamespaceLookupCache)
+		{
+			if (!_namespaces.TryGetValue(fullName, out UhtNamespace? namespaceObj))
+			{
+				if (namespaceElements[index - 1].Namespace == null)
+				{
+					GetNamespace(namespaceElements, index - 1, fullName.Substring(0, fullName.Length - 2 - namespaceElements[index].Name.Length), headerNamespaceLookupCache);
+				}
+				namespaceObj = _namespaces.GetOrAdd(fullName, (x) =>
+				{
+					return new(namespaceElements[index - 1].Namespace, namespaceElements[index].Name.ToString(), fullName, false);
+				});
+			}
+			headerNamespaceLookupCache.Add(fullName, namespaceObj);
+			namespaceElements[index] = new(namespaceElements[index].Name, namespaceObj);
+			return namespaceObj;
+		}
+
 		#region IUHTMessageSource implementation
 		/// <summary>
 		/// Add a message to the collection of output messages
@@ -2133,10 +2178,12 @@ namespace EpicGames.UHT.Utils
 			}
 
 			Log.Logger.LogTrace("Step - Parse Headers");
+			UhtNamespace globalNamespace = new UhtNamespace(null, "", "", true);
+			_namespaces.TryAdd("", globalNamespace);
 			ForEachHeader(headerFile =>
 			{
 				headerFile.Read();
-				UhtHeaderFileParser.Parse(headerFile);
+				UhtHeaderFileParser.Parse(headerFile, globalNamespace);
 			}, true);
 		}
 
