Various minor dllstorage fixes to be able to compile with merged modules enabled in combination with ltcg/thinlto

diff --git a/Engine/Source/Editor/EditorFramework/Public/EditorModes.h b/Engine/Source/Editor/EditorFramework/Public/EditorModes.h
--- a/Engine/Source/Editor/EditorFramework/Public/EditorModes.h
+++ b/Engine/Source/Editor/EditorFramework/Public/EditorModes.h
@@ -50,7 +50,7 @@ namespace FBuiltinEditorModes
 };
 
 /** Material proxy wrapper that can be created on the game thread and passed on to the render thread. */
-class UNREALED_API FDynamicColoredMaterialRenderProxy : public FDynamicPrimitiveResource, public FColoredMaterialRenderProxy
+class FDynamicColoredMaterialRenderProxy : public FDynamicPrimitiveResource, public FColoredMaterialRenderProxy
 {
 public:
 	/** Initialization constructor. */
diff --git a/Engine/Source/Programs/Shared/EpicGames.UHT/Exporters/CodeGen/UhtHeaderCodeGeneratorHFile.cs b/Engine/Source/Programs/Shared/EpicGames.UHT/Exporters/CodeGen/UhtHeaderCodeGeneratorHFile.cs
--- a/Engine/Source/Programs/Shared/EpicGames.UHT/Exporters/CodeGen/UhtHeaderCodeGeneratorHFile.cs
+++ b/Engine/Source/Programs/Shared/EpicGames.UHT/Exporters/CodeGen/UhtHeaderCodeGeneratorHFile.cs
@@ -262,13 +262,7 @@ namespace EpicGames.UHT.Exporters.CodeGen
 					constPrefix = "const ";
 				}
 
-				builder.Append('\t');
-				if (!scriptStruct.ScriptStructFlags.HasAnyFlags(EStructFlags.RequiredAPI) && 
-				    (Module.Module.AlwaysExportStructs || scriptStruct.ScriptStructExportFlags.HasAnyFlags(UhtScriptStructExportFlags.MinimalAPI)))
-				{
-					builder.Append(Module.Api);
-				}
-				builder.Append("virtual const UScriptStruct* GetExecuteContextStruct() const override { return ");
+				builder.Append("\tvirtual const UScriptStruct* GetExecuteContextStruct() const override { return ");
 				builder.Append(scriptStruct.RigVMStructInfo.ExecuteContextType);
 				builder.Append("::StaticStruct(); } \\\r\n");
 				
diff --git a/Engine/Source/Runtime/Core/Public/ProfilingDebugging/CookStats.h b/Engine/Source/Runtime/Core/Public/ProfilingDebugging/CookStats.h
--- a/Engine/Source/Runtime/Core/Public/ProfilingDebugging/CookStats.h
+++ b/Engine/Source/Runtime/Core/Public/ProfilingDebugging/CookStats.h
@@ -438,7 +438,7 @@ struct FDerivedDataUsageStats
 /** any expression inside this block will be compiled out if ENABLE_COOK_STATS is not true. Shorthand for the multi-line #if ENABLE_COOK_STATS guard. */
 #define COOK_STAT(...) __VA_ARGS__
 
-CORE_API UE_TRACE_CHANNEL_EXTERN(CookChannel);
+UE_TRACE_CHANNEL_EXTERN(CookChannel, CORE_API);
 
 #define UE_INSIGHTS_COOK_PROFILER_VERSION 2
 
diff --git a/Engine/Source/Runtime/Engine/Classes/EdGraph/EdGraphNode.h b/Engine/Source/Runtime/Engine/Classes/EdGraph/EdGraphNode.h
--- a/Engine/Source/Runtime/Engine/Classes/EdGraph/EdGraphNode.h
+++ b/Engine/Source/Runtime/Engine/Classes/EdGraph/EdGraphNode.h
@@ -857,7 +857,7 @@ public:
 
 	/** Update node size to new value */
 	UE_DEPRECATED(5.6, "Use the version of the function accepting FVector2f")
-	ENGINE_API virtual void ResizeNode(const FVector2D& NewSize) UE_SLATE_DEPRECATED_VECTOR_VIRTUAL_FUNCTION {}
+	virtual void ResizeNode(const FVector2D& NewSize) UE_SLATE_DEPRECATED_VECTOR_VIRTUAL_FUNCTION {}
 	ENGINE_API virtual void ResizeNode(const FVector2f& NewSize);
 
 	/**
diff --git a/Engine/Source/Runtime/Engine/Private/WorldPartition/WorldPartition.cpp b/Engine/Source/Runtime/Engine/Private/WorldPartition/WorldPartition.cpp
--- a/Engine/Source/Runtime/Engine/Private/WorldPartition/WorldPartition.cpp
+++ b/Engine/Source/Runtime/Engine/Private/WorldPartition/WorldPartition.cpp
@@ -2156,6 +2156,8 @@ void UWorldPartition::DumpActorDescs(const FString& Path)
 	}
 }
 
+TOptional<UWorldPartition::FStreamingGenerationErrorHandlerOverride> UWorldPartition::StreamingGenerationErrorHandlerOverride;
+
 void UWorldPartition::AppendAssetRegistryTags(TArray<FAssetRegistryTag>& OutTags) const
 {
 	FAssetRegistryTagsContextData Context(this, EAssetRegistryTagsCaller::Uncategorized);
diff --git a/Engine/Source/Runtime/Engine/Public/WorldPartition/WorldPartition.h b/Engine/Source/Runtime/Engine/Public/WorldPartition/WorldPartition.h
--- a/Engine/Source/Runtime/Engine/Public/WorldPartition/WorldPartition.h
+++ b/Engine/Source/Runtime/Engine/Public/WorldPartition/WorldPartition.h
@@ -405,7 +405,7 @@ public:
 	static ENGINE_API void CheckForErrors(const FCheckForErrorsParams& Params);
 
 	using FStreamingGenerationErrorHandlerOverride = TFunction<IStreamingGenerationErrorHandler*(IStreamingGenerationErrorHandler* InErrorHandler)>;
-	ENGINE_API inline static TOptional<FStreamingGenerationErrorHandlerOverride> StreamingGenerationErrorHandlerOverride;
+	ENGINE_API static TOptional<FStreamingGenerationErrorHandlerOverride> StreamingGenerationErrorHandlerOverride;
 
 	ENGINE_API void AppendAssetRegistryTags(FAssetRegistryTagsContext Context) const;
 	UE_DEPRECATED(5.4, "Implement the version that takes FAssetRegistryTagsContext instead.")
