[UHT]
* Made method AppendRigVmStructInfoMethods static to make static analyzer happy

diff --git a/Engine/Source/Programs/Shared/EpicGames.UHT/Exporters/CodeGen/UhtHeaderCodeGeneratorHFile.cs b/Engine/Source/Programs/Shared/EpicGames.UHT/Exporters/CodeGen/UhtHeaderCodeGeneratorHFile.cs
--- a/Engine/Source/Programs/Shared/EpicGames.UHT/Exporters/CodeGen/UhtHeaderCodeGeneratorHFile.cs
+++ b/Engine/Source/Programs/Shared/EpicGames.UHT/Exporters/CodeGen/UhtHeaderCodeGeneratorHFile.cs
@@ -250,7 +250,7 @@ namespace EpicGames.UHT.Exporters.CodeGen
 			return builder;
 		}
 
-		private StringBuilder AppendRigVmStructInfoMethods(StringBuilder builder, UhtScriptStruct scriptStruct)
+		private static StringBuilder AppendRigVmStructInfoMethods(StringBuilder builder, UhtScriptStruct scriptStruct)
 		{
 			// if we have RigVM methods on this struct we need to 
 			// declare the static method as well as the stub method
