Fixing static analysis issue

#rnx

diff --git a/Engine/Source/Programs/Shared/EpicGames.UHT/Types/UhtProperty.cs b/Engine/Source/Programs/Shared/EpicGames.UHT/Types/UhtProperty.cs
--- a/Engine/Source/Programs/Shared/EpicGames.UHT/Types/UhtProperty.cs
+++ b/Engine/Source/Programs/Shared/EpicGames.UHT/Types/UhtProperty.cs
@@ -1049,7 +1049,7 @@ namespace EpicGames.UHT.Types
 		/// <summary>
 		/// If set, the name was set via a UPROPERTY or UPARAM
 		/// </summary>
-		public string VerseNameOverride = String.Empty;
+		public string VerseNameOverride { get; set; } = String.Empty;
 
 		/// <summary>
 		/// Verse name of the property 
