Expand the EpicGames.Perforce API:
- Add partial support for the "p4 property" command (list single property only)
- "p4 info" now also retains the serverCluster property

[RN] rnx

diff --git a/Engine/Source/Programs/Shared/EpicGames.Perforce/PerforceConnection.cs b/Engine/Source/Programs/Shared/EpicGames.Perforce/PerforceConnection.cs
--- a/Engine/Source/Programs/Shared/EpicGames.Perforce/PerforceConnection.cs
+++ b/Engine/Source/Programs/Shared/EpicGames.Perforce/PerforceConnection.cs
@@ -2868,6 +2868,42 @@ namespace EpicGames.Perforce
 
 		#endregion
 
+		#region property
+
+		// ATTN: We currently support only getting a single value (p4 property -l -n)
+
+		/// <summary>
+		/// Execute the 'property' command for a single value 
+		/// </summary>
+		/// <param name="connection">Connection to the Perforce server</param>
+		/// <param name="propertyName">Name of the property</param>
+		/// <param name="cancellationToken">Token used to cancel the operation</param>
+		/// <returns>Value of the property</returns>
+		public static async Task<string> GetProperty(this IPerforceConnection connection, string propertyName, CancellationToken cancellationToken = default)
+		{
+			return (await TryGetSingleProperty(connection, propertyName, cancellationToken)).Data.Value;
+		}
+
+
+		/// <summary>
+		/// Execute the 'property' command for a single value 
+		/// </summary>
+		/// <param name="connection">Connection to the Perforce server</param>
+		/// <param name="propertyName">Name of the property</param>
+		/// <param name="cancellationToken">Token used to cancel the operation</param>
+		/// <returns>Value of the property</returns>
+		public static async Task<PerforceResponse<PropertyRecord>> TryGetSingleProperty(this IPerforceConnection connection, string propertyName, CancellationToken cancellationToken = default)
+		{
+			List<string> arguments = new List<string>();
+			arguments.Add("-l");
+			arguments.Add("-n");
+			arguments.Add(propertyName);
+			
+			return await SingleResponseCommandAsync<PropertyRecord>(connection, "property", arguments, null, cancellationToken);
+		}
+
+		#endregion
+
 		#region p4 reconcile
 
 		/// <summary>
diff --git a/Engine/Source/Programs/Shared/EpicGames.Perforce/Records/InfoRecord.cs b/Engine/Source/Programs/Shared/EpicGames.Perforce/Records/InfoRecord.cs
--- a/Engine/Source/Programs/Shared/EpicGames.Perforce/Records/InfoRecord.cs
+++ b/Engine/Source/Programs/Shared/EpicGames.Perforce/Records/InfoRecord.cs
@@ -45,6 +45,11 @@ namespace EpicGames.Perforce
 		/// </summary>
 		[PerforceTag("serverAddress", Optional = true)]
 		public string? ServerAddress { get; set; }
+		/// <summary>
+		/// Address of the Perforce server
+		/// </summary>
+		[PerforceTag("serverCluster", Optional = true)]
+		public string? ServerCluster { get; set; }
 
 		/// <summary>
 		/// Date and time on the server
diff --git a/Engine/Source/Programs/Shared/EpicGames.Perforce/Records/PropertyRecord.cs b/Engine/Source/Programs/Shared/EpicGames.Perforce/Records/PropertyRecord.cs
new file mode 100644
--- /dev/null
+++ b/Engine/Source/Programs/Shared/EpicGames.Perforce/Records/PropertyRecord.cs
@@ -0,0 +1,25 @@
+// Copyright Epic Games, Inc. All Rights Reserved.
+
+using System;
+using System.Collections.Generic;
+
+namespace EpicGames.Perforce
+{
+	/// <summary>
+	/// Contains information about a Perforce property
+	/// </summary>
+	public class PropertyRecord
+	{
+		/// <summary>
+		/// The name of the property
+		/// </summary>
+		[PerforceTag("name")]
+		public string Name { get; set; } = String.Empty;
+
+		/// <summary>
+		/// The value of the property
+		/// </summary>
+		[PerforceTag("value")]
+		public string Value { get; set; } = String.Empty;
+	}
+}
