Fix Horde warnings in PerforceConnection.cs and PropertyRecord.cs

#rnx

diff --git a/Engine/Source/Programs/Shared/EpicGames.Perforce/PerforceConnection.cs b/Engine/Source/Programs/Shared/EpicGames.Perforce/PerforceConnection.cs
--- a/Engine/Source/Programs/Shared/EpicGames.Perforce/PerforceConnection.cs
+++ b/Engine/Source/Programs/Shared/EpicGames.Perforce/PerforceConnection.cs
@@ -2879,12 +2879,11 @@ namespace EpicGames.Perforce
 		/// <param name="propertyName">Name of the property</param>
 		/// <param name="cancellationToken">Token used to cancel the operation</param>
 		/// <returns>Value of the property</returns>
-		public static async Task<string> GetProperty(this IPerforceConnection connection, string propertyName, CancellationToken cancellationToken = default)
+		public static async Task<string> GetPropertyAsync(this IPerforceConnection connection, string propertyName, CancellationToken cancellationToken = default)
 		{
-			return (await TryGetSingleProperty(connection, propertyName, cancellationToken)).Data.Value;
+			return (await TryGetSinglePropertyAsync(connection, propertyName, cancellationToken)).Data.Value;
 		}
 
-
 		/// <summary>
 		/// Execute the 'property' command for a single value 
 		/// </summary>
@@ -2892,7 +2891,7 @@ namespace EpicGames.Perforce
 		/// <param name="propertyName">Name of the property</param>
 		/// <param name="cancellationToken">Token used to cancel the operation</param>
 		/// <returns>Value of the property</returns>
-		public static async Task<PerforceResponse<PropertyRecord>> TryGetSingleProperty(this IPerforceConnection connection, string propertyName, CancellationToken cancellationToken = default)
+		public static async Task<PerforceResponse<PropertyRecord>> TryGetSinglePropertyAsync(this IPerforceConnection connection, string propertyName, CancellationToken cancellationToken = default)
 		{
 			List<string> arguments = new List<string>();
 			arguments.Add("-l");
diff --git a/Engine/Source/Programs/Shared/EpicGames.Perforce/Records/PropertyRecord.cs b/Engine/Source/Programs/Shared/EpicGames.Perforce/Records/PropertyRecord.cs
--- a/Engine/Source/Programs/Shared/EpicGames.Perforce/Records/PropertyRecord.cs
+++ b/Engine/Source/Programs/Shared/EpicGames.Perforce/Records/PropertyRecord.cs
@@ -1,7 +1,6 @@
 // Copyright Epic Games, Inc. All Rights Reserved.
 
 using System;
-using System.Collections.Generic;
 
 namespace EpicGames.Perforce
 {
