EpicGames.Build: If UnrealBuildTool is the entry assembly, use it's location as the dll path

#rnx

diff --git a/Engine/Source/Programs/Shared/EpicGames.Build/Unreal.cs b/Engine/Source/Programs/Shared/EpicGames.Build/Unreal.cs
--- a/Engine/Source/Programs/Shared/EpicGames.Build/Unreal.cs
+++ b/Engine/Source/Programs/Shared/EpicGames.Build/Unreal.cs
@@ -99,6 +99,13 @@ namespace UnrealBuildBase
 
 		private static FileReference FindUnrealBuildToolDll()
 		{
+			// Return the entry assembly location if it is UnrealBuildTool
+			FileReference? entryDll = FileReference.FromString(Assembly.GetEntryAssembly()?.GetOriginalLocation());
+			if (entryDll?.GetFileName().Equals("UnrealBuildTool.dll", StringComparison.Ordinal) == true)
+			{
+				return FileReference.FindCorrectCase(entryDll);
+			}
+
 			// UnrealBuildTool.dll is assumed to be located under {RootDirectory}/Engine/Binaries/DotNET/UnrealBuildTool/
 			FileReference UnrealBuildToolDllPath = FileReference.Combine(EngineDirectory, "Binaries", "DotNET", "UnrealBuildTool", "UnrealBuildTool.dll");
 
