Horde: Add C# implementation of UBA crypto

diff --git a/Engine/Source/Programs/Shared/EpicGames.Horde.Tests/Compute/UbaCryptoTests.cs b/Engine/Source/Programs/Shared/EpicGames.Horde.Tests/Compute/UbaCryptoTests.cs
new file mode 100644
--- /dev/null
+++ b/Engine/Source/Programs/Shared/EpicGames.Horde.Tests/Compute/UbaCryptoTests.cs
@@ -0,0 +1,45 @@
+// Copyright Epic Games, Inc. All Rights Reserved.
+
+using System;
+using System.Text;
+using EpicGames.Horde.Compute;
+using Microsoft.VisualStudio.TestTools.UnitTesting;
+
+namespace EpicGames.Horde.Tests.Compute;
+
+[TestClass]
+public class UbaCryptoTests
+{
+	private readonly byte[] _key = [0x01, 0x23, 0x45, 0x67, 0x89, 0xAB, 0xCD, 0xEF, 0xFE, 0xDC, 0xBA, 0x98, 0x76, 0x54, 0x32, 0x10];
+	
+	private const string Decrypted1 = "This is 15 char";
+	private const string Encrypted1 = "54 68 69 73 20 69 73 20 31 35 20 63 68 61 72";
+	private const string Decrypted2 = "This is 16 chars";
+	private const string Encrypted2 = "B6 72 E4 F2 3B 6F 86 05 2A BC AA B3 69 30 39 42";
+	private const string Decrypted3 = "This is 17 chars!";
+	private const string Encrypted3 = "81 98 1C 37 C8 48 09 25 6A 0D C2 24 E7 82 71 9D 21";
+	
+	[TestMethod]
+	[DataRow(Decrypted1, Encrypted1)]
+	[DataRow(Decrypted2, Encrypted2)]
+	[DataRow(Decrypted3, Encrypted3)]
+	public void Encrypt(string data, string expected)
+	{
+		using UbaCrypto crypto = new(_key);
+		byte[] actual = crypto.Encrypt(Encoding.UTF8.GetBytes(data));
+		Assert.AreEqual(expected.Replace(" ", "", StringComparison.Ordinal), Convert.ToHexString(actual));
+	}
+	
+	[TestMethod]
+	[DataRow(Encrypted1, Decrypted1)]
+	[DataRow(Encrypted2, Decrypted2)]
+	[DataRow(Encrypted3, Decrypted3)]
+	public void Decrypt(string encryptedHexString, string expected)
+	{
+		byte[] encryptedData = Convert.FromHexString(encryptedHexString.Replace(" ", "", StringComparison.Ordinal));
+		using UbaCrypto crypto = new(_key);
+		byte[] actual = crypto.Decrypt(encryptedData);
+		Assert.AreEqual(expected, Encoding.UTF8.GetString(actual));
+	}
+}
+
diff --git a/Engine/Source/Programs/Shared/EpicGames.Horde/Compute/UbaCrypto.cs b/Engine/Source/Programs/Shared/EpicGames.Horde/Compute/UbaCrypto.cs
new file mode 100644
--- /dev/null
+++ b/Engine/Source/Programs/Shared/EpicGames.Horde/Compute/UbaCrypto.cs
@@ -0,0 +1,137 @@
+﻿// Copyright Epic Games, Inc. All Rights Reserved.
+
+using System;
+using System.Security.Cryptography;
+using System.Threading;
+
+namespace EpicGames.Horde.Compute;
+
+/// <summary>
+/// Provides AES encryption and decryption functionality compatible with Unreal Build Accelerator (UBA) communication protocol.
+/// Uses AES-128 in CBC mode with zero IV and no padding. Thread-safe through thread-local AES instances.
+/// </summary>
+public class UbaCrypto : IDisposable
+{
+	private const int AesBlockSize = 16;
+	private readonly ThreadLocal<Aes> _aes;
+	private bool _disposed;
+	
+	/// <summary>
+	/// Constructor
+	/// </summary>
+	/// <param name="key">The 128-bit (16 byte) AES encryption key.</param>
+	/// <exception cref="ArgumentException"></exception>
+	public UbaCrypto(byte[] key)
+	{
+		if (key is not { Length: 16 })
+		{
+			throw new ArgumentException("Key must be exactly 16 bytes");
+		}
+		
+		_aes = new ThreadLocal<Aes>(() =>
+		{
+			Aes aes = Aes.Create();
+			aes.Key = key;
+			aes.Mode = CipherMode.CBC;
+			aes.Padding = PaddingMode.None;
+			return aes;
+		});
+	}
+	
+	/// <summary>
+	/// Encrypts the specified data
+	/// </summary>
+	/// <param name="data">The data to encrypt. Can be any length.</param>
+	/// <returns>
+	/// A byte array containing the encrypted data. The array length matches the input length.
+	/// Only complete 16-byte blocks are encrypted; remaining bytes are left unmodified.
+	/// </returns>
+	public byte[] Encrypt(byte[] data)
+	{
+		if (data == null || data.Length == 0)
+		{
+			throw new ArgumentException("Data cannot be null or empty");
+		}
+		
+		byte[] result = new byte[data.Length];
+		Array.Copy(data, result, data.Length);
+		
+		int alignedSize = (data.Length / AesBlockSize) * AesBlockSize;
+		if (alignedSize > 0)
+		{
+			Aes aes = _aes.Value ?? throw new Exception("Failed to acquire thread-local AES instance");
+
+			// Use zero IV to match BCrypt behavior with nullptr IV
+			aes.IV = new byte[AesBlockSize]; // All zeros
+			
+			using ICryptoTransform encryptor = aes.CreateEncryptor();
+			byte[] alignedData = new byte[alignedSize];
+			Array.Copy(result, alignedData, alignedSize);
+			
+			byte[] encrypted = encryptor.TransformFinalBlock(alignedData, 0, alignedSize);
+			Array.Copy(encrypted, result, alignedSize);
+		}
+		
+		return result;
+	}
+	
+	/// <summary>
+	/// Decrypts the specified data.
+	/// </summary>
+	/// <param name="data">The data to decrypt. Can be any length.</param>
+	/// <returns>
+	/// A byte array containing the decrypted data. The array length matches the input length.
+	/// Only complete 16-byte blocks are decrypted; remaining bytes are left unmodified.
+	/// </returns>
+	public byte[] Decrypt(byte[] data)
+	{
+		if (data == null || data.Length == 0)
+		{
+			throw new ArgumentException("Data cannot be null or empty");
+		}
+		
+		byte[] result = new byte[data.Length];
+		Array.Copy(data, result, data.Length);
+		
+		int alignedSize = (data.Length / AesBlockSize) * AesBlockSize;
+		if (alignedSize > 0)
+		{
+			Aes aes = _aes.Value ?? throw new Exception("Failed to acquire thread-local AES instance");
+			
+			// Use zero IV to match BCrypt behavior with nullptr IV
+			aes.IV = new byte[AesBlockSize]; // All zeros
+			
+			using ICryptoTransform decryptor = aes.CreateDecryptor();
+			byte[] alignedData = new byte[alignedSize];
+			Array.Copy(result, alignedData, alignedSize);
+			
+			byte[] decrypted = decryptor.TransformFinalBlock(alignedData, 0, alignedSize);
+			Array.Copy(decrypted, result, alignedSize);
+		}
+		
+		return result;
+	}
+	
+	/// <inheritdoc />
+	public void Dispose()
+	{
+		Dispose(true);
+		GC.SuppressFinalize(this);
+	}
+	
+	/// <summary>
+	/// Dispose
+	/// </summary>
+	protected virtual void Dispose(bool disposing)
+	{
+		if (!_disposed)
+		{
+			if (disposing)
+			{
+				_aes.Dispose();
+			}
+			_disposed = true;
+		}
+	}
+}
+
