Fixing static analysis issue.

#rnx

diff --git a/Engine/Source/Programs/Shared/EpicGames.Core/UnrealEngineTypes.cs b/Engine/Source/Programs/Shared/EpicGames.Core/UnrealEngineTypes.cs
--- a/Engine/Source/Programs/Shared/EpicGames.Core/UnrealEngineTypes.cs
+++ b/Engine/Source/Programs/Shared/EpicGames.Core/UnrealEngineTypes.cs
@@ -1578,9 +1578,9 @@ namespace EpicGames.Core
 		EpicInternalConstructor = 0x00000400, // This class' has a <castable> attribute's constructor is epic_internal
 
 		// @TODO: this should be a per-function flag; a class flag is not granular enough
-		Err_Inoperable = 0x40000000, // One or more of the class's functions contain mis-linked (malformed) bytecode
-		Err_Incomplete = 0x80000000, // The class layout is malformed (missing super, illformed data-member, etc.)
+		ErrInoperable = 0x40000000, // One or more of the class's functions contain mis-linked (malformed) bytecode
+		ErrIncomplete = 0x80000000, // The class layout is malformed (missing super, illformed data-member, etc.)
 
-		Err = (Err_Incomplete | Err_Inoperable)
+		Err = (ErrIncomplete | ErrInoperable)
 	};
 }
