Configure debug info type in shared props file and allow environment to override type to portable

diff --git a/Engine/Source/Programs/Shared/EpicGames.Build/EpicGames.Build.csproj b/Engine/Source/Programs/Shared/EpicGames.Build/EpicGames.Build.csproj
--- a/Engine/Source/Programs/Shared/EpicGames.Build/EpicGames.Build.csproj
+++ b/Engine/Source/Programs/Shared/EpicGames.Build/EpicGames.Build.csproj
@@ -10,8 +10,6 @@
     <GenerateTargetFrameworkAttribute>false</GenerateTargetFrameworkAttribute>
     <Configurations>Debug;Release;Development;Analyze</Configurations>
     <Nullable>enable</Nullable>
-    <DebugType>pdbonly</DebugType>
-    <DebugType Condition="'$([System.Runtime.InteropServices.RuntimeInformation]::IsOSPlatform($([System.Runtime.InteropServices.OSPlatform]::Windows)))' == 'true' And '$([System.Runtime.InteropServices.RuntimeInformation]::OSArchitecture)' == 'Arm64'">portable</DebugType>
     <ProduceReferenceAssemblyInOutDir>true</ProduceReferenceAssemblyInOutDir>
     <AllowUnsafeBlocks>true</AllowUnsafeBlocks>
     <RunAnalyzersDuringBuild>false</RunAnalyzersDuringBuild>
diff --git a/Engine/Source/Programs/Shared/EpicGames.Core/EpicGames.Core.csproj b/Engine/Source/Programs/Shared/EpicGames.Core/EpicGames.Core.csproj
--- a/Engine/Source/Programs/Shared/EpicGames.Core/EpicGames.Core.csproj
+++ b/Engine/Source/Programs/Shared/EpicGames.Core/EpicGames.Core.csproj
@@ -12,8 +12,6 @@
     <Configurations>Debug;Release;Development;Analyze</Configurations>
     <Nullable>enable</Nullable>
     <AllowUnsafeBlocks>true</AllowUnsafeBlocks>
-    <DebugType>pdbonly</DebugType>
-    <DebugType Condition="'$([System.Runtime.InteropServices.RuntimeInformation]::IsOSPlatform($([System.Runtime.InteropServices.OSPlatform]::Windows)))' == 'true' And '$([System.Runtime.InteropServices.RuntimeInformation]::OSArchitecture)' == 'Arm64'">portable</DebugType>
     <RunAnalyzersDuringBuild>false</RunAnalyzersDuringBuild>
   </PropertyGroup>
   <PropertyGroup Condition=" '$(Configuration)|$(Platform)' == 'Debug|AnyCPU' ">
diff --git a/Engine/Source/Programs/Shared/EpicGames.Horde/EpicGames.Horde.csproj b/Engine/Source/Programs/Shared/EpicGames.Horde/EpicGames.Horde.csproj
--- a/Engine/Source/Programs/Shared/EpicGames.Horde/EpicGames.Horde.csproj
+++ b/Engine/Source/Programs/Shared/EpicGames.Horde/EpicGames.Horde.csproj
@@ -11,8 +11,6 @@
     <GenerateTargetFrameworkAttribute>false</GenerateTargetFrameworkAttribute>
     <Configurations>Debug;Release;Development;Analyze</Configurations>
     <Nullable>enable</Nullable>
-    <DebugType>pdbonly</DebugType>
-    <DebugType Condition="'$([System.Runtime.InteropServices.RuntimeInformation]::IsOSPlatform($([System.Runtime.InteropServices.OSPlatform]::Windows)))' == 'true' And '$([System.Runtime.InteropServices.RuntimeInformation]::OSArchitecture)' == 'Arm64'">portable</DebugType>
     <GenerateDocumentationFile>True</GenerateDocumentationFile>
     <AllowUnsafeBlocks>True</AllowUnsafeBlocks>
     <RunAnalyzersDuringBuild>false</RunAnalyzersDuringBuild>
diff --git a/Engine/Source/Programs/Shared/EpicGames.IoHash/EpicGames.IoHash.csproj b/Engine/Source/Programs/Shared/EpicGames.IoHash/EpicGames.IoHash.csproj
--- a/Engine/Source/Programs/Shared/EpicGames.IoHash/EpicGames.IoHash.csproj
+++ b/Engine/Source/Programs/Shared/EpicGames.IoHash/EpicGames.IoHash.csproj
@@ -11,9 +11,6 @@
     <GenerateTargetFrameworkAttribute>false</GenerateTargetFrameworkAttribute>
     <Configurations>Debug;Release;Development;Analyze</Configurations>
     <Nullable>enable</Nullable>
-    <DebugType>full</DebugType>
-    <DebugType>pdbonly</DebugType>
-    <DebugType Condition="'$([System.Runtime.InteropServices.RuntimeInformation]::IsOSPlatform($([System.Runtime.InteropServices.OSPlatform]::Windows)))' == 'true' And '$([System.Runtime.InteropServices.RuntimeInformation]::OSArchitecture)' == 'Arm64'">portable</DebugType>
   </PropertyGroup>
   <PropertyGroup Condition=" '$(Configuration)|$(Platform)' == 'Debug|AnyCPU' ">
     <DocumentationFile></DocumentationFile>
diff --git a/Engine/Source/Programs/Shared/EpicGames.MsBuild/EpicGames.MsBuild.csproj b/Engine/Source/Programs/Shared/EpicGames.MsBuild/EpicGames.MsBuild.csproj
--- a/Engine/Source/Programs/Shared/EpicGames.MsBuild/EpicGames.MsBuild.csproj
+++ b/Engine/Source/Programs/Shared/EpicGames.MsBuild/EpicGames.MsBuild.csproj
@@ -9,8 +9,6 @@
     <GenerateAssemblyInfo>false</GenerateAssemblyInfo>
     <GenerateTargetFrameworkAttribute>false</GenerateTargetFrameworkAttribute>
     <Configurations>Debug;Release;Development;Analyze</Configurations>
-    <DebugType>pdbonly</DebugType>
-    <DebugType Condition="'$([System.Runtime.InteropServices.RuntimeInformation]::IsOSPlatform($([System.Runtime.InteropServices.OSPlatform]::Windows)))' == 'true' And '$([System.Runtime.InteropServices.RuntimeInformation]::OSArchitecture)' == 'Arm64'">portable</DebugType>
     <ProduceReferenceAssemblyInOutDir>true</ProduceReferenceAssemblyInOutDir>
   </PropertyGroup>
   <PropertyGroup Condition=" '$(Configuration)|$(Platform)' == 'Debug|AnyCPU' ">
diff --git a/Engine/Source/Programs/Shared/EpicGames.OIDC/EpicGames.OIDC.csproj b/Engine/Source/Programs/Shared/EpicGames.OIDC/EpicGames.OIDC.csproj
--- a/Engine/Source/Programs/Shared/EpicGames.OIDC/EpicGames.OIDC.csproj
+++ b/Engine/Source/Programs/Shared/EpicGames.OIDC/EpicGames.OIDC.csproj
@@ -5,8 +5,6 @@
     <OutputType>Library</OutputType>
     <RootNamespace>EpicGames.OIDC</RootNamespace>
     <TargetFramework>net8.0</TargetFramework>
-    <DebugType>pdbonly</DebugType>
-    <DebugType Condition="'$([System.Runtime.InteropServices.RuntimeInformation]::IsOSPlatform($([System.Runtime.InteropServices.OSPlatform]::Windows)))' == 'true' And '$([System.Runtime.InteropServices.RuntimeInformation]::OSArchitecture)' == 'Arm64'">portable</DebugType>
     <Nullable>enable</Nullable>
     <GenerateAssemblyInfo>false</GenerateAssemblyInfo>
     <GenerateTargetFrameworkAttribute>false</GenerateTargetFrameworkAttribute>
diff --git a/Engine/Source/Programs/Shared/EpicGames.Oodle/EpicGames.Oodle.csproj b/Engine/Source/Programs/Shared/EpicGames.Oodle/EpicGames.Oodle.csproj
--- a/Engine/Source/Programs/Shared/EpicGames.Oodle/EpicGames.Oodle.csproj
+++ b/Engine/Source/Programs/Shared/EpicGames.Oodle/EpicGames.Oodle.csproj
@@ -11,8 +11,6 @@
     <GenerateTargetFrameworkAttribute>false</GenerateTargetFrameworkAttribute>
     <Configurations>Debug;Release;Development;Analyze</Configurations>
     <Nullable>enable</Nullable>
-    <DebugType>pdbonly</DebugType>
-    <DebugType Condition="'$([System.Runtime.InteropServices.RuntimeInformation]::IsOSPlatform($([System.Runtime.InteropServices.OSPlatform]::Windows)))' == 'true' And '$([System.Runtime.InteropServices.RuntimeInformation]::OSArchitecture)' == 'Arm64'">portable</DebugType>
     <GenerateDocumentationFile>True</GenerateDocumentationFile>
     <AllowUnsafeBlocks>True</AllowUnsafeBlocks>
     <RunAnalyzersDuringBuild>false</RunAnalyzersDuringBuild>
diff --git a/Engine/Source/Programs/Shared/EpicGames.Perforce/EpicGames.Perforce.csproj b/Engine/Source/Programs/Shared/EpicGames.Perforce/EpicGames.Perforce.csproj
--- a/Engine/Source/Programs/Shared/EpicGames.Perforce/EpicGames.Perforce.csproj
+++ b/Engine/Source/Programs/Shared/EpicGames.Perforce/EpicGames.Perforce.csproj
@@ -10,8 +10,6 @@
     <GenerateTargetFrameworkAttribute>false</GenerateTargetFrameworkAttribute>
     <Configurations>Analyze;Debug;Release;Development</Configurations>
     <nullable>enable</nullable>
-    <DebugType>pdbonly</DebugType>
-    <DebugType Condition="'$([System.Runtime.InteropServices.RuntimeInformation]::IsOSPlatform($([System.Runtime.InteropServices.OSPlatform]::Windows)))' == 'true' And '$([System.Runtime.InteropServices.RuntimeInformation]::OSArchitecture)' == 'Arm64'">portable</DebugType>
   </PropertyGroup>
   <PropertyGroup Condition=" '$(Configuration)|$(Platform)' == 'Debug|AnyCPU' ">
     <DocumentationFile>EpicGames.Perforce.xml</DocumentationFile>
@@ -48,8 +46,6 @@
   <PropertyGroup>
     <WindowsNativeConfig>Release</WindowsNativeConfig>
     <WindowsNativeConfig Condition=" '$(Configuration)|$(Platform)' == 'Debug|AnyCPU' ">Debug</WindowsNativeConfig>
-    <DebugType>full</DebugType>
-    <DebugType Condition="'$([System.Runtime.InteropServices.RuntimeInformation]::IsOSPlatform($([System.Runtime.InteropServices.OSPlatform]::Windows)))' == 'true' And '$([System.Runtime.InteropServices.RuntimeInformation]::OSArchitecture)' == 'Arm64'">portable</DebugType>
     <EnableNETAnalyzers>True</EnableNETAnalyzers>
   </PropertyGroup>
 
diff --git a/Engine/Source/Programs/Shared/EpicGames.ProjectStore/EpicGames.ProjectStore.csproj b/Engine/Source/Programs/Shared/EpicGames.ProjectStore/EpicGames.ProjectStore.csproj
--- a/Engine/Source/Programs/Shared/EpicGames.ProjectStore/EpicGames.ProjectStore.csproj
+++ b/Engine/Source/Programs/Shared/EpicGames.ProjectStore/EpicGames.ProjectStore.csproj
@@ -11,8 +11,6 @@
     <GenerateTargetFrameworkAttribute>false</GenerateTargetFrameworkAttribute>
     <Configurations>Debug;Release;Development;Analyze</Configurations>
     <Nullable>enable</Nullable>
-    <DebugType>pdbonly</DebugType>
-    <DebugType Condition="'$([System.Runtime.InteropServices.RuntimeInformation]::IsOSPlatform($([System.Runtime.InteropServices.OSPlatform]::Windows)))' == 'true' And '$([System.Runtime.InteropServices.RuntimeInformation]::OSArchitecture)' == 'Arm64'">portable</DebugType>
     <RunAnalyzersDuringBuild>false</RunAnalyzersDuringBuild>
   </PropertyGroup>
   <PropertyGroup Condition=" '$(Configuration)|$(Platform)' == 'Debug|AnyCPU' ">
diff --git a/Engine/Source/Programs/Shared/EpicGames.ScriptBuild/EpicGames.ScriptBuild.csproj b/Engine/Source/Programs/Shared/EpicGames.ScriptBuild/EpicGames.ScriptBuild.csproj
--- a/Engine/Source/Programs/Shared/EpicGames.ScriptBuild/EpicGames.ScriptBuild.csproj
+++ b/Engine/Source/Programs/Shared/EpicGames.ScriptBuild/EpicGames.ScriptBuild.csproj
@@ -11,9 +11,6 @@
     <GenerateTargetFrameworkAttribute>false</GenerateTargetFrameworkAttribute>
     <Configurations>Debug;Release;Development;Analyze</Configurations>
     <Nullable>enable</Nullable>
-    <DebugType>full</DebugType>
-    <DebugType>pdbonly</DebugType>
-    <DebugType Condition="'$([System.Runtime.InteropServices.RuntimeInformation]::IsOSPlatform($([System.Runtime.InteropServices.OSPlatform]::Windows)))' == 'true' And '$([System.Runtime.InteropServices.RuntimeInformation]::OSArchitecture)' == 'Arm64'">portable</DebugType>
   </PropertyGroup>
   <PropertyGroup Condition=" '$(Configuration)|$(Platform)' == 'Debug|AnyCPU' ">
     <DocumentationFile></DocumentationFile>
diff --git a/Engine/Source/Programs/Shared/EpicGames.Serialization/EpicGames.Serialization.csproj b/Engine/Source/Programs/Shared/EpicGames.Serialization/EpicGames.Serialization.csproj
--- a/Engine/Source/Programs/Shared/EpicGames.Serialization/EpicGames.Serialization.csproj
+++ b/Engine/Source/Programs/Shared/EpicGames.Serialization/EpicGames.Serialization.csproj
@@ -11,8 +11,6 @@
     <GenerateTargetFrameworkAttribute>false</GenerateTargetFrameworkAttribute>
     <Configurations>Debug;Release;Development;Analyze</Configurations>
     <Nullable>enable</Nullable>
-    <DebugType>pdbonly</DebugType>
-    <DebugType Condition="'$([System.Runtime.InteropServices.RuntimeInformation]::IsOSPlatform($([System.Runtime.InteropServices.OSPlatform]::Windows)))' == 'true' And '$([System.Runtime.InteropServices.RuntimeInformation]::OSArchitecture)' == 'Arm64'">portable</DebugType>
     <RunAnalyzersDuringBuild>false</RunAnalyzersDuringBuild>
   </PropertyGroup>
   <PropertyGroup Condition=" '$(Configuration)|$(Platform)' == 'Debug|AnyCPU' ">
diff --git a/Engine/Source/Programs/Shared/EpicGames.UBA/EpicGames.UBA.csproj b/Engine/Source/Programs/Shared/EpicGames.UBA/EpicGames.UBA.csproj
--- a/Engine/Source/Programs/Shared/EpicGames.UBA/EpicGames.UBA.csproj
+++ b/Engine/Source/Programs/Shared/EpicGames.UBA/EpicGames.UBA.csproj
@@ -12,8 +12,6 @@
     <Nullable>enable</Nullable>
     <AllowUnsafeBlocks>true</AllowUnsafeBlocks>
     <GenerateDocumentationFile>True</GenerateDocumentationFile>
-    <DebugType>pdbonly</DebugType>
-    <DebugType Condition="'$([System.Runtime.InteropServices.RuntimeInformation]::IsOSPlatform($([System.Runtime.InteropServices.OSPlatform]::Windows)))' == 'true' And '$([System.Runtime.InteropServices.RuntimeInformation]::OSArchitecture)' == 'Arm64'">portable</DebugType>
     <EnforceCodeStyleInBuild>False</EnforceCodeStyleInBuild>
     <ProduceReferenceAssemblyInOutDir>true</ProduceReferenceAssemblyInOutDir>
     <RunAnalyzersDuringBuild>false</RunAnalyzersDuringBuild>
diff --git a/Engine/Source/Programs/Shared/EpicGames.UHT/EpicGames.UHT.csproj b/Engine/Source/Programs/Shared/EpicGames.UHT/EpicGames.UHT.csproj
--- a/Engine/Source/Programs/Shared/EpicGames.UHT/EpicGames.UHT.csproj
+++ b/Engine/Source/Programs/Shared/EpicGames.UHT/EpicGames.UHT.csproj
@@ -12,8 +12,6 @@
     <Nullable>enable</Nullable>
     <AllowUnsafeBlocks>true</AllowUnsafeBlocks>
     <GenerateDocumentationFile>True</GenerateDocumentationFile>
-    <DebugType>pdbonly</DebugType>
-    <DebugType Condition="'$([System.Runtime.InteropServices.RuntimeInformation]::IsOSPlatform($([System.Runtime.InteropServices.OSPlatform]::Windows)))' == 'true' And '$([System.Runtime.InteropServices.RuntimeInformation]::OSArchitecture)' == 'Arm64'">portable</DebugType>
     <EnforceCodeStyleInBuild>False</EnforceCodeStyleInBuild>
     <ProduceReferenceAssemblyInOutDir>true</ProduceReferenceAssemblyInOutDir>
     <RunAnalyzersDuringBuild>false</RunAnalyzersDuringBuild>
diff --git a/Engine/Source/Programs/UnrealBuildTool/UnrealBuildTool.csproj b/Engine/Source/Programs/UnrealBuildTool/UnrealBuildTool.csproj
--- a/Engine/Source/Programs/UnrealBuildTool/UnrealBuildTool.csproj
+++ b/Engine/Source/Programs/UnrealBuildTool/UnrealBuildTool.csproj
@@ -14,8 +14,6 @@
     <RootNamespace>UnrealBuildTool</RootNamespace>
     <SatelliteResourceLanguages>en</SatelliteResourceLanguages> <!-- remove non english resource languages -->
     <AppendRuntimeIdentifierToOutputPath>false</AppendRuntimeIdentifierToOutputPath>
-    <DebugType>pdbonly</DebugType>
-    <DebugType Condition="'$(IsWindows)' == 'true' And '$(IsArm64)' == 'true'">portable</DebugType>
     <Nullable>enable</Nullable>
     <RunAnalyzersDuringBuild>false</RunAnalyzersDuringBuild>
     <GenerateDocumentationFile>True</GenerateDocumentationFile>
@@ -28,8 +26,6 @@
   </PropertyGroup>
   <PropertyGroup Condition=" '$(Configuration)|$(Platform)' == 'Debug|AnyCPU' ">
     <DebugSymbols>true</DebugSymbols>
-    <DebugType>full</DebugType>
-    <DebugType Condition="'$(IsWindows)' == 'true' And '$(IsArm64)' == 'true'">portable</DebugType>
   </PropertyGroup>
   <PropertyGroup Condition=" '$(Configuration)|$(Platform)' == 'Analyze|AnyCPU' ">
     <RunAnalyzersDuringBuild>True</RunAnalyzersDuringBuild>
