Fix compile warnings in Horde

diff --git a/Engine/Source/Programs/Shared/EpicGames.Perforce/PerforceChildProcess.cs b/Engine/Source/Programs/Shared/EpicGames.Perforce/PerforceChildProcess.cs
--- a/Engine/Source/Programs/Shared/EpicGames.Perforce/PerforceChildProcess.cs
+++ b/Engine/Source/Programs/Shared/EpicGames.Perforce/PerforceChildProcess.cs
@@ -15,7 +15,7 @@ namespace EpicGames.Perforce
 	/// <summary>
 	/// Wraps a call to a p4.exe child process, and allows reading data from it
 	/// </summary>
-	public class PerforceChildProcess : IPerforceOutput, IDisposable
+	public sealed class PerforceChildProcess : IPerforceOutput, IDisposable
 	{
 		/// <summary>
 		/// The process group
