Removed "virtual" from generated functions if the class is marked as "final"

diff --git a/Engine/Source/Programs/Shared/EpicGames.UHT/Exporters/CodeGen/UhtHeaderCodeGenerator.cs b/Engine/Source/Programs/Shared/EpicGames.UHT/Exporters/CodeGen/UhtHeaderCodeGenerator.cs
--- a/Engine/Source/Programs/Shared/EpicGames.UHT/Exporters/CodeGen/UhtHeaderCodeGenerator.cs
+++ b/Engine/Source/Programs/Shared/EpicGames.UHT/Exporters/CodeGen/UhtHeaderCodeGenerator.cs
@@ -173,23 +173,38 @@ namespace EpicGames.UHT.Exporters.CodeGen
 				{
 					builder.Append("static ");
 				}
-				else if (isK2Override)
-				{
-					builder.Append("virtual ");
-				}
-				// if the owning class is an interface class
-				else if (isInterface)
-				{
-					builder.Append("virtual ");
-				}
-				// this is not an event, the function is not a static function and the function is not marked final
-				else if (textType != UhtPropertyTextType.EventFunctionArgOrRetVal && !function.FunctionFlags.HasAnyFlags(EFunctionFlags.Static) && !exportFlags.HasAnyFlags(UhtFunctionExportFlags.Final))
-				{
-					builder.Append("virtual ");
-				}
-				else if (exportFlags.HasAnyFlags(UhtFunctionExportFlags.Inline))
+				else
 				{
-					builder.Append("inline ");
+					bool isVirtual = false;
+					if (outerClass != null && outerClass.ClassExportFlags.HasAnyFlags(UhtClassExportFlags.Final))
+					{
+						//
+					}
+					else if (isK2Override)
+					{
+						isVirtual = true;
+					}
+					// if the owning class is an interface class
+					else if (isInterface)
+					{
+						isVirtual = true;
+					}
+					// this is not an event, the function is not a static function and the function is not marked final
+					else if (textType != UhtPropertyTextType.EventFunctionArgOrRetVal &&
+						!function.FunctionFlags.HasAnyFlags(EFunctionFlags.Static) &&
+						!exportFlags.HasAnyFlags(UhtFunctionExportFlags.Final))
+					{
+						isVirtual = true;
+					}
+
+					if (isVirtual)
+					{
+						builder.Append("virtual ");
+					}
+					else if (exportFlags.HasAnyFlags(UhtFunctionExportFlags.Inline))
+					{
+						builder.Append("inline ");
+					}
 				}
 			}
 
diff --git a/Engine/Source/Programs/Shared/EpicGames.UHT/Exporters/CodeGen/UhtHeaderCodeGeneratorHFile.cs b/Engine/Source/Programs/Shared/EpicGames.UHT/Exporters/CodeGen/UhtHeaderCodeGeneratorHFile.cs
--- a/Engine/Source/Programs/Shared/EpicGames.UHT/Exporters/CodeGen/UhtHeaderCodeGeneratorHFile.cs
+++ b/Engine/Source/Programs/Shared/EpicGames.UHT/Exporters/CodeGen/UhtHeaderCodeGeneratorHFile.cs
@@ -1298,7 +1298,9 @@ namespace EpicGames.UHT.Exporters.CodeGen
 					builder.Append(Module.Api);
 				}
 
-				if (!function.FunctionFlags.HasAnyFlags(EFunctionFlags.Static) && !function.FunctionExportFlags.HasAnyFlags(UhtFunctionExportFlags.Final))
+				if (!function.FunctionFlags.HasAnyFlags(EFunctionFlags.Static) && 
+					!function.FunctionExportFlags.HasAnyFlags(UhtFunctionExportFlags.Final) &&
+					!classObj.ClassExportFlags.HasAnyFlags(UhtClassExportFlags.Final))
 				{
 					builder.Append("virtual");
 				}
diff --git a/Engine/Source/Programs/Shared/EpicGames.UHT/Parsers/UhtClassParser.cs b/Engine/Source/Programs/Shared/EpicGames.UHT/Parsers/UhtClassParser.cs
--- a/Engine/Source/Programs/Shared/EpicGames.UHT/Parsers/UhtClassParser.cs
+++ b/Engine/Source/Programs/Shared/EpicGames.UHT/Parsers/UhtClassParser.cs
@@ -95,7 +95,10 @@ namespace EpicGames.UHT.Parsers
 					}
 
 					// Skip optional final keyword
-					topScope.TokenReader.Optional("final");
+					if (topScope.TokenReader.TryOptional("final"))
+					{
+						classObj.ClassExportFlags |= UhtClassExportFlags.Final;
+					}
 
 					// Parse the inheritance
 					UhtParserHelpers.ParseInheritance(parentScope.HeaderParser, topScope.Session.Config!, out UhtInheritance inheritance);
diff --git a/Engine/Source/Programs/Shared/EpicGames.UHT/Types/UhtClass.cs b/Engine/Source/Programs/Shared/EpicGames.UHT/Types/UhtClass.cs
--- a/Engine/Source/Programs/Shared/EpicGames.UHT/Types/UhtClass.cs
+++ b/Engine/Source/Programs/Shared/EpicGames.UHT/Types/UhtClass.cs
@@ -96,6 +96,11 @@ namespace EpicGames.UHT.Types
 		/// legacy generated body.  If this flag is set, generate the legacy instead of the GENERATED_BODY macros.
 		/// </summary>
 		UsesGeneratedBodyLegacy = 1 << 12,
+
+		/// <summary>
+		/// Class is marked as final
+		/// </summary>
+		Final = 1 << 13,
 	}
 
 	/// <summary>
