Fixed code generation where USTRUCT might generate wrong calls to is_polymorphic.

#rnx

diff --git a/Engine/Source/Programs/Shared/EpicGames.UHT/Exporters/CodeGen/UhtHeaderCodeGeneratorCppFile.cs b/Engine/Source/Programs/Shared/EpicGames.UHT/Exporters/CodeGen/UhtHeaderCodeGeneratorCppFile.cs
--- a/Engine/Source/Programs/Shared/EpicGames.UHT/Exporters/CodeGen/UhtHeaderCodeGeneratorCppFile.cs
+++ b/Engine/Source/Programs/Shared/EpicGames.UHT/Exporters/CodeGen/UhtHeaderCodeGeneratorCppFile.cs
@@ -574,8 +574,10 @@ namespace EpicGames.UHT.Exporters.CodeGen
 				if (scriptStruct.SuperScriptStruct != null)
 				{
 					builder.Append("static_assert(std::is_polymorphic<")
+						.Append(scriptStruct.Namespace.FullSourceName)
 						.Append(scriptStruct.SourceName)
 						.Append(">() == std::is_polymorphic<")
+						.Append(scriptStruct.SuperScriptStruct.Namespace.FullSourceName)
 						.Append(scriptStruct.SuperScriptStruct.SourceName)
 						.Append(">(), \"USTRUCT ")
 						.Append(scriptStruct.SourceName)
