EpicGames.UHT: Fix relative path length check when determining best include path to use when including a module header

diff --git a/Engine/Source/Programs/Shared/EpicGames.UHT/Utils/UhtSession.cs b/Engine/Source/Programs/Shared/EpicGames.UHT/Utils/UhtSession.cs
--- a/Engine/Source/Programs/Shared/EpicGames.UHT/Utils/UhtSession.cs
+++ b/Engine/Source/Programs/Shared/EpicGames.UHT/Utils/UhtSession.cs
@@ -332,7 +332,7 @@ namespace EpicGames.UHT.Utils
 			{
 				if (filePath.StartsWith(includePath, StringComparison.Ordinal))
 				{
-					if (relativePath == null || (includePath.Length - filePath.Length - 1) < relativePath.Length)
+					if (relativePath == null || (filePath.Length - includePath.Length - 1) < relativePath.Length)
 					{
 						relativePath = filePath.Substring(includePath.Length + 1);
 					}
