Horde: Update UbaCrypto to match C++ impl

diff --git a/Engine/Source/Programs/Shared/EpicGames.Horde.Tests/Compute/UbaCryptoTests.cs b/Engine/Source/Programs/Shared/EpicGames.Horde.Tests/Compute/UbaCryptoTests.cs
--- a/Engine/Source/Programs/Shared/EpicGames.Horde.Tests/Compute/UbaCryptoTests.cs
+++ b/Engine/Source/Programs/Shared/EpicGames.Horde.Tests/Compute/UbaCryptoTests.cs
@@ -10,16 +10,17 @@ namespace EpicGames.Horde.Tests.Compute;
 [TestClass]
 public class UbaCryptoTests
 {
-	private readonly byte[] _key = [0x01, 0x23, 0x45, 0x67, 0x89, 0xAB, 0xCD, 0xEF, 0xFE, 0xDC, 0xBA, 0x98, 0x76, 0x54, 0x32, 0x10];
+	private readonly byte[] _key = [0x3f, 0x58, 0xaa, 0x57, 0x46, 0x6d, 0xb9, 0x99, 0x92, 0x13, 0x45, 0x67, 0x89, 0x12, 0x34, 0x45];
 	
 	private const string Decrypted1 = "This is 15 char";
-	private const string Encrypted1 = "54 68 69 73 20 69 73 20 31 35 20 63 68 61 72";
+	private const string Encrypted1 = "5B 67 66 7C 2F 66 7C 2F 3E 3A 2F 6C 67 6E 7D";
 	private const string Decrypted2 = "This is 16 chars";
-	private const string Encrypted2 = "B6 72 E4 F2 3B 6F 86 05 2A BC AA B3 69 30 39 42";
+	private const string Encrypted2 = "4C 94 C9 27 FC 58 91 D3 FA 05 09 38 12 00 5E 84";
 	private const string Decrypted3 = "This is 17 chars!";
-	private const string Encrypted3 = "81 98 1C 37 C8 48 09 25 6A 0D C2 24 E7 82 71 9D 21";
+	private const string Encrypted3 = "08 33 98 87 51 CD F0 A7 1D 66 58 B2 EF 38 3F 2F 75";
 	
 	[TestMethod]
+	[DataRow("", "")]
 	[DataRow(Decrypted1, Encrypted1)]
 	[DataRow(Decrypted2, Encrypted2)]
 	[DataRow(Decrypted3, Encrypted3)]
@@ -31,6 +32,7 @@ public class UbaCryptoTests
 	}
 	
 	[TestMethod]
+	[DataRow("", "")]
 	[DataRow(Encrypted1, Decrypted1)]
 	[DataRow(Encrypted2, Decrypted2)]
 	[DataRow(Encrypted3, Decrypted3)]
diff --git a/Engine/Source/Programs/Shared/EpicGames.Horde/Compute/UbaCrypto.cs b/Engine/Source/Programs/Shared/EpicGames.Horde/Compute/UbaCrypto.cs
--- a/Engine/Source/Programs/Shared/EpicGames.Horde/Compute/UbaCrypto.cs
+++ b/Engine/Source/Programs/Shared/EpicGames.Horde/Compute/UbaCrypto.cs
@@ -1,6 +1,7 @@
 ﻿// Copyright Epic Games, Inc. All Rights Reserved.
 
 using System;
+using System.Diagnostics.CodeAnalysis;
 using System.Security.Cryptography;
 using System.Threading;
 
@@ -13,6 +14,8 @@ namespace EpicGames.Horde.Compute;
 public class UbaCrypto : IDisposable
 {
 	private const int AesBlockSize = 16;
+	private static readonly byte[] s_fixedIvLe = [0x07, 0x06, 0x05, 0x04, 0x03, 0x02, 0x01, 0x00, 0x0F, 0x0E, 0x0D, 0x0C, 0x0B, 0x0A, 0x09, 0x08];
+	
 	private readonly ThreadLocal<Aes> _aes;
 	private bool _disposed;
 	
@@ -46,34 +49,7 @@ public class UbaCrypto : IDisposable
 	/// A byte array containing the encrypted data. The array length matches the input length.
 	/// Only complete 16-byte blocks are encrypted; remaining bytes are left unmodified.
 	/// </returns>
-	public byte[] Encrypt(byte[] data)
-	{
-		if (data == null || data.Length == 0)
-		{
-			throw new ArgumentException("Data cannot be null or empty");
-		}
-		
-		byte[] result = new byte[data.Length];
-		Array.Copy(data, result, data.Length);
-		
-		int alignedSize = (data.Length / AesBlockSize) * AesBlockSize;
-		if (alignedSize > 0)
-		{
-			Aes aes = _aes.Value ?? throw new Exception("Failed to acquire thread-local AES instance");
-
-			// Use zero IV to match BCrypt behavior with nullptr IV
-			aes.IV = new byte[AesBlockSize]; // All zeros
-			
-			using ICryptoTransform encryptor = aes.CreateEncryptor();
-			byte[] alignedData = new byte[alignedSize];
-			Array.Copy(result, alignedData, alignedSize);
-			
-			byte[] encrypted = encryptor.TransformFinalBlock(alignedData, 0, alignedSize);
-			Array.Copy(encrypted, result, alignedSize);
-		}
-		
-		return result;
-	}
+	public byte[] Encrypt(byte[] data) => Transform(data, encrypt: true);
 	
 	/// <summary>
 	/// Decrypts the specified data.
@@ -83,33 +59,66 @@ public class UbaCrypto : IDisposable
 	/// A byte array containing the decrypted data. The array length matches the input length.
 	/// Only complete 16-byte blocks are decrypted; remaining bytes are left unmodified.
 	/// </returns>
-	public byte[] Decrypt(byte[] data)
+	public byte[] Decrypt(byte[] data) => Transform(data, encrypt: false);
+	
+	/// <summary>
+	/// Decrypt/encrypt method ported from C++ UBA
+	/// </summary>
+	/// <param name="src">Input buffer to operate on</param>
+	/// <param name="encrypt">Whether to encrypt (if false, decrypt)</param>
+	/// <returns>A new byte array with encrypted/decrypted data</returns>
+	[SuppressMessage("Security", "CA5401:Do not use CreateEncryptor with non-default IV")]
+	private byte[] Transform(byte[] src, bool encrypt)
 	{
-		if (data == null || data.Length == 0)
+		if (src.Length == 0)
 		{
-			throw new ArgumentException("Data cannot be null or empty");
+			return [];
 		}
 		
-		byte[] result = new byte[data.Length];
-		Array.Copy(data, result, data.Length);
+		// Avoid modifying caller's buffer
+		byte[] buf = (byte[])src.Clone();
+		int size = buf.Length;
 		
-		int alignedSize = (data.Length / AesBlockSize) * AesBlockSize;
-		if (alignedSize > 0)
+		// Messages smaller than one block -> XOR with length and return
+		if (size < AesBlockSize)
 		{
-			Aes aes = _aes.Value ?? throw new Exception("Failed to acquire thread-local AES instance");
-			
-			// Use zero IV to match BCrypt behavior with nullptr IV
-			aes.IV = new byte[AesBlockSize]; // All zeros
-			
-			using ICryptoTransform decryptor = aes.CreateDecryptor();
-			byte[] alignedData = new byte[alignedSize];
-			Array.Copy(result, alignedData, alignedSize);
-			
-			byte[] decrypted = decryptor.TransformFinalBlock(alignedData, 0, alignedSize);
-			Array.Copy(decrypted, result, alignedSize);
+			byte len = (byte)size;
+			for (int i = 0; i < size; ++i)
+			{
+				buf[i] ^= len;
+			}
+			return buf;
+		}
+		
+		int aligned = size / AesBlockSize * AesBlockSize;
+		int overflow = size - aligned;
+		
+		// If encrypting and we have overflowing bytes, XOR them
+		if (overflow > 0 && encrypt)
+		{
+			for (int i = 0; i < overflow; ++i)
+			{
+				buf[aligned + i] ^= buf[i];
+			}
+		}
+		
+		// AES-CBC on the aligned prefix
+		Aes aes = _aes.Value!;
+		byte[] clonedId = (byte[])s_fixedIvLe.Clone();
+		using ICryptoTransform tr = encrypt ? aes.CreateEncryptor(aes.Key, clonedId) : aes.CreateDecryptor(aes.Key, clonedId);
+		byte[] prefix = tr.TransformFinalBlock(buf, 0, aligned);
+		Buffer.BlockCopy(prefix, 0, buf, 0, aligned);
+		
+		// If decrypting and we have overflowing bytes, XOR them
+		if (overflow > 0 && !encrypt)
+		{
+			for (int i = 0; i < overflow; ++i)
+			{
+				buf[aligned + i] ^= buf[i];
+			}
 		}
 		
-		return result;
+		return buf;
 	}
 	
 	/// <inheritdoc />
