Fix various warnings in new telemetry code

#rnx

diff --git a/Engine/Source/Programs/Shared/EpicGames.Core/Telemetry/SimpleJsonTelemetry.cs b/Engine/Source/Programs/Shared/EpicGames.Core/Telemetry/SimpleJsonTelemetry.cs
--- a/Engine/Source/Programs/Shared/EpicGames.Core/Telemetry/SimpleJsonTelemetry.cs
+++ b/Engine/Source/Programs/Shared/EpicGames.Core/Telemetry/SimpleJsonTelemetry.cs
@@ -1,7 +1,9 @@
 // Copyright Epic Games, Inc. All Rights Reserved.
 
 using System;
+using System.ComponentModel;
 using System.Net.Http;
+using System.Reflection.Metadata;
 using System.Text.Json;
 using System.Threading.Tasks;
 using Microsoft.Extensions.Logging;
@@ -11,9 +13,10 @@ namespace EpicGames.Core.Telemetry
 	/// <summary>
 	/// Simple telemetry wrapper for accessing JSON data
 	/// </summary>
-	public class SimpleJsonTelemetry
+	public class SimpleJsonTelemetry : IDisposable
 	{
 		private readonly HttpClient _httpClient;
+		private bool _disposed = false;
 
 		/// <summary>
 		/// Constructor
@@ -25,16 +28,20 @@ namespace EpicGames.Core.Telemetry
 				BaseAddress = new Uri(server),
 			};
 		}
+		~SimpleJsonTelemetry()
+		{
+			Dispose(disposing: false);
+		}
 
 		/// <summary>
 		/// Fetch the data at a given endpoint and return it as string
 		/// </summary>
-		public async Task<T?> GetAsync<T>(string endpoint)
+		public async Task<T?> GetAsync<T>(string endPoint)
 			where T: new()
 		{
 			try
 			{
-				using HttpResponseMessage response = await _httpClient.GetAsync(endpoint);
+				using HttpResponseMessage response = await _httpClient.GetAsync(new Uri(endPoint, UriKind.Relative));
 
 				response.EnsureSuccessStatusCode();
 
@@ -44,9 +51,26 @@ namespace EpicGames.Core.Telemetry
 			}
 			catch (HttpRequestException e)
 			{
-				Log.Logger.LogWarning("GetAsync: failed to retrieve {endpoint}: {1}", endpoint, e.Message);
+				Log.Logger.LogWarning("GetAsync: failed to retrieve {EndPoint}: {Message}", endPoint, e.Message);
 				return new T();
 			}
 		}
+		public void Dispose()
+		{
+			Dispose(disposing: true);
+			GC.SuppressFinalize(this);
+		}
+		protected virtual void Dispose(bool disposing)
+		{
+			if (!_disposed)
+			{
+				if (disposing)
+				{
+					_httpClient.Dispose();
+				}
+
+				_disposed = true;
+			}
+		}
 	};
 }
