Fix more warnings

#rnx

diff --git a/Engine/Source/Programs/Shared/EpicGames.Core/Telemetry/SimpleJsonTelemetry.cs b/Engine/Source/Programs/Shared/EpicGames.Core/Telemetry/SimpleJsonTelemetry.cs
--- a/Engine/Source/Programs/Shared/EpicGames.Core/Telemetry/SimpleJsonTelemetry.cs
+++ b/Engine/Source/Programs/Shared/EpicGames.Core/Telemetry/SimpleJsonTelemetry.cs
@@ -55,11 +55,20 @@ namespace EpicGames.Core.Telemetry
 				return new T();
 			}
 		}
+
+		/// <summary>
+		/// Implement IDisposable.
+		/// </summary>
 		public void Dispose()
 		{
 			Dispose(disposing: true);
 			GC.SuppressFinalize(this);
 		}
+
+		/// <summary>
+		/// Implement IDisposable.
+		/// </summary>
+		/// <param name="disposing">If call comes from Dispose, the value is 'true'. If call comes from finalizer, the value is 'false'.</param>
 		protected virtual void Dispose(bool disposing)
 		{
 			if (!_disposed)
