Fix more warnings

#rnx

diff --git a/Engine/Source/Programs/Shared/EpicGames.Core/Telemetry/SimpleJsonTelemetry.cs b/Engine/Source/Programs/Shared/EpicGames.Core/Telemetry/SimpleJsonTelemetry.cs
--- a/Engine/Source/Programs/Shared/EpicGames.Core/Telemetry/SimpleJsonTelemetry.cs
+++ b/Engine/Source/Programs/Shared/EpicGames.Core/Telemetry/SimpleJsonTelemetry.cs
@@ -1,9 +1,7 @@
 // Copyright Epic Games, Inc. All Rights Reserved.
 
 using System;
-using System.ComponentModel;
 using System.Net.Http;
-using System.Reflection.Metadata;
 using System.Text.Json;
 using System.Threading.Tasks;
 using Microsoft.Extensions.Logging;
@@ -28,6 +26,10 @@ namespace EpicGames.Core.Telemetry
 				BaseAddress = new Uri(server),
 			};
 		}
+
+		/// <summary>
+		/// Destructor
+		/// </summary>
 		~SimpleJsonTelemetry()
 		{
 			Dispose(disposing: false);
