Horde: Add support in compute client for requesting a UBA cache server

diff --git a/Engine/Source/Programs/Horde/HordeCmd/Commands/Compute/ComputeRun.cs b/Engine/Source/Programs/Horde/HordeCmd/Commands/Compute/ComputeRun.cs
--- a/Engine/Source/Programs/Horde/HordeCmd/Commands/Compute/ComputeRun.cs
+++ b/Engine/Source/Programs/Horde/HordeCmd/Commands/Compute/ComputeRun.cs
@@ -121,7 +121,7 @@ namespace Horde.Commands.Compute
 				case "NONE": cmr.Encryption = EpicGames.Horde.Compute.Encryption.None; break;
 			}
 
-			await using IComputeLease? lease = await client.TryAssignWorkerAsync(new ClusterId(ClusterId), requirements, null, cmr, logger, CancellationToken.None);
+			await using IComputeLease? lease = await client.TryAssignWorkerAsync(new ClusterId(ClusterId), requirements, null, cmr, null, logger, CancellationToken.None);
 			if (lease == null)
 			{
 				throw new Exception("Unable to create lease");
diff --git a/Engine/Source/Programs/Shared/EpicGames.Horde/Compute/Clients/AgentComputeClient.cs b/Engine/Source/Programs/Shared/EpicGames.Horde/Compute/Clients/AgentComputeClient.cs
--- a/Engine/Source/Programs/Shared/EpicGames.Horde/Compute/Clients/AgentComputeClient.cs
+++ b/Engine/Source/Programs/Shared/EpicGames.Horde/Compute/Clients/AgentComputeClient.cs
@@ -39,6 +39,9 @@ namespace EpicGames.Horde.Compute.Clients
 
 			/// <inheritdoc/>
 			public string Ip => "127.0.0.1";
+			
+			/// <inheritdoc/>
+			public UbaConfig? Uba => null;
 
 			/// <inheritdoc/>
 			public ConnectionMode ConnectionMode => ConnectionMode.Direct;
@@ -83,7 +86,7 @@ namespace EpicGames.Horde.Compute.Clients
 		}
 		
 		/// <inheritdoc/>
-		public async Task<IComputeLease?> TryAssignWorkerAsync(ClusterId? clusterId, Requirements? requirements, string? requestId, ConnectionMetadataRequest? connection, ILogger logger, CancellationToken cancellationToken)
+		public async Task<IComputeLease?> TryAssignWorkerAsync(ClusterId? clusterId, Requirements? requirements, string? requestId, ConnectionMetadataRequest? connection, bool? useUbaCache, ILogger logger, CancellationToken cancellationToken)
 		{
 			logger.LogInformation("** CLIENT **");
 			logger.LogInformation("Launching {Path} to handle remote", _hordeAgentAssembly);
diff --git a/Engine/Source/Programs/Shared/EpicGames.Horde/Compute/Clients/LocalComputeClient.cs b/Engine/Source/Programs/Shared/EpicGames.Horde/Compute/Clients/LocalComputeClient.cs
--- a/Engine/Source/Programs/Shared/EpicGames.Horde/Compute/Clients/LocalComputeClient.cs
+++ b/Engine/Source/Programs/Shared/EpicGames.Horde/Compute/Clients/LocalComputeClient.cs
@@ -26,6 +26,7 @@ namespace EpicGames.Horde.Compute.Clients
 			public IReadOnlyDictionary<string, int> AssignedResources => new Dictionary<string, int>();
 			public RemoteComputeSocket Socket => _socket;
 			public string Ip => "127.0.0.1";
+			public UbaConfig? Uba => null;
 			public ConnectionMode ConnectionMode => ConnectionMode.Direct;
 			public IReadOnlyDictionary<string, ConnectionMetadataPort> Ports => new Dictionary<string, ConnectionMetadataPort>();
 
@@ -94,7 +95,7 @@ namespace EpicGames.Horde.Compute.Clients
 		}
 		
 		/// <inheritdoc/>
-		public Task<IComputeLease?> TryAssignWorkerAsync(ClusterId? clusterId, Requirements? requirements, string? requestId, ConnectionMetadataRequest? connection, ILogger logger, CancellationToken cancellationToken)
+		public Task<IComputeLease?> TryAssignWorkerAsync(ClusterId? clusterId, Requirements? requirements, string? requestId, ConnectionMetadataRequest? connection, bool? useUbaCache, ILogger logger, CancellationToken cancellationToken)
 		{
 #pragma warning disable CA2000 // Dispose objects before losing scope
 			RemoteComputeSocket socket = new RemoteComputeSocket(new TcpTransport(_socket), ComputeProtocol.Latest, logger);
diff --git a/Engine/Source/Programs/Shared/EpicGames.Horde/Compute/Clients/ServerComputeClient.cs b/Engine/Source/Programs/Shared/EpicGames.Horde/Compute/Clients/ServerComputeClient.cs
--- a/Engine/Source/Programs/Shared/EpicGames.Horde/Compute/Clients/ServerComputeClient.cs
+++ b/Engine/Source/Programs/Shared/EpicGames.Horde/Compute/Clients/ServerComputeClient.cs
@@ -121,6 +121,7 @@ namespace EpicGames.Horde.Compute.Clients
 			IReadOnlyDictionary<string, int> AssignedResources,
 			RemoteComputeSocket Socket,
 			string Ip,
+			UbaConfig? Uba,
 			ConnectionMode ConnectionMode,
 			IReadOnlyDictionary<string, ConnectionMetadataPort> Ports);
 
@@ -131,6 +132,7 @@ namespace EpicGames.Horde.Compute.Clients
 			public IReadOnlyDictionary<string, int> AssignedResources => _source.Current.AssignedResources;
 			public RemoteComputeSocket Socket => _source.Current.Socket;
 			public string Ip => _source.Current.Ip;
+			public UbaConfig? Uba => _source.Current.Uba;
 			public ConnectionMode ConnectionMode => _source.Current.ConnectionMode;
 			public IReadOnlyDictionary<string, ConnectionMetadataPort> Ports => _source.Current.Ports;
 
@@ -241,11 +243,11 @@ namespace EpicGames.Horde.Compute.Clients
 		}
 		
 		/// <inheritdoc/>
-		public async Task<IComputeLease?> TryAssignWorkerAsync(ClusterId? clusterId, Requirements? requirements, string? requestId, ConnectionMetadataRequest? connection, ILogger logger, CancellationToken cancellationToken)
+		public async Task<IComputeLease?> TryAssignWorkerAsync(ClusterId? clusterId, Requirements? requirements, string? requestId, ConnectionMetadataRequest? connection, bool? useUbaCache, ILogger logger, CancellationToken cancellationToken)
 		{
 			try
 			{
-				IAsyncEnumerator<LeaseInfo> source = ConnectAsync(clusterId, requirements, requestId, connection, logger, cancellationToken).GetAsyncEnumerator(cancellationToken);
+				IAsyncEnumerator<LeaseInfo> source = ConnectAsync(clusterId, requirements, requestId, connection, useUbaCache, logger, cancellationToken).GetAsyncEnumerator(cancellationToken);
 				if (!await source.MoveNextAsync())
 				{
 					await source.DisposeAsync();
@@ -268,17 +270,20 @@ namespace EpicGames.Horde.Compute.Clients
 			response.EnsureSuccessStatusCode();
 		}
 
-		async IAsyncEnumerable<LeaseInfo> ConnectAsync(ClusterId? clusterId, Requirements? requirements, string? requestId, ConnectionMetadataRequest? connection, ILogger workerLogger, [EnumeratorCancellation] CancellationToken cancellationToken)
+		async IAsyncEnumerable<LeaseInfo> ConnectAsync(ClusterId? clusterId, Requirements? requirements, string? requestId, ConnectionMetadataRequest? connection, bool? useUbaCache, ILogger workerLogger, [EnumeratorCancellation] CancellationToken cancellationToken)
 		{
 			_logger.LogDebug("Requesting compute resource");
 
 			// Assign a compute worker
-			AssignComputeRequest request = new AssignComputeRequest();
-			request.Requirements = requirements;
-			request.RequestId = requestId;
-			request.Connection = connection;
-			request.Protocol = (int)ComputeProtocol.Latest;
-
+			AssignComputeRequest request = new()
+			{
+				Requirements = requirements,
+				RequestId = requestId,
+				Connection = connection,
+				UseUbaCache = useUbaCache,
+				Protocol = (int)ComputeProtocol.Latest
+			};
+			
 			if (connection is { ModePreference: ConnectionMode.Relay })
 			{
 				connection.ClientPublicIp = (await _externalIpResolver.GetExternalIpAddressAsync(cancellationToken)).ToString();
@@ -408,7 +413,7 @@ namespace EpicGames.Horde.Compute.Clients
 			response.Properties = [..response.Properties, $"{KnownPropertyNames.LeaseId}={response.LeaseId}"];
 			await using ComputeTransport transport = await CreateTransportAsync(socket, response, cancellationToken);
 			await using RemoteComputeSocket computeSocket = new(transport, (ComputeProtocol)response.Protocol, workerLogger);
-			yield return new LeaseInfo(response.ClusterId, response.Properties, response.AssignedResources, computeSocket, response.Ip, response.ConnectionMode, response.Ports);
+			yield return new LeaseInfo(response.ClusterId, response.Properties, response.AssignedResources, computeSocket, response.Ip, response.Uba, response.ConnectionMode, response.Ports);
 		}
 		
 		private async Task ConnectWithRetryAsync(Socket socket, string host, int port, TimeSpan timeout, int maxRetries, CancellationToken cancellationToken)
diff --git a/Engine/Source/Programs/Shared/EpicGames.Horde/Compute/IComputeClient.cs b/Engine/Source/Programs/Shared/EpicGames.Horde/Compute/IComputeClient.cs
--- a/Engine/Source/Programs/Shared/EpicGames.Horde/Compute/IComputeClient.cs
+++ b/Engine/Source/Programs/Shared/EpicGames.Horde/Compute/IComputeClient.cs
@@ -29,10 +29,11 @@ namespace EpicGames.Horde.Compute
 		/// <param name="clusterId">Optional cluster ID. If not set, cluster will automatically be resolved by server</param>
 		/// <param name="requirements">Requirements for the agent</param>
 		/// <param name="requestId">Optional ID identifying the request over multiple calls, such as retrying the same request</param>
+		/// <param name="useUbaCache">Whether to allocate an Unreal Build Accelerator (UBA) cache server session for this compute request</param>
 		/// <param name="connection">Optional preference of connection details</param>
 		/// <param name="logger">Logger for output from this worker</param>
 		/// <param name="cancellationToken">Cancellation token for the operation</param>
-		public Task<IComputeLease?> TryAssignWorkerAsync(ClusterId? clusterId, Requirements? requirements, string? requestId, ConnectionMetadataRequest? connection, ILogger logger, CancellationToken cancellationToken = default);
+		public Task<IComputeLease?> TryAssignWorkerAsync(ClusterId? clusterId, Requirements? requirements, string? requestId, ConnectionMetadataRequest? connection, bool? useUbaCache, ILogger logger, CancellationToken cancellationToken = default);
 
 		/// <summary>
 		/// Declare resource needs for current client
@@ -55,14 +56,14 @@ namespace EpicGames.Horde.Compute
 		public static Task<IComputeLease?> TryAssignWorkerAsync(this IComputeClient computeClient, ClusterId clusterId, Requirements? requirements, string? requestId, ConnectionMode? connectionPreference, ILogger logger, CancellationToken cancellationToken = default)
 		{
 			ConnectionMetadataRequest cmr = new() { ModePreference = connectionPreference };
-			return computeClient.TryAssignWorkerAsync(clusterId, requirements, requestId, cmr, logger, cancellationToken);
+			return computeClient.TryAssignWorkerAsync(clusterId, requirements, requestId, cmr, null, logger, cancellationToken);
 		}
 
 		/// <inheritdoc cref="IComputeClient.TryAssignWorkerAsync" />
 		[Obsolete("Prefer taking a requestId parameter")]
 		public static Task<IComputeLease?> TryAssignWorkerAsync(this IComputeClient computeClient, ClusterId clusterId, Requirements? requirements, ILogger logger, CancellationToken cancellationToken)
 		{
-			return computeClient.TryAssignWorkerAsync(clusterId, requirements, null, logger, cancellationToken);
+			return computeClient.TryAssignWorkerAsync(clusterId, requirements, null, null, logger, cancellationToken);
 		}
 
 		/// <inheritdoc cref="IComputeClient.TryAssignWorkerAsync" />
diff --git a/Engine/Source/Programs/Shared/EpicGames.Horde/Compute/IComputeLease.cs b/Engine/Source/Programs/Shared/EpicGames.Horde/Compute/IComputeLease.cs
--- a/Engine/Source/Programs/Shared/EpicGames.Horde/Compute/IComputeLease.cs
+++ b/Engine/Source/Programs/Shared/EpicGames.Horde/Compute/IComputeLease.cs
@@ -37,7 +37,10 @@ namespace EpicGames.Horde.Compute
 		/// When using relay connection mode, this may be the IP of the relay rather than the remote machine itself.
 		/// </summary>
 		public string Ip { get; }
-
+		
+		/// <inheritdoc cref="UbaConfig" />
+		public UbaConfig? Uba { get; }
+		
 		/// <summary>
 		/// How to establish a connection to the remote machine (when not using the default socket)
 		/// </summary>
diff --git a/Engine/Source/Programs/UnrealBuildTool/Executors/UnrealBuildAccelerator/UBAAgentCoordinatorHorde.cs b/Engine/Source/Programs/UnrealBuildTool/Executors/UnrealBuildAccelerator/UBAAgentCoordinatorHorde.cs
--- a/Engine/Source/Programs/UnrealBuildTool/Executors/UnrealBuildAccelerator/UBAAgentCoordinatorHorde.cs
+++ b/Engine/Source/Programs/UnrealBuildTool/Executors/UnrealBuildAccelerator/UBAAgentCoordinatorHorde.cs
@@ -282,7 +282,7 @@ namespace UnrealBuildTool
 				};
 				
 				await ResolveClusterIdAsync(hordeConfig, requirements, requestId, cmr, workerLogger, cancellationToken);
-				lease = await _client.TryAssignWorkerAsync(_clusterId, requirements, requestId, cmr, workerLogger, cancellationToken);
+				lease = await _client.TryAssignWorkerAsync(_clusterId, requirements, requestId, cmr, false, workerLogger, cancellationToken);
 				if (lease == null)
 				{
 					_logger.LogDebug("Unable to assign a remote worker");
