Add EditorDataStorage code generation extension
* Previously removed from the EditorDataStorage plugin - the code generation needs to occur at the engine level due to the interfaces defined in Engine/Source/Runtime/TypedElementFramework/Public/Elements/Common/TypedElementCommonTypes.h.  Breakage could occur if EditorDataStorage plugin is not enabled.

diff --git a/Engine/Plugins/Experimental/Dataflow/Source/DataflowEditor/Public/Dataflow/DataflowEditorPreviewSceneBase.h b/Engine/Plugins/Experimental/Dataflow/Source/DataflowEditor/Public/Dataflow/DataflowEditorPreviewSceneBase.h
--- a/Engine/Plugins/Experimental/Dataflow/Source/DataflowEditor/Public/Dataflow/DataflowEditorPreviewSceneBase.h
+++ b/Engine/Plugins/Experimental/Dataflow/Source/DataflowEditor/Public/Dataflow/DataflowEditorPreviewSceneBase.h
@@ -42,8 +42,6 @@ USTRUCT(meta = (DisplayName = "Dataflow scene object", EditorDataStorage_Dynamic
 struct FDataflowSceneObjectTag final : public FEditorDataStorageTag
 {
 	GENERATED_BODY()
-
-	struct EditorDataStorage_DynamicColumnTemplate{};
 };
 
 /**
@@ -53,8 +51,6 @@ USTRUCT(meta = (DisplayName = "Dataflow scene struct", EditorDataStorage_Dynamic
 struct FDataflowSceneStructTag final : public FEditorDataStorageTag
 {
 	GENERATED_BODY()
-
-	struct EditorDataStorage_DynamicColumnTemplate{};
 };
 
 /**
diff --git a/Engine/Plugins/Experimental/EditorDataStorage/Source/TedsCore/Private/Debug/TypedElementDatabaseDebugTypes.h b/Engine/Plugins/Experimental/EditorDataStorage/Source/TedsCore/Private/Debug/TypedElementDatabaseDebugTypes.h
--- a/Engine/Plugins/Experimental/EditorDataStorage/Source/TedsCore/Private/Debug/TypedElementDatabaseDebugTypes.h
+++ b/Engine/Plugins/Experimental/EditorDataStorage/Source/TedsCore/Private/Debug/TypedElementDatabaseDebugTypes.h
@@ -27,8 +27,6 @@ struct FTestDynamicColumn : public FEditorDataStorageColumn
 {
 	GENERATED_BODY()
 
-	struct EditorDataStorage_DynamicColumnTemplate{};
-
 	// Note: Not trivial type to check that non-default move/copying works
 	// Should work regardless of UPROPERTY
 	TArray<int32> IntArray;
@@ -38,6 +36,4 @@ USTRUCT(meta=(EditorDataStorage_DynamicColumnTemplate))
 struct FTestDynamicTag : public FEditorDataStorageTag
 {
 	GENERATED_BODY()
-	
-	struct EditorDataStorage_DynamicColumnTemplate{};
 };
diff --git a/Engine/Plugins/Experimental/EditorDataStorageFeatures/Source/TedsAssetData/Public/TedsAssetDataColumns.h b/Engine/Plugins/Experimental/EditorDataStorageFeatures/Source/TedsAssetData/Public/TedsAssetDataColumns.h
--- a/Engine/Plugins/Experimental/EditorDataStorageFeatures/Source/TedsAssetData/Public/TedsAssetDataColumns.h
+++ b/Engine/Plugins/Experimental/EditorDataStorageFeatures/Source/TedsAssetData/Public/TedsAssetDataColumns.h
@@ -153,8 +153,6 @@ struct FItemStringAttributeColumn_Experimental : public FItemAttributeBaseColumn
 {
 	GENERATED_BODY()
 
-	struct EditorDataStorage_DynamicColumnTemplate{};
-
 	UPROPERTY()
 	FString Value;
 };
\ No newline at end of file
diff --git a/Engine/Plugins/Fab/Source/Fab/Private/Teds/FabMyFolderIntegration.h b/Engine/Plugins/Fab/Source/Fab/Private/Teds/FabMyFolderIntegration.h
--- a/Engine/Plugins/Fab/Source/Fab/Private/Teds/FabMyFolderIntegration.h
+++ b/Engine/Plugins/Fab/Source/Fab/Private/Teds/FabMyFolderIntegration.h
@@ -33,8 +33,6 @@ USTRUCT(meta = (DisplayName = "Fab distribution method", EditorDataStorage_Dynam
 struct FFabDistributionMethodTag : public FEditorDataStorageTag
 {
 	GENERATED_BODY()
-
-	struct EditorDataStorage_DynamicColumnTemplate{};
 };
 
 USTRUCT(meta = (DisplayName = "Name"))
diff --git a/Engine/Source/Programs/Shared/EpicGames.UHT/Exporters/CodeGen/EditorDataStorageUhtExtension.cs b/Engine/Source/Programs/Shared/EpicGames.UHT/Exporters/CodeGen/EditorDataStorageUhtExtension.cs
new file mode 100644
--- /dev/null
+++ b/Engine/Source/Programs/Shared/EpicGames.UHT/Exporters/CodeGen/EditorDataStorageUhtExtension.cs
@@ -0,0 +1,28 @@
+﻿// Copyright Epic Games, Inc. All Rights Reserved.
+
+using System;
+using System.Text;
+using EpicGames.UHT.Tables;
+using EpicGames.UHT.Utils;
+using EpicGames.UHT.Exporters.CodeGen;
+using EpicGames.UHT.Types;
+
+namespace EditorDataStorageUhtExtension
+{
+	[UnrealHeaderTool]
+	class Extension
+	{
+		private const string DynamicTemplateMetadata = "EditorDataStorage_DynamicColumnTemplate";
+	    
+		[UhtCodeGeneratorInjector(
+			UhtType = typeof(UhtScriptStruct), 
+			Location = UhtCodeGeneratorInjectionLocation.GeneratedMacro)]
+		public static void InjectDynamicColumnTemplateStaticAttribute(StringBuilder builder, UhtType uhtType, int leadingTabs, string eolSequence)
+		{
+			if (uhtType.MetaData.ContainsKey(DynamicTemplateMetadata))
+			{
+				builder.Append('\t', leadingTabs).Append("struct EditorDataStorage_DynamicColumnTemplate{};").Append(eolSequence);
+			}
+		}
+	}
+}
\ No newline at end of file
diff --git a/Engine/Source/Runtime/TypedElementFramework/Public/Elements/Columns/TypedElementWebColumns.h b/Engine/Source/Runtime/TypedElementFramework/Public/Elements/Columns/TypedElementWebColumns.h
--- a/Engine/Source/Runtime/TypedElementFramework/Public/Elements/Columns/TypedElementWebColumns.h
+++ b/Engine/Source/Runtime/TypedElementFramework/Public/Elements/Columns/TypedElementWebColumns.h
@@ -28,8 +28,6 @@ struct FEditorDataStorageWebImageColumn final : public FEditorDataStorageColumn
 {
 	GENERATED_BODY()
 
-	struct EditorDataStorage_DynamicColumnTemplate{};
-
 	/** URL to the image in string form. Can be used with for instance with WebImage(Cache) to create a local image. */
 	UPROPERTY()
 	FString UrlString;
diff --git a/Engine/Source/Runtime/TypedElementFramework/Tests/Elements/Framework/TypedElementTestColumns.h b/Engine/Source/Runtime/TypedElementFramework/Tests/Elements/Framework/TypedElementTestColumns.h
--- a/Engine/Source/Runtime/TypedElementFramework/Tests/Elements/Framework/TypedElementTestColumns.h
+++ b/Engine/Source/Runtime/TypedElementFramework/Tests/Elements/Framework/TypedElementTestColumns.h
@@ -53,8 +53,6 @@ USTRUCT(meta = (DisplayName = "ColumnDynamic", EditorDataStorage_DynamicColumnTe
 struct FTestColumnDynamic final : public FEditorDataStorageColumn
 {
 	GENERATED_BODY()
-
-	struct EditorDataStorage_DynamicColumnTemplate {};
 };
 
 USTRUCT(meta = (DisplayName = "ColumnInt"))
