Horde: Use zeroed IV for UbaCrypto to match HTTP endpoint in UBA cache server

diff --git a/Engine/Source/Programs/Shared/EpicGames.Horde/Compute/UbaCrypto.cs b/Engine/Source/Programs/Shared/EpicGames.Horde/Compute/UbaCrypto.cs
--- a/Engine/Source/Programs/Shared/EpicGames.Horde/Compute/UbaCrypto.cs
+++ b/Engine/Source/Programs/Shared/EpicGames.Horde/Compute/UbaCrypto.cs
@@ -14,7 +14,7 @@ namespace EpicGames.Horde.Compute;
 public class UbaCrypto : IDisposable
 {
 	private const int AesBlockSize = 16;
-	private static readonly byte[] s_fixedIvLe = [0x07, 0x06, 0x05, 0x04, 0x03, 0x02, 0x01, 0x00, 0x0F, 0x0E, 0x0D, 0x0C, 0x0B, 0x0A, 0x09, 0x08];
+	private static readonly byte[] s_fixedIvLe = [0x00, 0x00, 0x00, 0x00, 0x00, 0x00, 0x00, 0x00, 0x00, 0x00, 0x00, 0x00, 0x00, 0x00, 0x00, 0x00];
 	
 	private readonly ThreadLocal<Aes> _aes;
 	private bool _disposed;
