Horde: Fix broken UbaCrypto test

diff --git a/Engine/Source/Programs/Shared/EpicGames.Horde.Tests/Compute/UbaCryptoTests.cs b/Engine/Source/Programs/Shared/EpicGames.Horde.Tests/Compute/UbaCryptoTests.cs
--- a/Engine/Source/Programs/Shared/EpicGames.Horde.Tests/Compute/UbaCryptoTests.cs
+++ b/Engine/Source/Programs/Shared/EpicGames.Horde.Tests/Compute/UbaCryptoTests.cs
@@ -15,9 +15,9 @@ public class UbaCryptoTests
 	private const string Decrypted1 = "This is 15 char";
 	private const string Encrypted1 = "5B 67 66 7C 2F 66 7C 2F 3E 3A 2F 6C 67 6E 7D";
 	private const string Decrypted2 = "This is 16 chars";
-	private const string Encrypted2 = "4C 94 C9 27 FC 58 91 D3 FA 05 09 38 12 00 5E 84";
+	private const string Encrypted2 = "24 DE B6 E3 3B 6A CD D0 A1 BA A8 46 F2 DC B6 CD";
 	private const string Decrypted3 = "This is 17 chars!";
-	private const string Encrypted3 = "08 33 98 87 51 CD F0 A7 1D 66 58 B2 EF 38 3F 2F 75";
+	private const string Encrypted3 = "22 42 52 20 03 0F 14 53 AD CA B4 20 A8 1D 30 EA 75";
 	
 	[TestMethod]
 	[DataRow("", "")]
