Horde: Add method to compute client for allocating a UBA cache server session

diff --git a/Engine/Source/Programs/Shared/EpicGames.Horde/Compute/Clients/AgentComputeClient.cs b/Engine/Source/Programs/Shared/EpicGames.Horde/Compute/Clients/AgentComputeClient.cs
--- a/Engine/Source/Programs/Shared/EpicGames.Horde/Compute/Clients/AgentComputeClient.cs
+++ b/Engine/Source/Programs/Shared/EpicGames.Horde/Compute/Clients/AgentComputeClient.cs
@@ -107,7 +107,13 @@ namespace EpicGames.Horde.Compute.Clients
 		{
 			return Task.CompletedTask;
 		}
-
+		
+		/// <inheritdoc/>
+		public Task<UbaConfig> AllocateUbaCacheServerAsync(ClusterId clusterId, CancellationToken cancellationToken = default)
+		{
+			throw new System.NotImplementedException();
+		}
+		
 		async IAsyncEnumerable<RemoteComputeSocket> ConnectAsync(ILogger logger, [EnumeratorCancellation] CancellationToken cancellationToken)
 		{
 			using Socket listener = new Socket(SocketType.Stream, ProtocolType.IP);
diff --git a/Engine/Source/Programs/Shared/EpicGames.Horde/Compute/Clients/LocalComputeClient.cs b/Engine/Source/Programs/Shared/EpicGames.Horde/Compute/Clients/LocalComputeClient.cs
--- a/Engine/Source/Programs/Shared/EpicGames.Horde/Compute/Clients/LocalComputeClient.cs
+++ b/Engine/Source/Programs/Shared/EpicGames.Horde/Compute/Clients/LocalComputeClient.cs
@@ -108,5 +108,11 @@ namespace EpicGames.Horde.Compute.Clients
 		{
 			return Task.CompletedTask;
 		}
+		
+		/// <inheritdoc/>
+		public Task<UbaConfig> AllocateUbaCacheServerAsync(ClusterId clusterId, CancellationToken cancellationToken = default)
+		{
+			throw new NotImplementedException();
+		}
 	}
 }
diff --git a/Engine/Source/Programs/Shared/EpicGames.Horde/Compute/Clients/ServerComputeClient.cs b/Engine/Source/Programs/Shared/EpicGames.Horde/Compute/Clients/ServerComputeClient.cs
--- a/Engine/Source/Programs/Shared/EpicGames.Horde/Compute/Clients/ServerComputeClient.cs
+++ b/Engine/Source/Programs/Shared/EpicGames.Horde/Compute/Clients/ServerComputeClient.cs
@@ -269,7 +269,36 @@ namespace EpicGames.Horde.Compute.Clients
 			using HttpResponseMessage response = await HordeHttpRequest.PostAsync(_httpClient, $"api/v2/compute/{clusterId}/resource-needs", request, _cancellationSource.Token);
 			response.EnsureSuccessStatusCode();
 		}
-
+		
+		/// <inheritdoc/>
+		public async Task<UbaConfig> AllocateUbaCacheServerAsync(ClusterId clusterId, CancellationToken cancellationToken = default)
+		{
+			AllocateUbaCacheServerRequest request = new();
+			using HttpResponseMessage httpResponse = await HordeHttpRequest.PostAsync(_httpClient, $"api/v2/compute/{clusterId}/uba-cache", request, _cancellationSource.Token);
+			if (httpResponse.StatusCode == HttpStatusCode.InternalServerError)
+			{
+				string? content;
+				try
+				{
+					content = await httpResponse.Content.ReadAsStringAsync(cancellationToken);
+				}
+				catch
+				{
+					content = "None";
+				}
+				throw new ComputeClientException($"InternalServerError requesting compute resources: \"{content}\"");
+			}
+			
+			httpResponse.EnsureSuccessStatusCode();
+			UbaConfig? response = await httpResponse.Content.ReadFromJsonAsync<UbaConfig>(HordeHttpClient.JsonSerializerOptions, cancellationToken);
+			if (response == null)
+			{
+				throw new InvalidOperationException();
+			}
+			
+			return response;
+		}
+		
 		async IAsyncEnumerable<LeaseInfo> ConnectAsync(ClusterId? clusterId, Requirements? requirements, string? requestId, ConnectionMetadataRequest? connection, bool? useUbaCache, ILogger workerLogger, [EnumeratorCancellation] CancellationToken cancellationToken)
 		{
 			_logger.LogDebug("Requesting compute resource");
diff --git a/Engine/Source/Programs/Shared/EpicGames.Horde/Compute/IComputeClient.cs b/Engine/Source/Programs/Shared/EpicGames.Horde/Compute/IComputeClient.cs
--- a/Engine/Source/Programs/Shared/EpicGames.Horde/Compute/IComputeClient.cs
+++ b/Engine/Source/Programs/Shared/EpicGames.Horde/Compute/IComputeClient.cs
@@ -45,6 +45,13 @@ namespace EpicGames.Horde.Compute
 		/// <param name="resourceNeeds">Properties with a target amount of each, such as CPU or RAM</param>
 		/// <param name="cancellationToken">Cancellation token for the operation</param>
 		public Task DeclareResourceNeedsAsync(ClusterId clusterId, string pool, Dictionary<string, int> resourceNeeds, CancellationToken cancellationToken = default);
+		
+		/// <summary>
+		/// Allocate a session for the UBA cache associated with the given cluster
+		/// </summary>
+		/// <param name="clusterId">Cluster to execute the request</param>
+		/// <param name="cancellationToken">Cancellation token for the operation</param>
+		public Task<UbaConfig> AllocateUbaCacheServerAsync(ClusterId clusterId, CancellationToken cancellationToken = default);
 	}
 
 	/// <summary>
