Horde: Skip client IP allow listing for allocated compute agents

diff --git a/Engine/Source/Programs/Shared/EpicGames.Horde/Compute/Clients/ServerComputeClient.cs b/Engine/Source/Programs/Shared/EpicGames.Horde/Compute/Clients/ServerComputeClient.cs
--- a/Engine/Source/Programs/Shared/EpicGames.Horde/Compute/Clients/ServerComputeClient.cs
+++ b/Engine/Source/Programs/Shared/EpicGames.Horde/Compute/Clients/ServerComputeClient.cs
@@ -184,7 +184,6 @@ namespace EpicGames.Horde.Compute.Clients
 		readonly CancellationTokenSource _cancellationSource = new CancellationTokenSource();
 		readonly string _sessionId;
 		readonly ILogger _logger;
-		readonly ExternalIpResolver _externalIpResolver;
 
 		/// <summary>
 		/// Constructor
@@ -206,7 +205,6 @@ namespace EpicGames.Horde.Compute.Clients
 			_httpClient = httpClient;
 			_sessionId = sessionId ?? Guid.NewGuid().ToString();
 			_logger = logger;
-			_externalIpResolver = new ExternalIpResolver(_httpClient);
 		}
 
 		/// <inheritdoc/>
@@ -312,11 +310,6 @@ namespace EpicGames.Horde.Compute.Clients
 				UseUbaCache = useUbaCache,
 				Protocol = (int)ComputeProtocol.Latest
 			};
-			
-			if (connection is { ModePreference: ConnectionMode.Relay })
-			{
-				connection.ClientPublicIp = (await _externalIpResolver.GetExternalIpAddressAsync(cancellationToken)).ToString();
-			}
 
 			AssignComputeResponse? response;
 			string path = clusterId == null ? "api/v2/compute" : $"api/v2/compute/{clusterId}";
