Apply minsize pragma to generated code.

diff --git a/Engine/Source/Programs/Shared/EpicGames.UHT/Exporters/CodeGen/UhtHeaderCodeGeneratorCppFile.cs b/Engine/Source/Programs/Shared/EpicGames.UHT/Exporters/CodeGen/UhtHeaderCodeGeneratorCppFile.cs
--- a/Engine/Source/Programs/Shared/EpicGames.UHT/Exporters/CodeGen/UhtHeaderCodeGeneratorCppFile.cs
+++ b/Engine/Source/Programs/Shared/EpicGames.UHT/Exporters/CodeGen/UhtHeaderCodeGeneratorCppFile.cs
@@ -158,6 +158,7 @@ namespace EpicGames.UHT.Exporters.CodeGen
 				}
 
 				builder.Append("\r\n");
+				builder.Append(PragmaPushMinsizeFunctions).Append("\r\n");
 				builder.Append(DisableDeprecationWarnings).Append("\r\n");
 
 				builder.Append("\r\n");
@@ -354,6 +355,7 @@ namespace EpicGames.UHT.Exporters.CodeGen
 
 				builder.Append("\r\n");
 				builder.Append(EnableDeprecationWarnings).Append("\r\n");
+				builder.Append(PragmaPopMinsizeFunctions).Append("\r\n");
 
 				{
 					using UhtRentedPoolBuffer<char> borrowBuffer = builder.RentPoolBuffer();
diff --git a/Engine/Source/Programs/Shared/EpicGames.UHT/Exporters/CodeGen/UhtPackageCodeGenerator.cs b/Engine/Source/Programs/Shared/EpicGames.UHT/Exporters/CodeGen/UhtPackageCodeGenerator.cs
--- a/Engine/Source/Programs/Shared/EpicGames.UHT/Exporters/CodeGen/UhtPackageCodeGenerator.cs
+++ b/Engine/Source/Programs/Shared/EpicGames.UHT/Exporters/CodeGen/UhtPackageCodeGenerator.cs
@@ -21,6 +21,8 @@ namespace EpicGames.UHT.Exporters.CodeGen
 
 		public static string EnableDeprecationWarnings = "PRAGMA_ENABLE_DEPRECATION_WARNINGS";
 		public static string DisableDeprecationWarnings = "PRAGMA_DISABLE_DEPRECATION_WARNINGS";
+		public static string PragmaPushMinsizeFunctions = "PRAGMA_PUSH_ATTRIBUTE_MINSIZE_FUNCTIONS";
+		public static string PragmaPopMinsizeFunctions = "PRAGMA_POP_ATTRIBUTE_MINSIZE_FUNCTIONS";
 
 		public readonly UhtCodeGenerator CodeGenerator;
 		public readonly UhtModule Module;
diff --git a/Engine/Source/Programs/Shared/EpicGames.UHT/Exporters/CodeGen/UhtPackageCodeGeneratorCppFile.cs b/Engine/Source/Programs/Shared/EpicGames.UHT/Exporters/CodeGen/UhtPackageCodeGeneratorCppFile.cs
--- a/Engine/Source/Programs/Shared/EpicGames.UHT/Exporters/CodeGen/UhtPackageCodeGeneratorCppFile.cs
+++ b/Engine/Source/Programs/Shared/EpicGames.UHT/Exporters/CodeGen/UhtPackageCodeGeneratorCppFile.cs
@@ -89,6 +89,7 @@ namespace EpicGames.UHT.Exporters.CodeGen
 
 				builder.Append(HeaderCopyright);
 				builder.Append(RequiredCPPIncludes);
+				builder.Append(PragmaPushMinsizeFunctions).Append("\r\n");
 				builder.Append(DisableDeprecationWarnings).Append("\r\n");
 				builder.Append("void EmptyLinkFunctionForGeneratedCode").Append(Module.ShortName).Append("_init() {}\r\n");
 
@@ -151,6 +152,7 @@ namespace EpicGames.UHT.Exporters.CodeGen
 						.Append($"0x{bodiesHash:X8}, 0x{declarationsHash:X8}").Append("));\r\n");
 				}
 				builder.Append(EnableDeprecationWarnings).Append("\r\n");
+				builder.Append(PragmaPopMinsizeFunctions).Append("\r\n");
 
 				if (SaveExportedHeaders)
 				{
