EpicGames.Perforce: Log if untracked files are removed

diff --git a/Engine/Source/Programs/Shared/EpicGames.Perforce.Managed/ManagedWorkspace.cs b/Engine/Source/Programs/Shared/EpicGames.Perforce.Managed/ManagedWorkspace.cs
--- a/Engine/Source/Programs/Shared/EpicGames.Perforce.Managed/ManagedWorkspace.cs
+++ b/Engine/Source/Programs/Shared/EpicGames.Perforce.Managed/ManagedWorkspace.cs
@@ -976,11 +976,11 @@ namespace EpicGames.Perforce.Managed
 			Stopwatch timer = Stopwatch.StartNew();
 			if (changeNumber == -1)
 			{
-				_logger.LogInformation("Syncing to {StreamName} at latest", streamName);
+				_logger.LogInformation("Syncing to {StreamName} at latest (removeUntracked={RemoveUntracked})", streamName, removeUntracked);
 			}
 			else
 			{
-				_logger.LogInformation("Syncing to {StreamName} at CL {CL}", streamName, changeNumber);
+				_logger.LogInformation("Syncing to {StreamName} at CL {CL} (removeUntracked={RemoveUntracked})", streamName, changeNumber, removeUntracked);
 			}
 
 			using (_logger.WithProperty("useHaveTable", _options.UseHaveTable).BeginScope())
