EpicGames.Core: Cache result of FileSystemReference.GetHashCode()

#rnx

diff --git a/Engine/Source/Programs/Shared/EpicGames.Core/DirectoryReference.cs b/Engine/Source/Programs/Shared/EpicGames.Core/DirectoryReference.cs
--- a/Engine/Source/Programs/Shared/EpicGames.Core/DirectoryReference.cs
+++ b/Engine/Source/Programs/Shared/EpicGames.Core/DirectoryReference.cs
@@ -193,11 +193,7 @@ namespace EpicGames.Core
 		/// <returns>False if the names represent the same object, true otherwise</returns>
 		public static bool operator !=(DirectoryReference? a, DirectoryReference? b) => !(a == b);
 
-		/// <summary>
-		/// Compares against another object for equality.
-		/// </summary>
-		/// <param name="obj">other instance to compare.</param>
-		/// <returns>True if the names represent the same object, false otherwise</returns>
+		/// <inheritdoc/>
 		public override bool Equals(object? obj) => (obj is DirectoryReference dir) && dir == this;
 
 		/// <summary>
@@ -207,11 +203,8 @@ namespace EpicGames.Core
 		/// <returns>True if the names represent the same object, false otherwise</returns>
 		public bool Equals(DirectoryReference? obj) => obj == this;
 
-		/// <summary>
-		/// Returns a hash code for this object
-		/// </summary>
-		/// <returns></returns>
-		public override int GetHashCode() => Comparer.GetHashCode(FullName);
+		/// <inheritdoc/>
+		public override int GetHashCode() => base.GetHashCode();
 
 		/// <inheritdoc/>
 		public int CompareTo(DirectoryReference? other) => Comparer.Compare(FullName, other?.FullName);
@@ -232,27 +225,27 @@ namespace EpicGames.Core
 		/// Moves a directory from one location to another
 		/// </summary>
 		/// <param name="sourceDir">The source directory</param>
-		/// <param name="destDir">The destination directory</param>
-		public static void Move(DirectoryReference sourceDir, DirectoryReference destDir)
-			=> Directory.Move(sourceDir.FullName, destDir.FullName);
+		/// <param name="destinationDir">The destination directory</param>
+		public static void Move(DirectoryReference sourceDir, DirectoryReference destinationDir)
+			=> Directory.Move(sourceDir.FullName, destinationDir.FullName);
 		
 		/// <summary>
 		/// Copies a directory from one location to another
 		/// </summary>
 		/// <param name="sourceDir">The source directory</param>
-		/// <param name="destDir">The destination directory</param>
-		public static void Copy(DirectoryReference sourceDir, DirectoryReference destDir)
+		/// <param name="destinationDir">The destination directory</param>
+		public static void Copy(DirectoryReference sourceDir, DirectoryReference destinationDir)
 		{
-			CreateDirectory(destDir);
+			CreateDirectory(destinationDir);
 			
 			foreach (string file in Directory.GetFiles(sourceDir.FullName))
 			{
-				File.Copy(file, Path.Combine(destDir.FullName, Path.GetFileName(file)), true);
+				File.Copy(file, Path.Combine(destinationDir.FullName, Path.GetFileName(file)), true);
 			}
 			
 			foreach (string directory in Directory.GetDirectories(sourceDir.FullName))
 			{
-				Copy(new DirectoryReference(directory), Combine(destDir, Path.GetFileName(directory)));
+				Copy(new DirectoryReference(directory), Combine(destinationDir, Path.GetFileName(directory)));
 			}
 		}
 		
diff --git a/Engine/Source/Programs/Shared/EpicGames.Core/FileReference.cs b/Engine/Source/Programs/Shared/EpicGames.Core/FileReference.cs
--- a/Engine/Source/Programs/Shared/EpicGames.Core/FileReference.cs
+++ b/Engine/Source/Programs/Shared/EpicGames.Core/FileReference.cs
@@ -188,11 +188,7 @@ namespace EpicGames.Core
 		/// <returns>False if the names represent the same object, true otherwise</returns>
 		public static bool operator !=(FileReference? a, FileReference? b) => !(a == b);
 
-		/// <summary>
-		/// Compares against another object for equality.
-		/// </summary>
-		/// <param name="obj">other instance to compare.</param>
-		/// <returns>True if the names represent the same object, false otherwise</returns>
+		/// <inheritdoc/>
 		public override bool Equals(object? obj) => obj is FileReference file && file == this;
 
 		/// <summary>
@@ -202,11 +198,8 @@ namespace EpicGames.Core
 		/// <returns>True if the names represent the same object, false otherwise</returns>
 		public bool Equals(FileReference? obj) => obj == this;
 
-		/// <summary>
-		/// Returns a hash code for this object
-		/// </summary>
-		/// <returns></returns>
-		public override int GetHashCode() => Comparer.GetHashCode(FullName);
+		/// <inheritdoc/>
+		public override int GetHashCode() => base.GetHashCode();
 
 		/// <inheritdoc/>
 		public int CompareTo(FileReference? other) => Comparer.Compare(FullName, other?.FullName);
diff --git a/Engine/Source/Programs/Shared/EpicGames.Core/FileSystemReference.cs b/Engine/Source/Programs/Shared/EpicGames.Core/FileSystemReference.cs
--- a/Engine/Source/Programs/Shared/EpicGames.Core/FileSystemReference.cs
+++ b/Engine/Source/Programs/Shared/EpicGames.Core/FileSystemReference.cs
@@ -33,7 +33,19 @@ namespace EpicGames.Core
 		/// <summary>
 		/// Direct constructor for a path
 		/// </summary>
-		protected FileSystemReference(string fullName) => FullName = fullName;
+		protected FileSystemReference(string fullName)
+		{
+			FullName = fullName;
+			_hashCode = new(() => Comparer.GetHashCode(FullName));
+		}
+
+		/// <inheritdoc/>
+		public override int GetHashCode() => _hashCode.Value;
+
+		/// <summary>
+		/// Lazy initialized hash code
+		/// </summary>
+		private readonly Lazy<int> _hashCode;
 
 		static readonly ThreadLocal<StringBuilder> s_combineStringsStringBuilder = new ThreadLocal<StringBuilder>(() => new StringBuilder(260));
 
@@ -134,7 +146,7 @@ namespace EpicGames.Core
 		/// Checks to see if this exists as either a file or directory
 		/// This is helpful for Mac, because a binary may be a .app which is a directory
 		/// </summary>
-		/// <param name="location">FileSsytem object to check</param>
+		/// <param name="location">FileSystem object to check</param>
 		/// <returns>True if a file or a directory exists</returns>
 		public static bool Exists(FileSystemReference location) => File.Exists(location.FullName) || Directory.Exists(location.FullName);
 
@@ -215,7 +227,7 @@ namespace EpicGames.Core
 		}
 
 		/// <summary>
-		/// Appens a relative path to a string builder
+		/// Appends a relative path to a string builder
 		/// </summary>
 		/// <param name="directory"></param>
 		/// <param name="result"></param>
