UnrealBuildTool: Update lazy initialization in FileItem \ DirectoryItem to be thread-safe
Use comparison functions from FileSystemReference rather than using own implemetation to remove duplicate logic

#rnx

diff --git a/Engine/Source/Programs/Shared/EpicGames.Build/System/DirectoryItem.cs b/Engine/Source/Programs/Shared/EpicGames.Build/System/DirectoryItem.cs
--- a/Engine/Source/Programs/Shared/EpicGames.Build/System/DirectoryItem.cs
+++ b/Engine/Source/Programs/Shared/EpicGames.Build/System/DirectoryItem.cs
@@ -29,17 +29,17 @@ namespace UnrealBuildBase
 		/// <summary>
 		/// Cached map of name to subdirectory item
 		/// </summary>
-		Dictionary<string, DirectoryItem>? Directories;
+		Lazy<Dictionary<string, DirectoryItem>> CachedDirectories;
 
 		/// <summary>
 		/// Cached map of name to file
 		/// </summary>
-		Dictionary<string, FileItem>? Files;
+		Lazy<Dictionary<string, FileItem>> CachedFiles;
 
 		/// <summary>
 		/// Global map of location to item
 		/// </summary>
-		static ConcurrentDictionary<DirectoryReference, DirectoryItem> LocationToItem = new ConcurrentDictionary<DirectoryReference, DirectoryItem>();
+		static ConcurrentDictionary<DirectoryReference, DirectoryItem> LocationToItem = [];
 
 		/// <summary>
 		/// Constructor
@@ -49,21 +49,29 @@ namespace UnrealBuildBase
 		private DirectoryItem(DirectoryReference Location, DirectoryInfo Info)
 		{
 			this.Location = Location;
-			if (OperatingSystem.IsWindows())
-			{
-				this.Info = new Lazy<DirectoryInfo>(Info);
-			}
-			else
-			{
-				// For some reason we need to call an extra Refresh on linux/mac to not get wrong results from "Exists"
-				this.Info = new Lazy<DirectoryInfo>(() =>
+			CachedDirectories = new(ScanDirectories);
+			CachedFiles = new(ScanFiles);
+
+			// For some reason we need to call an extra Refresh on Linux/Mac to not get wrong results from "Exists"
+			this.Info = OperatingSystem.IsWindows()
+				? new(Info)
+				: new(() =>
 				{
 					Info.Refresh();
 					return Info;
 				});
-			}
 		}
 
+		/// <summary>
+		/// Accessor for map of name to subdirectory item
+		/// </summary>
+		Dictionary<string, DirectoryItem> Directories => CachedDirectories.Value;
+
+		/// <summary>
+		/// Accessor for map of name to file
+		/// </summary>
+		Dictionary<string, FileItem> Files => CachedFiles.Value;
+
 		/// <summary>
 		/// The name of this directory
 		/// </summary>
@@ -111,19 +119,14 @@ namespace UnrealBuildBase
 		/// <param name="Fragments">The path fragments to append</param>
 		/// <returns>Directory item corresponding to the combined path</returns>
 		public static DirectoryItem Combine(DirectoryItem BaseDirectory, params string[] Fragments)
-		{
-			return DirectoryItem.GetItemByDirectoryReference(DirectoryReference.Combine(BaseDirectory.Location, Fragments));
-		}
+			=> GetItemByDirectoryReference(DirectoryReference.Combine(BaseDirectory.Location, Fragments));
 
 		/// <summary>
 		/// Finds or creates a directory item from its location
 		/// </summary>
 		/// <param name="Location">Path to the directory</param>
 		/// <returns>The directory item for this location</returns>
-		public static DirectoryItem GetItemByPath(string Location)
-		{
-			return GetItemByDirectoryReference(new DirectoryReference(Location));
-		}
+		public static DirectoryItem GetItemByPath(string Location) => GetItemByDirectoryReference(new(Location));
 
 		/// <summary>
 		/// Finds or creates a directory item from its location
@@ -131,59 +134,45 @@ namespace UnrealBuildBase
 		/// <param name="Location">Path to the directory</param>
 		/// <returns>The directory item for this location</returns>
 		public static DirectoryItem GetItemByDirectoryReference(DirectoryReference Location)
-		{
-			if (LocationToItem.TryGetValue(Location, out DirectoryItem? Result))
-			{
-				return Result;
-			}
-			return LocationToItem.GetOrAdd(Location, new DirectoryItem(Location, new DirectoryInfo(Location.FullName)));
-		}
+			=> LocationToItem.TryGetValue(Location, out DirectoryItem? Result)
+				? Result
+				: LocationToItem.GetOrAdd(Location, new DirectoryItem(Location, new(Location.FullName)));
 
 		/// <summary>
 		/// Finds or creates a directory item from a DirectoryInfo object
 		/// </summary>
 		/// <param name="Info">Path to the directory</param>
 		/// <returns>The directory item for this location</returns>
-		public static DirectoryItem GetItemByDirectoryInfo(DirectoryInfo Info)
-		{
-			DirectoryReference Location = new DirectoryReference(Info);
-			if (LocationToItem.TryGetValue(Location, out DirectoryItem? Result))
-			{
-				return Result;
-			}
-			return LocationToItem.GetOrAdd(Location, new DirectoryItem(Location, Info));
-		}
+		public static DirectoryItem GetItemByDirectoryInfo(DirectoryInfo Info) => GetItemByDirectoryReference(new(Info));
 
 		/// <summary>
 		/// Reset the contents of the directory and allow them to be fetched again
 		/// </summary>
 		public void ResetCachedInfo()
 		{
-			Info = new Lazy<DirectoryInfo>(() =>
+			Info = new(() =>
 			{
 				DirectoryInfo Info = Location.ToDirectoryInfo();
 				Info.Refresh();
 				return Info;
 			});
 
-			Dictionary<string, DirectoryItem>? PrevDirectories = Directories;
-			if (PrevDirectories != null)
+			if (CachedDirectories.IsValueCreated)
 			{
-				foreach (DirectoryItem SubDirectory in PrevDirectories.Values)
+				foreach (DirectoryItem SubDirectory in CachedDirectories.Value.Values)
 				{
 					SubDirectory.ResetCachedInfo();
 				}
-				Directories = null;
+				CachedDirectories = new(ScanDirectories);
 			}
 
-			Dictionary<string, FileItem>? PrevFiles = Files;
-			if (PrevFiles != null)
+			if (CachedFiles.IsValueCreated)
 			{
-				foreach (FileItem File in PrevFiles.Values)
+				foreach (FileItem File in CachedFiles.Value.Values)
 				{
 					File.ResetCachedInfo();
 				}
-				Files = null;
+				CachedFiles = new(ScanFiles);
 			}
 		}
 
@@ -203,7 +192,7 @@ namespace UnrealBuildBase
 		/// </summary>
 		public static void ResetAllCachedInfo_SLOW()
 		{
-			foreach (DirectoryItem Item in LocationToItem.Values)
+			LocationToItem.Values.AsParallel().ForAll(Item =>
 			{
 				Item.Info = new Lazy<DirectoryInfo>(() =>
 				{
@@ -211,48 +200,52 @@ namespace UnrealBuildBase
 					Info.Refresh();
 					return Info;
 				});
-				Item.Directories = null;
-				Item.Files = null;
-			}
+				Item.CachedDirectories = new(Item.ScanDirectories);
+				Item.CachedFiles = new(Item.ScanFiles);
+			});
 			FileItem.ResetAllCachedInfo_SLOW();
 		}
 
 		/// <summary>
-		/// Caches the subdirectories of this directories
+		/// Scans the directory for directories, used for lazy initialization.
 		/// </summary>
-		public void CacheDirectories()
+		/// <returns></returns>
+		Dictionary<string, DirectoryItem> ScanDirectories()
 		{
-			if (Directories == null)
+			Dictionary<string, DirectoryItem>? NewDirectories = null;
+			if (Info.Value.Exists)
 			{
-				Dictionary<string, DirectoryItem>? NewDirectories = null;
-				if (Info.Value.Exists)
+				try
 				{
-					try
+					IEnumerable<DirectoryInfo> DirectoryInfos = Info.Value.EnumerateDirectories();
+					NewDirectories = new(DirectoryInfos.Count(), DirectoryReference.Comparer);
+					foreach (DirectoryInfo SubDirectoryInfo in DirectoryInfos)
 					{
-						DirectoryInfo[] Directories = Info.Value.GetDirectories();
-						NewDirectories = new Dictionary<string, DirectoryItem>(Directories.Length, DirectoryReference.Comparer);
-						foreach (DirectoryInfo SubDirectoryInfo in Directories)
+						if (NewDirectories.TryGetValue(SubDirectoryInfo.Name, out DirectoryItem? value))
 						{
-							if (NewDirectories.TryGetValue(SubDirectoryInfo.Name, out DirectoryItem? value))
-							{
-								throw new Exception($"Trying to add {SubDirectoryInfo.FullName} as '{SubDirectoryInfo.Name}' yet exists as {value.FullName}");
-							}
-
-							NewDirectories.Add(SubDirectoryInfo.Name, DirectoryItem.GetItemByDirectoryInfo(SubDirectoryInfo));
+							throw new Exception($"Trying to add {SubDirectoryInfo.FullName} as '{SubDirectoryInfo.Name}' yet exists as {value.FullName}");
 						}
-					}
-					catch (SecurityException)
-					{
-					}
-					catch (UnauthorizedAccessException)
-					{
+
+						NewDirectories.Add(SubDirectoryInfo.Name, DirectoryItem.GetItemByDirectoryInfo(SubDirectoryInfo));
 					}
 				}
-
-				Directories = NewDirectories ?? new Dictionary<string, DirectoryItem>(DirectoryReference.Comparer);
+				catch (SecurityException)
+				{
+				}
+				catch (UnauthorizedAccessException)
+				{
+				}
 			}
+
+			return NewDirectories ?? new(DirectoryReference.Comparer);
 		}
 
+
+		/// <summary>
+		/// Caches the subdirectories of this directories
+		/// </summary>
+		public bool CacheDirectories() => CachedDirectories.Value != null;
+
 		/// <summary>
 		/// Enumerates all the subdirectories
 		/// </summary>
@@ -260,7 +253,7 @@ namespace UnrealBuildBase
 		public IEnumerable<DirectoryItem> EnumerateDirectories()
 		{
 			CacheDirectories();
-			return Directories!.Values;
+			return Directories.Values;
 		}
 
 		/// <summary>
@@ -286,44 +279,45 @@ namespace UnrealBuildBase
 			}
 
 			CacheDirectories();
-			return Directories!.TryGetValue(Name, out OutDirectory);
+			return Directories.TryGetValue(Name, out OutDirectory);
 		}
 
 		/// <summary>
-		/// Caches the files in this directory
+		/// Scans the directory for files, used for lazy initialization.
 		/// </summary>
-		public void CacheFiles()
+		/// <returns></returns>
+		Dictionary<string, FileItem> ScanFiles()
 		{
-			if (Files == null)
+			Dictionary<string, FileItem>? NewFiles = null;
+			if (Info.Value.Exists)
 			{
-				Dictionary<string, FileItem>? NewFiles = null;
-				if (Info.Value.Exists)
+				try
 				{
-					try
-					{
-						FileInfo[] FileInfos = Info.Value.GetFiles();
-						NewFiles = new Dictionary<string, FileItem>(FileInfos.Length, FileReference.Comparer);
-						foreach (FileInfo FileInfo in FileInfos)
-						{
-							FileItem FileItem = FileItem.GetItemByFileInfo(FileInfo);
-							FileItem.UpdateCachedDirectory(this);
-							NewFiles[FileInfo.Name] = FileItem;
-						}
-						Files = NewFiles;
-						return;
-					}
-					catch (SecurityException)
-					{
-					}
-					catch (UnauthorizedAccessException)
+					IEnumerable<FileInfo> FileInfos = Info.Value.EnumerateFiles();
+					NewFiles = new(FileInfos.Count(), FileReference.Comparer);
+					foreach (FileInfo FileInfo in Info.Value.EnumerateFiles())
 					{
+						FileItem FileItem = FileItem.GetItemByFileInfo(FileInfo);
+						FileItem.UpdateCachedDirectory(this);
+						NewFiles[FileInfo.Name] = FileItem;
 					}
 				}
-
-				Files = NewFiles ?? new Dictionary<string, FileItem>(FileReference.Comparer);
+				catch (SecurityException)
+				{
+				}
+				catch (UnauthorizedAccessException)
+				{
+				}
 			}
+
+			return NewFiles ?? new(FileReference.Comparer);
 		}
 
+		/// <summary>
+		/// Caches the files in this directory
+		/// </summary>
+		public bool CacheFiles() => CachedFiles.Value != null;
+
 		/// <summary>
 		/// Enumerates all the files
 		/// </summary>
@@ -331,7 +325,7 @@ namespace UnrealBuildBase
 		public IEnumerable<FileItem> EnumerateFiles()
 		{
 			CacheFiles();
-			return Files!.Values;
+			return Files.Values;
 		}
 
 		/// <summary>
@@ -354,17 +348,14 @@ namespace UnrealBuildBase
 		public bool TryGetFile(string Name, [NotNullWhen(true)] out FileItem? OutFile)
 		{
 			CacheFiles();
-			return Files!.TryGetValue(Name, out OutFile);
+			return Files.TryGetValue(Name, out OutFile);
 		}
 
 		/// <summary>
 		/// Formats this object as a string for debugging
 		/// </summary>
 		/// <returns>Location of the directory</returns>
-		public override string ToString()
-		{
-			return Location.FullName;
-		}
+		public override string ToString() => Location.ToString();
 
 		/// <summary>
 		/// Writes out all the enumerated files full names sorted to OutFile
@@ -386,16 +377,13 @@ namespace UnrealBuildBase
 		}
 
 		#region IComparable, IEquatbale
-		public int CompareTo(DirectoryItem? other)
-		{
-			return FullName.CompareTo(other?.FullName);
-		}
+		/// <inheritdoc/>
+		public int CompareTo(DirectoryItem? other) => Location.CompareTo(other?.Location);
 
-		public bool Equals(DirectoryItem? other)
-		{
-			return FullName.Equals(other?.FullName);
-		}
+		/// <inheritdoc/>
+		public bool Equals(DirectoryItem? other) => Location.Equals(other?.Location);
 
+		/// <inheritdoc/>
 		public override bool Equals(object? obj)
 		{
 			if (ReferenceEquals(this, obj))
@@ -411,10 +399,8 @@ namespace UnrealBuildBase
 			return Equals(obj as DirectoryItem);
 		}
 
-		public override int GetHashCode()
-		{
-			return FullName.GetHashCode();
-		}
+		/// <inheritdoc/>
+		public override int GetHashCode() => Location.GetHashCode();
 
 		public static bool operator ==(DirectoryItem? left, DirectoryItem? right)
 		{
@@ -466,7 +452,7 @@ namespace UnrealBuildBase
 		public static DirectoryItem? ReadDirectoryItem(this BinaryArchiveReader Reader)
 		{
 			// Use lambda that doesn't require anything to be captured thus eliminating an allocation.
-			return Reader.ReadObjectReference<DirectoryItem>((BinaryArchiveReader Reader) => DirectoryItem.GetItemByDirectoryReference(Reader.ReadDirectoryReferenceNotNull()));
+			return Reader.ReadObjectReference(Reader => DirectoryItem.GetItemByDirectoryReference(Reader.ReadDirectoryReferenceNotNull()));
 		}
 
 		/// <summary>
@@ -477,7 +463,7 @@ namespace UnrealBuildBase
 		public static void WriteDirectoryItem(this BinaryArchiveWriter Writer, DirectoryItem DirectoryItem)
 		{
 			// Use lambda that doesn't require anything to be captured thus eliminating an allocation.
-			Writer.WriteObjectReference<DirectoryItem>(DirectoryItem, (BinaryArchiveWriter Writer, DirectoryItem DirectoryItem) => Writer.WriteDirectoryReference(DirectoryItem.Location));
+			Writer.WriteObjectReference(DirectoryItem, (Writer, DirectoryItem) => Writer.WriteDirectoryReference(DirectoryItem.Location));
 		}
 
 		/// <summary>
diff --git a/Engine/Source/Programs/Shared/EpicGames.Build/System/FileItem.cs b/Engine/Source/Programs/Shared/EpicGames.Build/System/FileItem.cs
--- a/Engine/Source/Programs/Shared/EpicGames.Build/System/FileItem.cs
+++ b/Engine/Source/Programs/Shared/EpicGames.Build/System/FileItem.cs
@@ -4,6 +4,7 @@ using System;
 using System.Collections.Concurrent;
 using System.Diagnostics;
 using System.IO;
+using System.Linq;
 using System.Threading;
 using EpicGames.Core;
 using Microsoft.Extensions.Logging;
@@ -19,7 +20,7 @@ namespace UnrealBuildBase
 		/// <summary>
 		/// The directory containing this file
 		/// </summary>
-		DirectoryItem? CachedDirectory;
+		Lazy<DirectoryItem> CachedDirectory;
 
 		/// <summary>
 		/// Location of this file
@@ -34,7 +35,7 @@ namespace UnrealBuildBase
 		/// <summary>
 		/// A case-insensitive dictionary that's used to map each unique file name to a single FileItem object.
 		/// </summary>
-		static ConcurrentDictionary<FileReference, FileItem> UniqueSourceFileMap = new ConcurrentDictionary<FileReference, FileItem>();
+		static ConcurrentDictionary<FileReference, FileItem> UniqueSourceFileMap = [];
 
 		/// <summary>
 		/// Constructor
@@ -44,19 +45,16 @@ namespace UnrealBuildBase
 		private FileItem(FileReference Location, FileInfo Info)
 		{
 			this.Location = Location;
-			if (OperatingSystem.IsWindows())
-			{
-				this.Info = new Lazy<FileInfo>(Info);
-			}
-			else
-			{
-				// For some reason we need to call an extra Refresh on linux/mac to not get wrong results from "Exists"
-				this.Info = new Lazy<FileInfo>(() =>
+			CachedDirectory = new(() => DirectoryItem.GetItemByDirectoryReference(Location.Directory));
+
+			// For some reason we need to call an extra Refresh on Linux/Mac to not get wrong results from "Exists"
+			this.Info = OperatingSystem.IsWindows()
+				? new(Info)
+				: new(() =>
 				{
 					Info.Refresh();
 					return Info;
 				});
-			}			
 		}
 
 		/// <summary>
@@ -77,17 +75,7 @@ namespace UnrealBuildBase
 		/// <summary>
 		/// Gets the directory that this file is in
 		/// </summary>
-		public DirectoryItem Directory
-		{
-			get
-			{
-				if (CachedDirectory == null)
-				{
-					CachedDirectory = DirectoryItem.GetItemByDirectoryReference(Location.Directory);
-				}
-				return CachedDirectory;
-			}
-		}
+		public DirectoryItem Directory => CachedDirectory.Value;
 
 		/// <summary>
 		/// Whether the file exists.
@@ -119,19 +107,13 @@ namespace UnrealBuildBase
 		/// </summary>
 		/// <param name="Extension">The extension to check for</param>
 		/// <returns>True if the file has the given extension, false otherwise</returns>
-		public bool HasExtension(string Extension)
-		{
-			return Location.HasExtension(Extension);
-		}
+		public bool HasExtension(string Extension) => Location.HasExtension(Extension);
 
 		/// <summary>
 		/// Gets the directory containing this file
 		/// </summary>
 		/// <returns>DirectoryItem for the directory containing this file</returns>
-		public DirectoryItem GetDirectoryItem()
-		{
-			return Directory;
-		}
+		public DirectoryItem GetDirectoryItem() => Directory;
 
 		/// <summary>
 		/// Updates the cached directory for this file. Used by DirectoryItem when enumerating files, to avoid having to look this up later.
@@ -140,7 +122,7 @@ namespace UnrealBuildBase
 		public void UpdateCachedDirectory(DirectoryItem Directory)
 		{
 			Debug.Assert(Directory.Location == Location.Directory);
-			CachedDirectory = Directory;
+			CachedDirectory = new(Directory);
 		}
 
 		/// <summary>
@@ -148,25 +130,14 @@ namespace UnrealBuildBase
 		/// </summary>
 		/// <param name="FilePath">Path for the FileItem</param>
 		/// <returns>The FileItem that represents the given file path.</returns>
-		public static FileItem GetItemByPath(string FilePath)
-		{
-			return GetItemByFileReference(new FileReference(FilePath));
-		}
+		public static FileItem GetItemByPath(string FilePath) => GetItemByFileReference(new(FilePath));
 
 		/// <summary>
 		/// Gets a FileItem for a given path
 		/// </summary>
 		/// <param name="Info">Information about the file</param>
 		/// <returns>The FileItem that represents the given a full file path.</returns>
-		public static FileItem GetItemByFileInfo(FileInfo Info)
-		{
-			FileReference Location = new FileReference(Info);
-			if (UniqueSourceFileMap.TryGetValue(Location, out FileItem? Result)) // 99.9% reads, so faster to front with a TryGet before GetOrAdd
-			{
-				return Result;
-			}
-			return UniqueSourceFileMap.GetOrAdd(Location, new FileItem(Location, Info));
-		}
+		public static FileItem GetItemByFileInfo(FileInfo Info) => GetItemByFileReference(new(Info));
 
 		/// <summary>
 		/// Gets a FileItem for a given path
@@ -174,13 +145,9 @@ namespace UnrealBuildBase
 		/// <param name="Location">Location of the file</param>
 		/// <returns>The FileItem that represents the given a full file path.</returns>
 		public static FileItem GetItemByFileReference(FileReference Location)
-		{
-			if (UniqueSourceFileMap.TryGetValue(Location, out FileItem? Result)) // 99.9% reads, so faster to front with a TryGet before GetOrAdd
-			{
-				return Result;
-			}
-			return UniqueSourceFileMap.GetOrAdd(Location, new FileItem(Location, Location.ToFileInfo()));
-		}
+			=> UniqueSourceFileMap.TryGetValue(Location, out FileItem? Result)
+				? Result
+				: UniqueSourceFileMap.GetOrAdd(Location, new FileItem(Location, Location.ToFileInfo()));
 
 		/// <summary>
 		/// Deletes the file.
@@ -203,7 +170,7 @@ namespace UnrealBuildBase
 				try
 				{
 					// Delete the destination file if it exists
-					FileInfo DeletedFileInfo = new FileInfo(AbsolutePath);
+					FileInfo DeletedFileInfo = new(AbsolutePath);
 					if (DeletedFileInfo.Exists)
 					{
 						DeletedFileInfo.IsReadOnly = false;
@@ -234,7 +201,7 @@ namespace UnrealBuildBase
 		/// </summary>
 		public void ResetCachedInfo()
 		{
-			Info = new Lazy<FileInfo>(() =>
+			Info = new(() =>
 			{
 				FileInfo Info = Location.ToFileInfo();
 				Info.Refresh();
@@ -247,7 +214,7 @@ namespace UnrealBuildBase
 		/// </summary>
 		public static void ResetCachedInfo(string Path)
 		{
-			if (UniqueSourceFileMap.TryGetValue(new FileReference(Path), out FileItem? Result))
+			if (UniqueSourceFileMap.TryGetValue(new(Path), out FileItem? Result))
 			{
 				Result.ResetCachedInfo();
 			}
@@ -258,32 +225,23 @@ namespace UnrealBuildBase
 		/// </summary>
 		public static void ResetAllCachedInfo_SLOW()
 		{
-			foreach (FileItem Item in UniqueSourceFileMap.Values)
-			{
-				Item.ResetCachedInfo();
-			}
+			UniqueSourceFileMap.Values.AsParallel().ForAll(Item => Item.ResetCachedInfo());
 		}
 
 		/// <summary>
 		/// Return the path to this FileItem to debugging
 		/// </summary>
 		/// <returns>Absolute path to this file item</returns>
-		public override string ToString()
-		{
-			return AbsolutePath;
-		}
+		public override string ToString() => Location.ToString();
 
 		#region IComparable, IEquatbale
-		public int CompareTo(FileItem? other)
-		{
-			return FullName.CompareTo(other?.FullName);
-		}
+		/// <inheritdoc/>
+		public int CompareTo(FileItem? other) => Location.CompareTo(other?.Location);
 
-		public bool Equals(FileItem? other)
-		{
-			return FullName.Equals(other?.FullName);
-		}
+		/// <inheritdoc/>
+		public bool Equals(FileItem? other) => Location.Equals(other?.Location);
 
+		/// <inheritdoc/>
 		public override bool Equals(object? obj)
 		{
 			if (ReferenceEquals(this, obj))
@@ -299,10 +257,8 @@ namespace UnrealBuildBase
 			return Equals(obj as FileItem);
 		}
 
-		public override int GetHashCode()
-		{
-			return FullName.GetHashCode();
-		}
+		/// <inheritdoc/>
+		public override int GetHashCode() => Location.GetHashCode();
 
 		public static bool operator ==(FileItem? left, FileItem? right)
 		{
@@ -353,7 +309,7 @@ namespace UnrealBuildBase
 		/// <returns>Instance of the serialized file item</returns>
 		public static FileItem? ReadFileItem(this BinaryArchiveReader Reader)
 		{
-			return Reader.ReadObjectReference<FileItem>((BinaryArchiveReader Reader) => FileItem.GetItemByFileReference(Reader.ReadFileReference()));
+			return Reader.ReadObjectReference(Reader => FileItem.GetItemByFileReference(Reader.ReadFileReference()));
 		}
 
 		/// <summary>
@@ -363,7 +319,7 @@ namespace UnrealBuildBase
 		/// <param name="FileItem">File item to write</param>
 		public static void WriteFileItem(this BinaryArchiveWriter Writer, FileItem? FileItem)
 		{
-			Writer.WriteObjectReference<FileItem>(FileItem!, () => Writer.WriteFileReference(FileItem!.Location));
+			Writer.WriteObjectReference(FileItem!, () => Writer.WriteFileReference(FileItem!.Location));
 		}
 
 		/// <summary>
@@ -390,7 +346,7 @@ namespace UnrealBuildBase
 		public static FileItem ReadCompactFileItem(this BinaryArchiveReader Reader)
 		{
 			// Use lambda that doesn't require anything to be captured thus eliminating an allocation.
-			return Reader.ReadObjectReference<FileItem>((BinaryArchiveReader Reader) => ReadCompactFileItemData(Reader))!;
+			return Reader.ReadObjectReference(ReadCompactFileItemData)!;
 		}
 
 		/// <summary>
@@ -401,7 +357,7 @@ namespace UnrealBuildBase
 		public static void WriteCompactFileItem(this BinaryArchiveWriter Writer, FileItem FileItem)
 		{
 			// Use lambda that doesn't require anything to be captured thus eliminating an allocation.
-			Writer.WriteObjectReference<FileItem>(FileItem, (BinaryArchiveWriter Writer, FileItem FileItem) =>
+			Writer.WriteObjectReference<FileItem>(FileItem, (Writer, FileItem) =>
 			{
 				Writer.WriteDirectoryItem(FileItem.GetDirectoryItem());
 				Writer.WriteString(FileItem.Name);
