Review change for LocalizationIssueHandler

diff --git a/Engine/Source/Programs/Shared/EpicGames.Horde/Issues/Handlers/LocalizationIssueHandler.cs b/Engine/Source/Programs/Shared/EpicGames.Horde/Issues/Handlers/LocalizationIssueHandler.cs
--- a/Engine/Source/Programs/Shared/EpicGames.Horde/Issues/Handlers/LocalizationIssueHandler.cs
+++ b/Engine/Source/Programs/Shared/EpicGames.Horde/Issues/Handlers/LocalizationIssueHandler.cs
@@ -74,7 +74,9 @@ namespace EpicGames.Horde.Issues.Handlers
 		private IssueEventGroup CreateNewIssueGroup()
 		{
 			string localizationIssueGroup = "Localization";
-			string localizationSummaryTemplate = (!System.String.IsNullOrEmpty(_context.NodeName) ? "{Meta:Node}: " : "") + "Localization {Severity}" + (!System.String.IsNullOrEmpty(_context.StreamId.ToString()) ? " in {Meta:Stream}" : "");
+			string localizationSummaryTemplate = System.String.Format("{0}Localization {{Severity}}{1}",
+				System.String.IsNullOrEmpty(_context.NodeName) ? "" : "{Meta:Node}: ",
+				System.String.IsNullOrEmpty(_context.StreamId.ToString()) ? "" : " in {Meta:Stream}");
 
 			IssueEventGroup newIssueGroup = new IssueEventGroup(localizationIssueGroup, localizationSummaryTemplate, IssueChangeFilter.All);
 
