[UHT]
* Fixed so {} is valid default value in ufunctions for TSubclassOf

diff --git a/Engine/Source/Programs/Shared/EpicGames.UHT/Types/Properties/UhtObjectPropertyBase.cs b/Engine/Source/Programs/Shared/EpicGames.UHT/Types/Properties/UhtObjectPropertyBase.cs
--- a/Engine/Source/Programs/Shared/EpicGames.UHT/Types/Properties/UhtObjectPropertyBase.cs
+++ b/Engine/Source/Programs/Shared/EpicGames.UHT/Types/Properties/UhtObjectPropertyBase.cs
@@ -548,6 +548,12 @@ namespace EpicGames.UHT.Types
 				innerDefaultValue.Append("None");
 				return true;
 			}
+			if (defaultValueReader.TryOptional("{"))
+			{
+				defaultValueReader.Require("}");
+				innerDefaultValue.Append("None");
+				return true;
+			}
 			return false;
 		}
 
