[UBT]
* Added a few aggressive inlining attributes based on profiling from superluminal

diff --git a/Engine/Source/Programs/Shared/EpicGames.Build/System/FileItem.cs b/Engine/Source/Programs/Shared/EpicGames.Build/System/FileItem.cs
--- a/Engine/Source/Programs/Shared/EpicGames.Build/System/FileItem.cs
+++ b/Engine/Source/Programs/Shared/EpicGames.Build/System/FileItem.cs
@@ -5,6 +5,7 @@ using System.Collections.Concurrent;
 using System.Diagnostics;
 using System.IO;
 using System.Linq;
+using System.Runtime.CompilerServices;
 using System.Threading;
 using EpicGames.Core;
 using Microsoft.Extensions.Logging;
@@ -328,6 +329,7 @@ namespace UnrealBuildBase
 		/// </summary>
 		/// <param name="Reader">Archive to read from</param>
 		/// <returns>FileItem read from the archive</returns>
+		[MethodImpl(MethodImplOptions.AggressiveInlining)]
 		static FileItem ReadCompactFileItemData(this BinaryArchiveReader Reader)
 		{
 			DirectoryItem Directory = Reader.ReadDirectoryItem()!;
@@ -343,6 +345,7 @@ namespace UnrealBuildBase
 		/// </summary>
 		/// <param name="Reader">Reader to serialize data from</param>
 		/// <returns>Instance of the serialized file item</returns>
+		[MethodImpl(MethodImplOptions.AggressiveInlining)]
 		public static FileItem ReadCompactFileItem(this BinaryArchiveReader Reader)
 		{
 			// Use lambda that doesn't require anything to be captured thus eliminating an allocation.
diff --git a/Engine/Source/Programs/Shared/EpicGames.Core/BinaryArchiveReader.cs b/Engine/Source/Programs/Shared/EpicGames.Core/BinaryArchiveReader.cs
--- a/Engine/Source/Programs/Shared/EpicGames.Core/BinaryArchiveReader.cs
+++ b/Engine/Source/Programs/Shared/EpicGames.Core/BinaryArchiveReader.cs
@@ -4,6 +4,7 @@ using System;
 using System.Collections.Generic;
 using System.Diagnostics.CodeAnalysis;
 using System.IO;
+using System.Runtime.CompilerServices;
 using System.Text;
 
 namespace EpicGames.Core
@@ -129,6 +130,7 @@ namespace EpicGames.Core
 		/// Reads a single int from the stream
 		/// </summary>
 		/// <returns>The value that was read</returns>
+		[MethodImpl(MethodImplOptions.AggressiveInlining)]
 		public int ReadInt()
 		{
 			return (int)ReadUnsignedInt();
@@ -138,6 +140,7 @@ namespace EpicGames.Core
 		/// Reads a single unsigned int from the stream
 		/// </summary>
 		/// <returns>The value that was read</returns>
+		[MethodImpl(MethodImplOptions.AggressiveInlining)]
 		public uint ReadUnsignedInt()
 		{
 			uint value = (uint)(_buffer![_bufferPos + 0] | (_buffer[_bufferPos + 1] << 8) | (_buffer[_bufferPos + 2] << 16) | (_buffer[_bufferPos + 3] << 24));
@@ -590,6 +593,7 @@ namespace EpicGames.Core
 		/// </summary>
 		/// <param name="readObject">Delegate used to create an object instance. The object may not reference itself recursively.</param>
 		/// <returns>Object instance</returns>
+		[MethodImpl(MethodImplOptions.AggressiveInlining)]
 		public object? ReadUntypedObjectReference(Func<BinaryArchiveReader, object?> readObject)
 		{
 			int index = ReadInt();
@@ -629,6 +633,7 @@ namespace EpicGames.Core
 		/// <typeparam name="T">Type of the object to read.</typeparam>
 		/// <param name="readObject">Delegate used to create an object instance. The object may not reference itself recursively.</param>
 		/// <returns>Object instance</returns>
+		[MethodImpl(MethodImplOptions.AggressiveInlining)]
 		public T? ReadObjectReference<T>(Func<BinaryArchiveReader, T> readObject) where T : class => (T?)ReadUntypedObjectReference(readObject);
 
 		/// <summary>
