[UBT]
* Changed so number of files written is part of "Reflection code generated" log line instead of on its own which has no context
* Removed code that has been obsolete for a long time now
* Fixed crash with uht logging when running in certain conditions

diff --git a/Engine/Source/Programs/Shared/EpicGames.UHT/Utils/UhtFileManager.cs b/Engine/Source/Programs/Shared/EpicGames.UHT/Utils/UhtFileManager.cs
--- a/Engine/Source/Programs/Shared/EpicGames.UHT/Utils/UhtFileManager.cs
+++ b/Engine/Source/Programs/Shared/EpicGames.UHT/Utils/UhtFileManager.cs
@@ -28,14 +28,6 @@ namespace EpicGames.UHT.Utils
 		/// <returns>True if the file was read</returns>
 		public bool ReadSource(string filePath, out UhtSourceFragment fragment);
 
-		/// <summary>
-		/// Read the given source file
-		/// </summary>
-		/// <param name="filePath">File path</param>
-		/// <returns>Buffer containing the read data or null if not found.  The returned buffer must be returned to the cache via a call to UhtBuffer.Return</returns>
-		[Obsolete("Use the new ReadOutput with UhtPoolBuffer")]
-		public UhtBuffer? ReadOutput(string filePath);
-
 		/// <summary>
 		/// Read the given source file
 		/// </summary>
@@ -91,38 +83,6 @@ namespace EpicGames.UHT.Utils
 			return false;
 		}
 
-		/// <inheritdoc/>
-		[Obsolete("Use the new ReadOutput with UhtPoolBuffer")]
-		public UhtBuffer? ReadOutput(string filePath)
-		{
-			// Exceptions are very expensive.  Don't bother trying to open the file if it doesn't exist
-			if (!File.Exists(filePath))
-			{
-				return null;
-			}
-
-			try
-			{
-				using FileStream fs = new(filePath, FileMode.Open, FileAccess.Read, FileShare.Read, 1, FileOptions.SequentialScan);
-				using UhtBorrowByteBuffer byteBuffer = new((int)fs.Length);
-				Span<byte> byteSpan = byteBuffer.Buffer.Memory.Span;
-				int readLength = fs.Read(byteSpan);
-				byteSpan = byteSpan[..readLength];
-
-				Encoding encoding = FileUtils.GetEncoding(byteSpan, out int skipBytes);
-				byteSpan = byteSpan[skipBytes..];
-
-				int charCount = encoding.GetCharCount(byteBuffer.Buffer.Memory.Span);
-				UhtBuffer initialBuffer = UhtBuffer.Borrow(charCount);
-				encoding.GetChars(byteBuffer.Buffer.Memory.Span, initialBuffer.Memory.Span);
-				return initialBuffer;
-			}
-			catch (IOException)
-			{
-				return null;
-			}
-		}
-
 		/// <inheritdoc/>
 		public bool ReadOutput(string filePath, out UhtPoolBuffer<char> output)
 		{
diff --git a/Engine/Source/Programs/Shared/EpicGames.UHT/Utils/UhtSession.cs b/Engine/Source/Programs/Shared/EpicGames.UHT/Utils/UhtSession.cs
--- a/Engine/Source/Programs/Shared/EpicGames.UHT/Utils/UhtSession.cs
+++ b/Engine/Source/Programs/Shared/EpicGames.UHT/Utils/UhtSession.cs
@@ -928,6 +928,11 @@ namespace EpicGames.UHT.Utils
 		/// </summary>
 		public UhtHeaderFile? NoExportTypesHeader { get; private set; } = null;
 
+		/// <summary>
+		/// Number of new/modified files written to disk
+		/// </summary>
+		public int TotalFilesWritten { get; internal set; }
+
 		private readonly List<UhtModule> _modules = [];
 		private readonly List<UhtHeaderFile> _headerFiles = [];
 		private readonly List<UhtHeaderFile> _sortedHeaderFiles = [];
@@ -1346,17 +1351,6 @@ namespace EpicGames.UHT.Utils
 			throw new UhtException($"File not found '{filePath}'");
 		}
 
-		/// <summary>
-		/// Read the given source file
-		/// </summary>
-		/// <param name="filePath">Full or relative file path</param>
-		/// <returns>Buffer containing the read data or null if not found.  The returned buffer must be returned to the cache via a call to UhtBuffer.Return</returns>
-		[Obsolete("Use the ReadSource method with the UhtPoolBuffer<char> output parameter")]
-		public UhtBuffer? ReadSourceToBuffer(string filePath)
-		{
-			return FileManager!.ReadOutput(filePath);
-		}
-
 		/// <summary>
 		/// Read the given source file
 		/// </summary>
@@ -1967,6 +1961,7 @@ namespace EpicGames.UHT.Utils
 		/// <param name="message">Text of the message</param>
 		private void GetMessageParts(UhtMessage msg, out string filePath, out int lineNumber, out string fragmentPath, out string severity, out string message)
 		{
+			filePath = "";
 			fragmentPath = "";
 			lineNumber = msg.LineNumber;
 			if (msg.FilePath != null)
@@ -2000,7 +1995,8 @@ namespace EpicGames.UHT.Utils
 					}
 				}
 			}
-			else
+
+			if (String.IsNullOrEmpty(filePath))
 			{
 				filePath = "UnknownSource";
 			}
@@ -2706,7 +2702,7 @@ namespace EpicGames.UHT.Utils
 		private void StepExport()
 		{
 			HashSet<string> externalDependencies = new();
-			long totalWrittenFiles = 0;
+			int totalWrittenFiles = 0;
 			Try(null, () =>
 			{
 				Logger.LogDebug("Step - Exports");
@@ -2771,7 +2767,7 @@ namespace EpicGames.UHT.Utils
 				}
 			});
 
-			Logger.LogInformation("Total of {NumFiles} written", totalWrittenFiles);
+			TotalFilesWritten = totalWrittenFiles;
 		}
 		#endregion
 
diff --git a/Engine/Source/Programs/UnrealBuildTool/Modes/UnrealHeaderToolMode.cs b/Engine/Source/Programs/UnrealBuildTool/Modes/UnrealHeaderToolMode.cs
--- a/Engine/Source/Programs/UnrealBuildTool/Modes/UnrealHeaderToolMode.cs
+++ b/Engine/Source/Programs/UnrealBuildTool/Modes/UnrealHeaderToolMode.cs
@@ -1270,6 +1270,7 @@ namespace UnrealBuildTool.Modes
 
 				// Read and parse
 				Session.Run(ManifestPath!, Arguments);
+				TotalFilesWritten = Session.TotalFilesWritten;
 				Session.LogMessages();
 				return (int)(Session.HasErrors ? CompilationResult.OtherCompilationError : CompilationResult.Succeeded);
 			}
@@ -1281,5 +1282,7 @@ namespace UnrealBuildTool.Modes
 				return (int)CompilationResult.OtherCompilationError;
 			}
 		}
+
+		public int TotalFilesWritten { get; private set; }
 	}
 }
diff --git a/Engine/Source/Programs/UnrealBuildTool/System/ExternalExecution.cs b/Engine/Source/Programs/UnrealBuildTool/System/ExternalExecution.cs
--- a/Engine/Source/Programs/UnrealBuildTool/System/ExternalExecution.cs
+++ b/Engine/Source/Programs/UnrealBuildTool/System/ExternalExecution.cs
@@ -1327,7 +1327,7 @@ namespace UnrealBuildTool
 						throw new CompilationResultException(UHTResult);
 					}
 
-					Logger.LogInformation("Reflection code generated for {TargetName} in {Time} seconds", ActualTargetName, s.Elapsed.TotalSeconds);
+					Logger.LogInformation("Reflection code generated for {TargetName} in {Time} seconds ({TotalFilesWritten} files written)", ActualTargetName, s.Elapsed.TotalSeconds, UHTTool.TotalFilesWritten);
 
 					// Update the tool info file
 					DirectoryReference.CreateDirectory(ToolInfoFile.Directory);
