Avoid considering only Summary property when summary is too broad

#rnx

diff --git a/Engine/Source/Programs/Shared/EpicGames.Horde/Issues/Handlers/GauntletIssueHandler.cs b/Engine/Source/Programs/Shared/EpicGames.Horde/Issues/Handlers/GauntletIssueHandler.cs
--- a/Engine/Source/Programs/Shared/EpicGames.Horde/Issues/Handlers/GauntletIssueHandler.cs
+++ b/Engine/Source/Programs/Shared/EpicGames.Horde/Issues/Handlers/GauntletIssueHandler.cs
@@ -111,6 +111,13 @@ namespace EpicGames.Horde.Issues.Handlers
 			KnownLogEvents.Gauntlet_UnrealEngineTestEvent
 		};
 
+		/// <summary>
+		///  Known summaries that are not specific enough
+		/// </summary>
+		static readonly string[] s_unspecificSummaries = [
+			"Unhandled Exception: EXCEPTION_ACCESS_VIOLATION reading address"
+		];
+
 		/// <inheritdoc/>
 		public override int Priority => 10;
 
@@ -172,7 +179,13 @@ namespace EpicGames.Horde.Issues.Handlers
 		private static bool TryGetHash(IssueEvent issueEvent, out Md5Hash hash)
 		{
 			// Use only the summary if one is found instead of the full callstack
-			string sanitized = GetSummaryProperty(issueEvent) ?? issueEvent.Render();
+			string? summary = GetSummaryProperty(issueEvent);
+			// Discard summaries that are too broad
+			if (summary != null && s_unspecificSummaries.Any(k => summary.Contains(k, StringComparison.InvariantCultureIgnoreCase)))
+			{
+				summary = null;
+			}
+			string sanitized = summary ?? issueEvent.Render();
 			sanitized = String.Join("\n", sanitized.Split("\n", MaxLines + 1).Take(MaxLines)); // Limit the number of lines to consider
 			sanitized = sanitized.Length > MaxMessageLength ? sanitized.Substring(0, MaxMessageLength) : sanitized;
 			sanitized = sanitized.Trim().ToUpperInvariant();
