[UnrealHeaderTool] fixed CPP_ARRAY_DIM code generation for types in namespace
e.g.,
namespace UE::Foo
{

USTRUCT()
struct FSomeStruct
{
GENERATED_BODY()
};

USTRUCT()
struct FSomeOtherStruct
{
GENERATED_BODY()

static constexpr uint32 MaxNum = 4;

UPROPERTY(Transient)
FSomeStruct Items[MaxNum];
};

} // UE::Foo

diff --git a/Engine/Source/Programs/Shared/EpicGames.UHT/Types/UhtProperty.cs b/Engine/Source/Programs/Shared/EpicGames.UHT/Types/UhtProperty.cs
--- a/Engine/Source/Programs/Shared/EpicGames.UHT/Types/UhtProperty.cs
+++ b/Engine/Source/Programs/Shared/EpicGames.UHT/Types/UhtProperty.cs
@@ -2493,7 +2493,7 @@ namespace EpicGames.UHT.Types
 		{
 			if (property.ArrayDimensions != null)
 			{
-				builder.Append("CPP_ARRAY_DIM(").Append(property.SourceName).Append(", ").Append(context.OuterStruct.SourceName).Append(')');
+				builder.Append("CPP_ARRAY_DIM(").Append(property.SourceName).Append(", ").Append(context.OuterStructNamespaceName).Append(context.OuterStruct.SourceName).Append(')');
 			}
 			else
 			{
