Expose function to convert / in c# scripts for path normalization

diff --git a/Engine/Source/Programs/Shared/EpicGames.Core/FileFilter.cs b/Engine/Source/Programs/Shared/EpicGames.Core/FileFilter.cs
--- a/Engine/Source/Programs/Shared/EpicGames.Core/FileFilter.cs
+++ b/Engine/Source/Programs/Shared/EpicGames.Core/FileFilter.cs
@@ -301,7 +301,7 @@ namespace EpicGames.Core
 		/// </summary>
 		/// <param name="path">Path to process</param>
 		/// <returns>Normalized path</returns>
-		static string NormalizePath(string path)
+		public static string NormalizePath(string path)
 		{
 			Match match = Regex.Match(path, @"^\\\\[^\\/]+\\[^\\/]+");
 			if (match.Success)
