EpicGames.Core: stackalloc buffer when reading small FileReference

#rnx

diff --git a/Engine/Source/Programs/Shared/EpicGames.Core/DirectoryReference.cs b/Engine/Source/Programs/Shared/EpicGames.Core/DirectoryReference.cs
--- a/Engine/Source/Programs/Shared/EpicGames.Core/DirectoryReference.cs
+++ b/Engine/Source/Programs/Shared/EpicGames.Core/DirectoryReference.cs
@@ -114,7 +114,7 @@ namespace EpicGames.Core
 					parentLength = 1;
 				}
 
-				return new DirectoryReference(FullName.Substring(0, parentLength), Sanitize.None);
+				return new(FullName.Substring(0, parentLength), Sanitize.None);
 			}
 		}
 
@@ -136,7 +136,7 @@ namespace EpicGames.Core
 		/// Returns the path of the current user's temporary folder
 		/// </summary>
 		/// <returns>The path to the temporary folder</returns>
-		public static DirectoryReference GetTempPath() => new DirectoryReference(Path.GetTempPath());
+		public static DirectoryReference GetTempPath() => new(Path.GetTempPath());
 
 		/// <summary>
 		/// Gets the path for a special folder
@@ -146,7 +146,7 @@ namespace EpicGames.Core
 		public static DirectoryReference? GetSpecialFolder(Environment.SpecialFolder folder)
 		{
 			string folderPath = Environment.GetFolderPath(folder);
-			return String.IsNullOrEmpty(folderPath) ? null : new DirectoryReference(folderPath);
+			return String.IsNullOrEmpty(folderPath) ? null : new(folderPath);
 		}
 
 		/// <summary>
@@ -161,7 +161,7 @@ namespace EpicGames.Core
 		/// <param name="baseDirectory">The base directory</param>
 		/// <param name="fragments">Fragments to combine with the base directory</param>
 		/// <returns>The new directory name</returns>
-		public static DirectoryReference Combine(DirectoryReference baseDirectory, params string[] fragments) => new DirectoryReference(CombineStrings(baseDirectory, fragments), Sanitize.None);
+		public static DirectoryReference Combine(DirectoryReference baseDirectory, params string[] fragments) => new(CombineStrings(baseDirectory, fragments), Sanitize.None);
 
 		/// <summary>
 		/// Compares two filesystem object names for equality. Uses the canonical name representation, not the display name representation.
@@ -245,7 +245,7 @@ namespace EpicGames.Core
 			
 			foreach (string directory in Directory.GetDirectories(sourceDir.FullName))
 			{
-				Copy(new DirectoryReference(directory), Combine(destinationDir, Path.GetFileName(directory)));
+				Copy(new(directory), Combine(destinationDir, Path.GetFileName(directory)));
 			}
 		}
 		
@@ -274,7 +274,7 @@ namespace EpicGames.Core
 		/// </summary>
 		/// <param name="absolutePath">The absolute path in the remote file system</param>
 		/// <returns>New directory reference</returns>
-		public static DirectoryReference MakeRemote(string absolutePath) => new DirectoryReference(absolutePath, Sanitize.None);
+		public static DirectoryReference MakeRemote(string absolutePath) => new(absolutePath, Sanitize.None);
 
 		/// <summary>
 		/// Gets the parent directory for a file, or returns null if it's null.
@@ -289,14 +289,14 @@ namespace EpicGames.Core
 		/// </summary>
 		/// <param name="directoryName">Path for the new object</param>
 		/// <returns>Returns a FileReference representing the given string, or null.</returns>
-		public static DirectoryReference? FromString(string? directoryName) => String.IsNullOrEmpty(directoryName) ? null : new DirectoryReference(directoryName);
+		public static DirectoryReference? FromString(string? directoryName) => String.IsNullOrEmpty(directoryName) ? null : new(directoryName);
 
 		/// <summary>
 		/// Finds the correct case to match the location of this file on disk. Uses the given case for parts of the path that do not exist.
 		/// </summary>
 		/// <param name="location">The path to find the correct case for</param>
 		/// <returns>Location of the file with the correct case</returns>
-		public static DirectoryReference FindCorrectCase(DirectoryReference location) => new DirectoryReference(DirectoryUtils.FindCorrectCase(location.ToDirectoryInfo()));
+		public static DirectoryReference FindCorrectCase(DirectoryReference location) => new(DirectoryUtils.FindCorrectCase(location.ToDirectoryInfo()));
 
 		/// <summary>
 		/// Constructs a DirectoryInfo object from this reference
@@ -310,7 +310,7 @@ namespace EpicGames.Core
 		/// Finds the current directory
 		/// </summary>
 		/// <returns>The current directory</returns>
-		public static DirectoryReference GetCurrentDirectory() => new DirectoryReference(Directory.GetCurrentDirectory());
+		public static DirectoryReference GetCurrentDirectory() => new(Directory.GetCurrentDirectory());
 
 		/// <summary>
 		/// Creates a directory
@@ -462,7 +462,7 @@ namespace EpicGames.Core
 		public static DirectoryReference? ReadDirectoryReference(this BinaryReader reader)
 		{
 			string fullName = reader.ReadString();
-			return (fullName.Length == 0) ? null : new DirectoryReference(fullName, DirectoryReference.Sanitize.None);
+			return (fullName.Length == 0) ? null : new(fullName, DirectoryReference.Sanitize.None);
 		}
 
 		/// <summary>
@@ -487,7 +487,7 @@ namespace EpicGames.Core
 		public static DirectoryReference? ReadDirectoryReference(this BinaryArchiveReader reader)
 		{
 			string? fullName = reader.ReadString();
-			return fullName == null ? null : new DirectoryReference(fullName, DirectoryReference.Sanitize.None);
+			return fullName == null ? null : new(fullName, DirectoryReference.Sanitize.None);
 		}
 
 		/// <summary>
diff --git a/Engine/Source/Programs/Shared/EpicGames.Core/FileReference.cs b/Engine/Source/Programs/Shared/EpicGames.Core/FileReference.cs
--- a/Engine/Source/Programs/Shared/EpicGames.Core/FileReference.cs
+++ b/Engine/Source/Programs/Shared/EpicGames.Core/FileReference.cs
@@ -70,7 +70,7 @@ namespace EpicGames.Core
 		/// <param name="fileName">FileName for the string</param>
 		/// <returns>Returns a FileReference representing the given string, or null.</returns>
 		[return: NotNullIfNotNull(nameof(fileName))]
-		public static FileReference? FromString(string? fileName) => String.IsNullOrEmpty(fileName) ? null : new FileReference(fileName);
+		public static FileReference? FromString(string? fileName) => String.IsNullOrEmpty(fileName) ? null : new(fileName);
 
 		/// <summary>
 		/// Gets the file name without path information
@@ -114,7 +114,7 @@ namespace EpicGames.Core
 		/// </summary>
 		/// <param name="extension">The new extension</param>
 		/// <returns>A FileReference with the same path and name, but with the new extension</returns>
-		public FileReference ChangeExtension(string? extension) => new FileReference(Path.ChangeExtension(FullName, extension), Sanitize.None);
+		public FileReference ChangeExtension(string? extension) => new(Path.ChangeExtension(FullName, extension), Sanitize.None);
 
 		/// <summary>
 		/// Gets the directory containing this file
@@ -138,7 +138,7 @@ namespace EpicGames.Core
 					parentLength = 1;
 				}
 
-				return new DirectoryReference(FullName.Substring(0, parentLength), DirectoryReference.Sanitize.None);
+				return new(FullName.Substring(0, parentLength), DirectoryReference.Sanitize.None);
 			}
 		}
 
@@ -148,7 +148,7 @@ namespace EpicGames.Core
 		/// <param name="baseDirectory">The base directory</param>
 		/// <param name="fragments">Fragments to combine with the base directory</param>
 		/// <returns>The new file name</returns>
-		public static FileReference Combine(DirectoryReference baseDirectory, params string[] fragments) => new FileReference(CombineStrings(baseDirectory, fragments), Sanitize.None);
+		public static FileReference Combine(DirectoryReference baseDirectory, params string[] fragments) => new(CombineStrings(baseDirectory, fragments), Sanitize.None);
 
 		/// <summary>
 		/// Append a string to the end of a filename
@@ -156,7 +156,7 @@ namespace EpicGames.Core
 		/// <param name="a">The base file reference</param>
 		/// <param name="b">Suffix to be appended</param>
 		/// <returns>The new file reference</returns>
-		public static FileReference operator +(FileReference a, string b) => new FileReference(a.FullName + b, Sanitize.None);
+		public static FileReference operator +(FileReference a, string b) => new(a.FullName + b, Sanitize.None);
 
 		/// <summary>
 		/// Compares two filesystem object names for equality. Uses the canonical name representation, not the display name representation.
@@ -222,7 +222,7 @@ namespace EpicGames.Core
 		/// </summary>
 		/// <param name="absolutePath">The absolute path in the remote file system</param>
 		/// <returns>New file reference</returns>
-		public static FileReference MakeRemote(string absolutePath) => new FileReference(absolutePath, Sanitize.None);
+		public static FileReference MakeRemote(string absolutePath) => new(absolutePath, Sanitize.None);
 
 		/// <summary>
 		/// Makes a file location writeable; 
@@ -261,7 +261,7 @@ namespace EpicGames.Core
 			{
 			}
 			skipBytes = 0;
-			return Encoding.UTF8;
+			return Encoding.Default;
 		}
 
 		/// <summary>
@@ -271,13 +271,13 @@ namespace EpicGames.Core
 		/// <returns>The encoding</returns>
 		public static Encoding GetEncoding(FileReference location) => GetEncoding(location, out int _);
 
-        /// <summary>
-        /// Get the encoding of a file by checking for BOM bytes
-        /// </summary>
-        /// <param name="location">Location of the file</param>
-        /// <param name="cancellationToken">Cancellation token for the operation</param>
-        /// <returns>The encoding</returns>
-        public static async Task<Encoding> GetEncodingAsync(FileReference location, CancellationToken cancellationToken = default)
+		/// <summary>
+		/// Get the encoding of a file by checking for BOM bytes
+		/// </summary>
+		/// <param name="location">Location of the file</param>
+		/// <param name="cancellationToken">Cancellation token for the operation</param>
+		/// <returns>The encoding</returns>
+		public static async Task<Encoding> GetEncodingAsync(FileReference location, CancellationToken cancellationToken = default)
 		{
 			try
 			{
@@ -291,7 +291,7 @@ namespace EpicGames.Core
 			catch (IOException)
 			{
 			}
-			return Encoding.UTF8;
+			return Encoding.Default;
 		}
 
 		/// <summary>
@@ -299,13 +299,13 @@ namespace EpicGames.Core
 		/// </summary>
 		/// <param name="location">The path to find the correct case for</param>
 		/// <returns>Location of the file with the correct case</returns>
-		public static FileReference FindCorrectCase(FileReference location) => new FileReference(FileUtils.FindCorrectCase(location.ToFileInfo()));
+		public static FileReference FindCorrectCase(FileReference location) => new(FileUtils.FindCorrectCase(location.ToFileInfo()));
 
 		/// <summary>
 		/// Constructs a FileInfo object from this reference
 		/// </summary>
 		/// <returns>New FileInfo object</returns>
-		public FileInfo ToFileInfo() => new FileInfo(FullName);
+		public FileInfo ToFileInfo() => new(FullName);
 
 		#region System.IO.File methods
 
@@ -418,29 +418,7 @@ namespace EpicGames.Core
 		/// </summary>
 		/// <param name="location">Location of the file</param>
 		/// <returns>Contents of the file as a single string</returns>
-		public static string ReadAllText(FileReference location)
-		{
-			using (FileStream fs = new FileStream(location.FullName, FileMode.Open, FileAccess.Read, FileShare.Read, 4 * 1024, FileOptions.SequentialScan))
-			{
-				using (StreamReader sr = new StreamReader(fs, Encoding.UTF8, true))
-				{
-					// Try to read the whole file into a buffer created by hand.  This avoids a LOT of memory allocations which in turn reduces the
-					// GC stress on the system.  Removing the StreamReader would be nice in the future.
-					long rawFileLength = fs.Length;
-					char[] initialBuffer = new char[rawFileLength];
-					int readLength = sr.Read(initialBuffer, 0, (int)rawFileLength);
-					if (sr.EndOfStream)
-					{
-						return new string(initialBuffer, 0, readLength);
-					}
-					else
-					{
-						string remaining = sr.ReadToEnd();
-						return String.Concat(new ReadOnlySpan<char>(initialBuffer, 0, readLength), remaining);
-					}
-				}
-			}
-		}
+		public static string ReadAllText(FileReference location) => ReadAllText(location, Encoding.Default);
 
 		/// <summary>
 		/// Reads the contents of a file
@@ -448,7 +426,27 @@ namespace EpicGames.Core
 		/// <param name="location">Location of the file</param>
 		/// <param name="encoding">Encoding of the file</param>
 		/// <returns>Contents of the file as a single string</returns>
-		public static string ReadAllText(FileReference location, Encoding encoding) => File.ReadAllText(location.FullName, encoding);
+		public static string ReadAllText(FileReference location, Encoding encoding)
+		{
+			// Try to read the whole file into a buffer created by hand. This avoids a LOT of memory allocations which in turn reduces the
+			// GC stress on the system. Removing the StreamReader would be nice in the future.
+			using FileStream fs = new(location.FullName, FileMode.Open, FileAccess.Read, FileShare.Read, 4 * 1024, FileOptions.SequentialScan);
+			int fileLength = (int)Math.Clamp(fs.Length, 0, Int32.MaxValue);
+			if (fileLength == 0)
+			{
+				return String.Empty;
+			}
+
+			// stack allocate the buffer for small files 1k or less to prevent gc
+			Span<char> buffer = fileLength <= 1024 ? stackalloc char[fileLength] : new char[fileLength];
+
+			using StreamReader sr = new(fs, encoding);
+			int charsRead = sr.Read(buffer);
+
+			return sr.EndOfStream
+				? buffer[0..charsRead].ToString()
+				: String.Concat(buffer[0..charsRead], sr.ReadToEnd());
+		}
 
 		/// <summary>
 		/// Reads the contents of a file
@@ -456,7 +454,7 @@ namespace EpicGames.Core
 		/// <param name="location">Location of the file</param>
 		/// <param name="cancellationToken">Cancellation token for the operation</param>
 		/// <returns>Contents of the file as a single string</returns>
-		public static Task<string> ReadAllTextAsync(FileReference location, CancellationToken cancellationToken = default) => File.ReadAllTextAsync(location.FullName, cancellationToken);
+		public static Task<string> ReadAllTextAsync(FileReference location, CancellationToken cancellationToken = default) => ReadAllTextAsync(location, Encoding.Default, cancellationToken);
 
 		/// <summary>
 		/// Reads the contents of a file
@@ -472,7 +470,7 @@ namespace EpicGames.Core
 		/// </summary>
 		/// <param name="location">Location of the file</param>
 		/// <returns>String array containing the contents of the file</returns>
-		public static string[] ReadAllLines(FileReference location) => File.ReadAllLines(location.FullName);
+		public static string[] ReadAllLines(FileReference location) => ReadAllLines(location, Encoding.Default);
 
 		/// <summary>
 		/// Reads the contents of a file
@@ -488,7 +486,7 @@ namespace EpicGames.Core
 		/// <param name="location">Location of the file</param>
 		/// <param name="cancellationToken">Cancellation token for the operation</param>
 		/// <returns>String array containing the contents of the file</returns>
-		public static Task<string[]> ReadAllLinesAsync(FileReference location, CancellationToken cancellationToken = default) => File.ReadAllLinesAsync(location.FullName, cancellationToken);
+		public static Task<string[]> ReadAllLinesAsync(FileReference location, CancellationToken cancellationToken = default) => ReadAllLinesAsync(location, Encoding.Default, cancellationToken);
 
 		/// <summary>
 		/// Reads the contents of a file
@@ -559,7 +557,7 @@ namespace EpicGames.Core
 		{
 			if (Exists(location))
 			{
-				byte[] currentContents = ReadAllBytes(location);
+				ReadOnlySpan<byte> currentContents = ReadAllBytes(location);
 				if (contents.AsSpan().SequenceEqual(currentContents))
 				{
 					return false;
@@ -636,7 +634,7 @@ namespace EpicGames.Core
 		/// </summary>
 		/// <param name="location">Location of the file</param>
 		/// <param name="contents">Contents of the file</param>
-		public static void WriteAllLines(FileReference location, IEnumerable<string> contents) => File.WriteAllLines(location.FullName, contents);
+		public static void WriteAllLines(FileReference location, IEnumerable<string> contents) => WriteAllLines(location, contents, Encoding.Default);
 
 		/// <summary>
 		/// Writes the contents of a file
@@ -651,7 +649,7 @@ namespace EpicGames.Core
 		/// </summary>
 		/// <param name="location">Location of the file</param>
 		/// <param name="contents">Contents of the file</param>
-		public static void WriteAllLines(FileReference location, string[] contents) => File.WriteAllLines(location.FullName, contents);
+		public static void WriteAllLines(FileReference location, string[] contents) => WriteAllLines(location, contents, Encoding.Default);
 
 		/// <summary>
 		/// Writes the contents of a file
@@ -667,7 +665,7 @@ namespace EpicGames.Core
 		/// <param name="location">Location of the file</param>
 		/// <param name="contents">Contents of the file</param>
 		/// <param name="cancellationToken">Cancellation token for the operation</param>
-		public static Task WriteAllLinesAsync(FileReference location, IEnumerable<string> contents, CancellationToken cancellationToken = default) => File.WriteAllLinesAsync(location.FullName, contents, cancellationToken);
+		public static Task WriteAllLinesAsync(FileReference location, IEnumerable<string> contents, CancellationToken cancellationToken = default) => WriteAllLinesAsync(location, contents, Encoding.Default, cancellationToken);
 
 		/// <summary>
 		/// Writes the contents of a file
@@ -684,7 +682,7 @@ namespace EpicGames.Core
 		/// <param name="location">Location of the file</param>
 		/// <param name="contents">Contents of the file</param>
 		/// <param name="cancellationToken">Cancellation token for the operation</param>
-		public static Task WriteAllLinesAsync(FileReference location, string[] contents, CancellationToken cancellationToken = default) => File.WriteAllLinesAsync(location.FullName, contents, cancellationToken);
+		public static Task WriteAllLinesAsync(FileReference location, string[] contents, CancellationToken cancellationToken = default) => WriteAllLinesAsync(location, contents, Encoding.Default, cancellationToken);
 
 		/// <summary>
 		/// Writes the contents of a file
@@ -700,7 +698,7 @@ namespace EpicGames.Core
 		/// </summary>
 		/// <param name="location">Location of the file</param>
 		/// <param name="contents">Contents of the file</param>
-		public static void WriteAllText(FileReference location, string contents) => File.WriteAllText(location.FullName, contents);
+		public static void WriteAllText(FileReference location, string contents) => WriteAllText(location, contents, Encoding.Default);
 
 		/// <summary>
 		/// Writes the contents of a file
@@ -716,7 +714,7 @@ namespace EpicGames.Core
 		/// <param name="location">Location of the file</param>
 		/// <param name="contents">Contents of the file</param>
 		/// <param name="cancellationToken">Cancellation token for the operation</param>
-		public static Task WriteAllTextAsync(FileReference location, string contents, CancellationToken cancellationToken = default) => File.WriteAllTextAsync(location.FullName, contents, cancellationToken);
+		public static Task WriteAllTextAsync(FileReference location, string contents, CancellationToken cancellationToken = default) => WriteAllTextAsync(location, contents, Encoding.Default, cancellationToken);
 
 		/// <summary>
 		/// Writes the contents of a file
@@ -732,23 +730,23 @@ namespace EpicGames.Core
 		/// </summary>
 		/// <param name="location">Location of the file</param>
 		/// <param name="contents">Contents to append to the file</param>
-		public static void AppendAllLines(FileReference location, IEnumerable<string> contents) => File.AppendAllLines(location.FullName, contents);
+		public static void AppendAllLines(FileReference location, IEnumerable<string> contents) => AppendAllLines(location, contents, Encoding.Default);
 
 		/// <summary>
 		/// Appends the contents to a file
 		/// </summary>
 		/// <param name="location">Location of the file</param>
 		/// <param name="contents">Contents to append to the file</param>
-		/// <param name="cancellationToken">Cancellation token for the operation</param>
-		public static Task AppendAllLinesAsync(FileReference location, IEnumerable<string> contents, CancellationToken cancellationToken = default) => File.AppendAllLinesAsync(location.FullName, contents, cancellationToken);
+		/// <param name="encoding">The encoding to use when parsing the file</param>
+		public static void AppendAllLines(FileReference location, IEnumerable<string> contents, Encoding encoding) => File.AppendAllLines(location.FullName, contents, encoding);
 
 		/// <summary>
 		/// Appends the contents to a file
 		/// </summary>
 		/// <param name="location">Location of the file</param>
 		/// <param name="contents">Contents to append to the file</param>
-		/// <param name="encoding">The encoding to use when parsing the file</param>
-		public static void AppendAllLines(FileReference location, IEnumerable<string> contents, Encoding encoding) => File.AppendAllLines(location.FullName, contents, encoding);
+		/// <param name="cancellationToken">Cancellation token for the operation</param>
+		public static Task AppendAllLinesAsync(FileReference location, IEnumerable<string> contents, CancellationToken cancellationToken = default) => AppendAllLinesAsync(location, contents, Encoding.Default, cancellationToken);
 
 		/// <summary>
 		/// Appends the contents to a file
@@ -764,23 +762,23 @@ namespace EpicGames.Core
 		/// </summary>
 		/// <param name="location">Location of the file</param>
 		/// <param name="contents">Contents to append to the file</param>
-		public static void AppendAllLines(FileReference location, string[] contents) => File.AppendAllLines(location.FullName, contents);
+		public static void AppendAllLines(FileReference location, string[] contents) => AppendAllLines(location, contents, Encoding.Default);
 
 		/// <summary>
 		/// Appends the contents to a file
 		/// </summary>
 		/// <param name="location">Location of the file</param>
 		/// <param name="contents">Contents to append to the file</param>
-		/// <param name="cancellationToken">Cancellation token for the operation</param>
-		public static Task AppendAllLinesAsync(FileReference location, string[] contents, CancellationToken cancellationToken = default) => File.AppendAllLinesAsync(location.FullName, contents, cancellationToken);
+		/// <param name="encoding">The encoding to use when parsing the file</param>
+		public static void AppendAllLines(FileReference location, string[] contents, Encoding encoding) => File.AppendAllLines(location.FullName, contents, encoding);
 
 		/// <summary>
 		/// Appends the contents to a file
 		/// </summary>
 		/// <param name="location">Location of the file</param>
 		/// <param name="contents">Contents to append to the file</param>
-		/// <param name="encoding">The encoding to use when parsing the file</param>
-		public static void AppendAllLines(FileReference location, string[] contents, Encoding encoding) => File.AppendAllLines(location.FullName, contents, encoding);
+		/// <param name="cancellationToken">Cancellation token for the operation</param>
+		public static Task AppendAllLinesAsync(FileReference location, string[] contents, CancellationToken cancellationToken = default) => AppendAllLinesAsync(location, contents, Encoding.Default, cancellationToken);
 
 		/// <summary>
 		/// Appends the contents to a file
@@ -796,23 +794,23 @@ namespace EpicGames.Core
 		/// </summary>
 		/// <param name="location">Location of the file</param>
 		/// <param name="contents">Contents to append to the file</param>
-		public static void AppendAllText(FileReference location, string contents) => File.AppendAllText(location.FullName, contents);
+		public static void AppendAllText(FileReference location, string contents) => AppendAllText(location, contents, Encoding.Default);
 
 		/// <summary>
 		/// Appends the contents to a file
 		/// </summary>
 		/// <param name="location">Location of the file</param>
 		/// <param name="contents">Contents to append to the file</param>
-		/// <param name="cancellationToken">Cancellation token for the operation</param>
-		public static Task AppendAllTextAsync(FileReference location, string contents, CancellationToken cancellationToken = default) => File.AppendAllTextAsync(location.FullName, contents, cancellationToken);
+		/// <param name="encoding">The encoding to use when parsing the file</param>
+		public static void AppendAllText(FileReference location, string contents, Encoding encoding) => File.AppendAllText(location.FullName, contents, encoding);
 
 		/// <summary>
 		/// Appends the contents to a file
 		/// </summary>
 		/// <param name="location">Location of the file</param>
 		/// <param name="contents">Contents to append to the file</param>
-		/// <param name="encoding">The encoding to use when parsing the file</param>
-		public static void AppendAllText(FileReference location, string contents, Encoding encoding) => File.AppendAllText(location.FullName, contents, encoding);
+		/// <param name="cancellationToken">Cancellation token for the operation</param>
+		public static Task AppendAllTextAsync(FileReference location, string contents, CancellationToken cancellationToken = default) => AppendAllTextAsync(location, contents, Encoding.Default, cancellationToken);
 
 		/// <summary>
 		/// Appends the contents to a file
@@ -916,7 +914,7 @@ namespace EpicGames.Core
 		public static FileReference? ReadFileReferenceOrNull(this BinaryReader reader)
 		{
 			string fullName = reader.ReadString();
-			return (fullName.Length == 0) ? null : new FileReference(fullName, FileReference.Sanitize.None);
+			return (fullName.Length == 0) ? null : new(fullName, FileReference.Sanitize.None);
 		}
 
 		/// <summary>
@@ -984,7 +982,7 @@ namespace EpicGames.Core
 		public static FileReference? ReadFileReferenceOrNull(this BinaryArchiveReader reader)
 		{
 			string? fullName = reader.ReadString();
-			return fullName == null ? null : new FileReference(fullName, FileReference.Sanitize.None);
+			return fullName == null ? null : new(fullName, FileReference.Sanitize.None);
 		}
 	}
 }
