[UBA]
* Tweaked settings for ubahost

diff --git a/Engine/Source/Programs/Shared/EpicGames.UBA/Impl/SchedulerImpl.cs b/Engine/Source/Programs/Shared/EpicGames.UBA/Impl/SchedulerImpl.cs
--- a/Engine/Source/Programs/Shared/EpicGames.UBA/Impl/SchedulerImpl.cs
+++ b/Engine/Source/Programs/Shared/EpicGames.UBA/Impl/SchedulerImpl.cs
@@ -133,7 +133,6 @@ namespace EpicGames.UBA.Impl
 
 			using IConfig config = IConfig.LoadConfig("");
 			config.AddValue("Scheduler", "MaxLocalProcessors", maxLocalProcessors.ToString());
-			config.AddValue("Scheduler", "MaxParallelCacheQueries", "100");
 			if (forceRemote)
 			{
 				config.AddValue("Scheduler", "ForceRemote", "True");
