[EpicGames.UBA]
* Added c# api for adding tasks to uba trace

diff --git a/Engine/Source/Programs/Shared/EpicGames.UBA/ISessionServer.cs b/Engine/Source/Programs/Shared/EpicGames.UBA/ISessionServer.cs
--- a/Engine/Source/Programs/Shared/EpicGames.UBA/ISessionServer.cs
+++ b/Engine/Source/Programs/Shared/EpicGames.UBA/ISessionServer.cs
@@ -257,6 +257,27 @@ namespace EpicGames.UBA
 		/// </summary>
 		public abstract void PrintSummary();
 
+		/// <summary>
+		/// Begin task. Will create a bar in uba visualizer
+		/// </summary>
+		/// <param name="description">Description showing on bar in visualizer</param>
+		/// <param name="details">Details showing when hovering over bar in visualizer</param>
+		/// <returns>Task id</returns>
+		public abstract uint TaskBegin(string description, string details);
+
+		/// <summary>
+		/// Add hint to task.. will show with milliseconds since start or previous hint
+		/// </summary>
+		/// <param name="taskId">Task id returned by TaskBegin</param>
+		/// <param name="hint">Hint text</param>
+		public abstract void TaskHint(uint taskId, string hint);
+
+		/// <summary>
+		/// End task.
+		/// </summary>
+		/// <param name="taskId">Task id returned by TaskBegin</param>
+		public abstract void TaskEnd(uint taskId);
+
 		/// <summary>
 		/// Create a ISessionServer object
 		/// </summary>
@@ -267,4 +288,33 @@ namespace EpicGames.UBA
 			return new SessionServerImpl(info);
 		}
 	}
+
+	/// <summary>
+	/// Helper type to track task in scope
+	/// </summary>
+	public readonly struct SessionTaskScope : IDisposable
+	{
+		private readonly ISessionServer _server;
+		private readonly uint _taskId;
+
+		/// <summary>
+		/// Will create a bar in uba visualizer
+		/// </summary>
+		/// <param name="server">Session used for this task scope</param>
+		/// <param name="description">Description showing on bar in visualizer</param>
+		/// <param name="details">Details showing when hovering over bar in visualizer</param>
+		public SessionTaskScope(ISessionServer server, string description, string details)
+		{
+			_server = server;
+			_taskId = server.TaskBegin(description, details);
+		}
+
+		/// <summary>
+		/// End task
+		/// </summary>
+		public void Dispose()
+		{
+			_server.TaskEnd(_taskId);
+		}
+	}
 }
diff --git a/Engine/Source/Programs/Shared/EpicGames.UBA/Impl/SessionServerImpl.cs b/Engine/Source/Programs/Shared/EpicGames.UBA/Impl/SessionServerImpl.cs
--- a/Engine/Source/Programs/Shared/EpicGames.UBA/Impl/SessionServerImpl.cs
+++ b/Engine/Source/Programs/Shared/EpicGames.UBA/Impl/SessionServerImpl.cs
@@ -130,6 +130,15 @@ namespace EpicGames.UBA.Impl
 		[DllImport("UbaHost", CharSet = CharSet.Auto)]
 		static extern void SessionServer_CancelAll(nint server);
 
+		[DllImport("UbaHost", CharSet = CharSet.Auto)]
+		static extern uint SessionServer_TaskBegin(nint server, string description, string details);
+
+		[DllImport("UbaHost", CharSet = CharSet.Auto)]
+		static extern void SessionServer_TaskHint(nint server, uint taskId, string hint);
+
+		[DllImport("UbaHost", CharSet = CharSet.Auto)]
+		static extern void SessionServer_TaskEnd(nint server, uint taskId, bool success);
+
 		[DllImport("UbaHost", CharSet = CharSet.Auto)]
 		static extern void SessionServer_Destroy(nint server);
 		#endregion
@@ -236,6 +245,13 @@ namespace EpicGames.UBA.Impl
 		public void PrintSummary() => SessionServer_PrintSummary(_handle);
 
 		public void CancelAll() => SessionServer_CancelAll(_handle);
+
+		public uint TaskBegin(string description, string details) => SessionServer_TaskBegin(_handle, description, details);
+
+		public void TaskHint(uint taskId, string hint) => SessionServer_TaskHint(_handle, taskId, hint);
+
+		public void TaskEnd(uint taskId) => SessionServer_TaskEnd(_handle, taskId, true);
+
 		#endregion
 
 		void RaiseRemoteProcessSlotAvailable(nint userData, byte isCrossArchitecture)
