Fix issue where if a TInstancedStruct is used as a function argument, Clang platforms would fail to compile.

diff --git a/Engine/Source/Programs/Shared/EpicGames.UHT/Types/Properties/UhtStructProperty.cs b/Engine/Source/Programs/Shared/EpicGames.UHT/Types/Properties/UhtStructProperty.cs
--- a/Engine/Source/Programs/Shared/EpicGames.UHT/Types/Properties/UhtStructProperty.cs
+++ b/Engine/Source/Programs/Shared/EpicGames.UHT/Types/Properties/UhtStructProperty.cs
@@ -524,10 +524,9 @@ namespace EpicGames.UHT.Types
 			base.CollectReferencesInternal(collector, addForwardDeclarations, isTemplateProperty);
 			if (addForwardDeclarations && !TemplateArgumentStruct.IsCoreType)
 			{
-				if (RootProperty.Outer is UhtField outerField)
+				if (RootProperty.Outer is UhtField)
 				{
 					collector.AddForwardDeclaration(TemplateArgumentStruct);
-					collector.AddForwardDeclaration(outerField.Namespace, $"template struct {TemplateWrapperName}<{TemplateArgumentStruct.SourceName}>;");
 				}
 			}
 		}
