Fixed UHT parser issue where if FArchive is forward declared when used in the Serialize function declaration, UHT would fail to generate the proper code.

diff --git a/Engine/Source/Programs/Shared/EpicGames.UHT/Parsers/UhtHeaderFileParser.cs b/Engine/Source/Programs/Shared/EpicGames.UHT/Parsers/UhtHeaderFileParser.cs
--- a/Engine/Source/Programs/Shared/EpicGames.UHT/Parsers/UhtHeaderFileParser.cs
+++ b/Engine/Source/Programs/Shared/EpicGames.UHT/Parsers/UhtHeaderFileParser.cs
@@ -1413,6 +1413,8 @@ namespace EpicGames.UHT.Parsers
 				return false;
 			}
 
+			replayReader.Optional("class");
+
 			UhtToken token = replayReader.GetToken();
 
 			UhtSerializerArchiveType archiveType = UhtSerializerArchiveType.None;
