Minor code style fixes to UHT

#rnx

diff --git a/Engine/Source/Programs/Shared/EpicGames.UHT/Exporters/CodeGen/UhtCodeGenerator.cs b/Engine/Source/Programs/Shared/EpicGames.UHT/Exporters/CodeGen/UhtCodeGenerator.cs
--- a/Engine/Source/Programs/Shared/EpicGames.UHT/Exporters/CodeGen/UhtCodeGenerator.cs
+++ b/Engine/Source/Programs/Shared/EpicGames.UHT/Exporters/CodeGen/UhtCodeGenerator.cs
@@ -3,14 +3,12 @@
 using System;
 using System.Collections.Generic;
 using System.IO;
-using System.Linq;
 using System.Text;
 using System.Threading.Tasks;
 using EpicGames.Core;
 using EpicGames.UHT.Tables;
 using EpicGames.UHT.Types;
 using EpicGames.UHT.Utils;
-using Microsoft.Extensions.Logging;
 
 namespace EpicGames.UHT.Exporters.CodeGen
 {
diff --git a/Engine/Source/Programs/Shared/EpicGames.UHT/Exporters/CodeGen/UhtHeaderCodeGeneratorCppFile.cs b/Engine/Source/Programs/Shared/EpicGames.UHT/Exporters/CodeGen/UhtHeaderCodeGeneratorCppFile.cs
--- a/Engine/Source/Programs/Shared/EpicGames.UHT/Exporters/CodeGen/UhtHeaderCodeGeneratorCppFile.cs
+++ b/Engine/Source/Programs/Shared/EpicGames.UHT/Exporters/CodeGen/UhtHeaderCodeGeneratorCppFile.cs
@@ -1366,11 +1366,8 @@ namespace EpicGames.UHT.Exporters.CodeGen
 		{
 			UhtProperty? returnProperty = function.ReturnProperty;
 
-			bool isDecides = function.FunctionExportFlags.HasAnyFlags(UhtFunctionExportFlags.VerseDecides);
 			bool isCoroutine = function.FunctionExportFlags.HasAnyFlags(UhtFunctionExportFlags.VerseSuspends);
 			bool isNoRollback = function.FunctionExportFlags.HasAnyFlags(UhtFunctionExportFlags.VerseNoRollback);
-			bool isIncompleteReturn = function.Session.IsIncompleteReturn(returnProperty);
-			bool isVoidReturn = returnProperty == null;
 
 			// Write _exec_ wrapper that checks for native method implementation
 			if (isNoRollback)
diff --git a/Engine/Source/Programs/Shared/EpicGames.UHT/Types/Properties/UhtLazyObjectPtrProperty.cs b/Engine/Source/Programs/Shared/EpicGames.UHT/Types/Properties/UhtLazyObjectPtrProperty.cs
--- a/Engine/Source/Programs/Shared/EpicGames.UHT/Types/Properties/UhtLazyObjectPtrProperty.cs
+++ b/Engine/Source/Programs/Shared/EpicGames.UHT/Types/Properties/UhtLazyObjectPtrProperty.cs
@@ -4,7 +4,6 @@ using System.Diagnostics.CodeAnalysis;
 using System.Text;
 using EpicGames.Core;
 using EpicGames.UHT.Tables;
-using EpicGames.UHT.Tokenizer;
 using EpicGames.UHT.Utils;
 
 namespace EpicGames.UHT.Types
