PR #12441: Fix invalid access to memory location error in UnrealBuildTool

diff --git a/Engine/Source/Programs/Shared/EpicGames.Core/ManagedProcess.cs b/Engine/Source/Programs/Shared/EpicGames.Core/ManagedProcess.cs
--- a/Engine/Source/Programs/Shared/EpicGames.Core/ManagedProcess.cs
+++ b/Engine/Source/Programs/Shared/EpicGames.Core/ManagedProcess.cs
@@ -102,7 +102,7 @@ namespace EpicGames.Core
 		static extern int SetInformationJobObject(SafeFileHandle hJob, int jobObjectInfoClass, IntPtr lpJobObjectInfo, int cbJobObjectInfoLength);
 
 		[DllImport("kernel32.dll", SetLastError = true)]
-		static extern bool QueryInformationJobObject(SafeFileHandle hJob, int JobObjectInformationClass, ref JOBOBJECT_BASIC_ACCOUNTING_INFORMATION lpJobObjectInformation, int cbJobObjectInformationLength);
+		static extern bool QueryInformationJobObject(SafeFileHandle hJob, int JobObjectInformationClass, out JOBOBJECT_BASIC_ACCOUNTING_INFORMATION lpJobObjectInformation, int cbJobObjectInformationLength, out int lpReturnLength);
 
 		[DllImport("kernel32.dll", SetLastError = true)]
 		static extern int AssignProcessToJobObject(SafeFileHandle? hJob, IntPtr hProcess);
@@ -200,7 +200,7 @@ namespace EpicGames.Core
 				if (SupportsJobObjects)
 				{
 					JOBOBJECT_BASIC_ACCOUNTING_INFORMATION accountingInformation = new JOBOBJECT_BASIC_ACCOUNTING_INFORMATION();
-					if (QueryInformationJobObject(JobHandle!, JobObjectBasicAccountingInformation, ref accountingInformation, Marshal.SizeOf(typeof(JOBOBJECT_BASIC_ACCOUNTING_INFORMATION))))
+					if (QueryInformationJobObject(JobHandle!, JobObjectBasicAccountingInformation, out accountingInformation, Marshal.SizeOf(typeof(JOBOBJECT_BASIC_ACCOUNTING_INFORMATION)), out _))
 					{
 						return new TimeSpan((long)accountingInformation.TotalUserTime + (long)accountingInformation.TotalKernelTime);
 					}
