EpicGames.Core: Fix analysis warning

#rnx

diff --git a/Engine/Source/Programs/Shared/EpicGames.Core/ManagedProcess.cs b/Engine/Source/Programs/Shared/EpicGames.Core/ManagedProcess.cs
--- a/Engine/Source/Programs/Shared/EpicGames.Core/ManagedProcess.cs
+++ b/Engine/Source/Programs/Shared/EpicGames.Core/ManagedProcess.cs
@@ -199,8 +199,7 @@ namespace EpicGames.Core
 			{
 				if (SupportsJobObjects)
 				{
-					JOBOBJECT_BASIC_ACCOUNTING_INFORMATION accountingInformation = new JOBOBJECT_BASIC_ACCOUNTING_INFORMATION();
-					if (QueryInformationJobObject(JobHandle!, JobObjectBasicAccountingInformation, out accountingInformation, Marshal.SizeOf(typeof(JOBOBJECT_BASIC_ACCOUNTING_INFORMATION)), out _))
+					if (QueryInformationJobObject(JobHandle!, JobObjectBasicAccountingInformation, out JOBOBJECT_BASIC_ACCOUNTING_INFORMATION accountingInformation, Marshal.SizeOf(typeof(JOBOBJECT_BASIC_ACCOUNTING_INFORMATION)), out _))
 					{
 						return new TimeSpan((long)accountingInformation.TotalUserTime + (long)accountingInformation.TotalKernelTime);
 					}
