Adding support for UHT generating the VNI package structure.

#rnx

diff --git a/Engine/Source/Programs/Shared/EpicGames.UHT/Exporters/CodeGen/UhtPackageCodeGeneratorCppFile.cs b/Engine/Source/Programs/Shared/EpicGames.UHT/Exporters/CodeGen/UhtPackageCodeGeneratorCppFile.cs
--- a/Engine/Source/Programs/Shared/EpicGames.UHT/Exporters/CodeGen/UhtPackageCodeGeneratorCppFile.cs
+++ b/Engine/Source/Programs/Shared/EpicGames.UHT/Exporters/CodeGen/UhtPackageCodeGeneratorCppFile.cs
@@ -89,6 +89,10 @@ namespace EpicGames.UHT.Exporters.CodeGen
 
 				builder.Append(HeaderCopyright);
 				builder.Append(RequiredCPPIncludes);
+				if (Module.Module.HasVerse)
+				{
+					builder.Append("#include \"VerseInteropMacros.h\"\r\n");
+				}
 				builder.Append(DisableDeprecationWarnings).Append("\r\n");
 				builder.Append("void EmptyLinkFunctionForGeneratedCode").Append(Module.ShortName).Append("_init() {}\r\n");
 
@@ -112,44 +116,74 @@ namespace EpicGames.UHT.Exporters.CodeGen
 				{
 					string strippedName = PackageInfos[package.PackageTypeIndex].StrippedName;
 					string singletonName = GetSingletonName(package, true);
-					builder.Append("\tstatic FPackageRegistrationInfo Z_Registration_Info_UPackage_").Append(strippedName).Append(";\r\n");
-					builder.Append("\tFORCENOINLINE UPackage* ").Append(singletonName).Append("()\r\n");
+					builder.Append("static FPackageRegistrationInfo Z_Registration_Info_UPackage_").Append(strippedName).Append(";\r\n");
+					builder.Append("FORCENOINLINE UPackage* ").Append(singletonName).Append("()\r\n");
+					builder.Append("{\r\n");
+					builder.Append("\tif (!Z_Registration_Info_UPackage_").Append(strippedName).Append(".OuterSingleton)\r\n");
 					builder.Append("\t{\r\n");
-					builder.Append("\t\tif (!Z_Registration_Info_UPackage_").Append(strippedName).Append(".OuterSingleton)\r\n");
-					builder.Append("\t\t{\r\n");
 
 					UhtField[] packageSingletons = singletons.Where(x => x.Package == package).ToArray();
 
 					if (packageSingletons.Length != 0)
 					{
-						builder.Append("\t\t\tstatic UObject* (*const SingletonFuncArray[])() = {\r\n");
+						builder.Append("\t\tstatic UObject* (*const SingletonFuncArray[])() = {\r\n");
 						foreach (UhtField field in packageSingletons)
 						{
-							builder.Append("\t\t\t\t(UObject* (*)())").Append(ObjectInfos[field.ObjectTypeIndex].RegisteredSingletonName).Append(",\r\n");
+							builder.Append("\t\t\t(UObject* (*)())").Append(ObjectInfos[field.ObjectTypeIndex].RegisteredSingletonName).Append(",\r\n");
 						}
-						builder.Append("\t\t\t};\r\n");
+						builder.Append("\t\t};\r\n");
 					}
 
 					EPackageFlags flags = package.PackageFlags & (EPackageFlags.ClientOptional | EPackageFlags.ServerSideOnly | EPackageFlags.EditorOnly | EPackageFlags.Developer | EPackageFlags.UncookedOnly);
-					builder.Append("\t\t\tstatic const UECodeGen_Private::FPackageParams PackageParams = {\r\n");
-					builder.Append("\t\t\t\t").AppendUTF8LiteralString(package.SourceName).Append(",\r\n");
-					builder.Append("\t\t\t\t").Append(packageSingletons.Length != 0 ? "SingletonFuncArray" : "nullptr").Append(",\r\n");
-					builder.Append("\t\t\t\t").Append(packageSingletons.Length != 0 ? "UE_ARRAY_COUNT(SingletonFuncArray)" : "0").Append(",\r\n");
-					builder.Append("\t\t\t\tPKG_CompiledIn | ").Append($"0x{(uint)flags:X8}").Append(",\r\n");
-					builder.Append("\t\t\t\t").Append($"0x{bodiesHash:X8}").Append(",\r\n");
-					builder.Append("\t\t\t\t").Append($"0x{declarationsHash:X8}").Append(",\r\n");
-					builder.Append("\t\t\t\t").AppendMetaDataParams(package, null, MetaDataParamsName).Append("\r\n");
-					builder.Append("\t\t\t};\r\n");
-					builder.Append("\t\t\tUECodeGen_Private::ConstructUPackage(Z_Registration_Info_UPackage_").Append(strippedName).Append(".OuterSingleton, PackageParams);\r\n");
-					builder.Append("\t\t}\r\n");
-					builder.Append("\t\treturn Z_Registration_Info_UPackage_").Append(strippedName).Append(".OuterSingleton;\r\n");
+					builder.Append("\t\tstatic const UECodeGen_Private::FPackageParams PackageParams = {\r\n");
+					builder.Append("\t\t\t").AppendUTF8LiteralString(package.SourceName).Append(",\r\n");
+					builder.Append("\t\t\t").Append(packageSingletons.Length != 0 ? "SingletonFuncArray" : "nullptr").Append(",\r\n");
+					builder.Append("\t\t\t").Append(packageSingletons.Length != 0 ? "UE_ARRAY_COUNT(SingletonFuncArray)" : "0").Append(",\r\n");
+					builder.Append("\t\t\tPKG_CompiledIn | ").Append($"0x{(uint)flags:X8}").Append(",\r\n");
+					builder.Append("\t\t\t").Append($"0x{bodiesHash:X8}").Append(",\r\n");
+					builder.Append("\t\t\t").Append($"0x{declarationsHash:X8}").Append(",\r\n");
+					builder.Append("\t\t\t").AppendMetaDataParams(package, null, MetaDataParamsName).Append("\r\n");
+					builder.Append("\t\t};\r\n");
+					builder.Append("\t\tUECodeGen_Private::ConstructUPackage(Z_Registration_Info_UPackage_").Append(strippedName).Append(".OuterSingleton, PackageParams);\r\n");
 					builder.Append("\t}\r\n");
+					builder.Append("\treturn Z_Registration_Info_UPackage_").Append(strippedName).Append(".OuterSingleton;\r\n");
+					builder.Append("}\r\n");
 
 					// Do not change the Z_CompiledInDeferPackage_UPackage_ without changing LC_SymbolPatterns
-					builder.Append("\tstatic FRegisterCompiledInInfo Z_CompiledInDeferPackage_UPackage_").Append(strippedName).Append('(').Append(singletonName)
+					builder.Append("static FRegisterCompiledInInfo Z_CompiledInDeferPackage_UPackage_").Append(strippedName).Append('(').Append(singletonName)
 						.Append(", TEXT(\"").Append(package.SourceName).Append("\"), Z_Registration_Info_UPackage_").Append(strippedName).Append(", CONSTRUCT_RELOAD_VERSION_INFO(FPackageReloadVersionInfo, ")
 						.Append($"0x{bodiesHash:X8}, 0x{declarationsHash:X8}").Append("));\r\n");
 				}
+
+				// Verse registration
+				if (Module.Module.HasVerse)
+				{
+					string staticsName = $"Z_VerseRegistration_{Module.Module.Name}_Statics";
+					builder.Append($"namespace {staticsName}\r\n");
+					builder.Append("{\r\n");
+					if (Module.Module.VerseDependencies.Count > 0)
+					{
+						builder.Append("\tconst FVniPackageName Dependencies[] = {\r\n");
+						foreach (string dependency in Module.Module.VerseDependencies)
+						{
+							builder.Append($"\t\t{GetVniPackageName(dependency)},\r\n");
+						}
+						builder.Append("\t};\r\n");
+					}
+					builder.Append($"\tconst FVniPackageName Name = {GetVniPackageName(Module.Module.VersePackageName)};\r\n");
+					builder.Append("\tV_DEFINE_CPP_MODULE_REGISTRAR(\r\n");
+					builder.Append("\t\tUHT,\r\n");
+					builder.Append($"\t\t{Module.Module.Name},\r\n");
+					builder.Append("\t\tName,\r\n");
+					builder.Append($"\t\tTEXT(\"{Module.Module.VersePath}\"),\r\n");
+					builder.Append($"\t\tEVerseScope::{Module.Module.VerseScope},\r\n");
+					builder.Append($"\t\tTEXT(\"{Module.Module.VerseDirectoryPath}\"),\r\n");
+					builder.Append($"\t\t").Append(Module.Module.VerseDependencies.Count > 0 ? "Dependencies" : "nullptr").Append(",\r\n");
+					builder.Append($"\t\t{Module.Module.VerseDependencies.Count},\r\n");
+					builder.Append($"\t\tnullptr);\r\n");
+					builder.Append("}\r\n");
+				}
+
 				builder.Append(EnableDeprecationWarnings).Append("\r\n");
 
 				if (SaveExportedHeaders)
@@ -159,5 +193,17 @@ namespace EpicGames.UHT.Exporters.CodeGen
 				}
 			}
 		}
+
+		private static string GetVniPackageName(ReadOnlySpan<char> packageName)
+		{
+			int slashIndex = packageName.IndexOf('/');
+			if (slashIndex == -1)
+			{
+				throw new UhtIceException("Verse package name expects a string containing a \"/\"");
+			}
+			ReadOnlySpan<char> mountPointName = packageName[..slashIndex];
+			ReadOnlySpan<char> cppModuleName = packageName[(slashIndex + 1)..];
+			return $"{{ TEXT(\"{mountPointName}\"), TEXT(\"{cppModuleName}\") }}";
+		}
 	}
 }
