[Horde] Bake out [CookWorker XX] from shader (specifically seen) and other warnings and error messages (presumably) to prevent Horde from opening closing the errors wildly, since the CookWorker can change between runs.

diff --git a/Engine/Source/Programs/Shared/EpicGames.Horde/Issues/Handlers/HashedIssueHandler.cs b/Engine/Source/Programs/Shared/EpicGames.Horde/Issues/Handlers/HashedIssueHandler.cs
--- a/Engine/Source/Programs/Shared/EpicGames.Horde/Issues/Handlers/HashedIssueHandler.cs
+++ b/Engine/Source/Programs/Shared/EpicGames.Horde/Issues/Handlers/HashedIssueHandler.cs
@@ -121,6 +121,7 @@ namespace EpicGames.Horde.Issues.Handlers
 			sanitized = Regex.Replace(sanitized, @"0[xX][0-9a-fA-F]+", "H"); // Redact hex strings
 			sanitized = Regex.Replace(sanitized, @"\d[\d.,:]*", "n"); // Redact numbers and timestamp like things
 			sanitized = Regex.Replace(sanitized, @"\\", "/");
+			sanitized = Regex.Replace(sanitized, @"\[CookWorker \d{1,2}\]", "[CookWorker X]");
 
 			if (sanitized.Length > 30)
 			{
