Fixed UHT issue where private properties in a structure or class defined in a namespace would result in a compile error in the generated code.

#rnx

diff --git a/Engine/Source/Programs/LowLevelTests/FoundationTests/Private/NamespacedTypes.h b/Engine/Source/Programs/LowLevelTests/FoundationTests/Private/NamespacedTypes.h
--- a/Engine/Source/Programs/LowLevelTests/FoundationTests/Private/NamespacedTypes.h
+++ b/Engine/Source/Programs/LowLevelTests/FoundationTests/Private/NamespacedTypes.h
@@ -37,6 +37,10 @@ struct FNamespacedStruct
 
     UPROPERTY()
     TArray<ENamespacedEnum> ArrayOfEnums;
+
+private:
+	UPROPERTY()
+	int64 SomePrivateProperty;
 };
 
 inline uint32 GetTypeHash(FNamespacedStruct Value)
@@ -94,6 +98,10 @@ public:
 
     UFUNCTION(BlueprintImplementableEvent, Category=Test)
     void BPImplementable(const FNamespacedStruct& StructParam, ENamespacedEnum EnumParam, UNamespacedObject* ObjectParam, TSubclassOf<UNamespacedObject> ClassParam);
+
+private:
+	UPROPERTY()
+	int64 SomePrivateProperty;
 };
 
 }
diff --git a/Engine/Source/Programs/Shared/EpicGames.UHT/Exporters/CodeGen/UhtHeaderCodeGeneratorHFile.cs b/Engine/Source/Programs/Shared/EpicGames.UHT/Exporters/CodeGen/UhtHeaderCodeGeneratorHFile.cs
--- a/Engine/Source/Programs/Shared/EpicGames.UHT/Exporters/CodeGen/UhtHeaderCodeGeneratorHFile.cs
+++ b/Engine/Source/Programs/Shared/EpicGames.UHT/Exporters/CodeGen/UhtHeaderCodeGeneratorHFile.cs
@@ -509,7 +509,7 @@ namespace EpicGames.UHT.Exporters.CodeGen
 			{
 				using (UhtMacroCreator macro = new(builder, this, scriptStruct, GeneratedBodyMacroSuffix))
 				{
-					builder.Append("\tfriend struct Z_Construct_UScriptStruct_").Append(scriptStruct.SourceName).Append("_Statics; \\\r\n");
+					builder.Append("\tfriend struct ::Z_Construct_UScriptStruct_").Append(scriptStruct.SourceName).Append("_Statics; \\\r\n");
 					builder.Append('\t');
 					if (!scriptStruct.ScriptStructFlags.HasAnyFlags(EStructFlags.RequiredAPI) && 
 					    (Module.Module.AlwaysExportStructs || scriptStruct.ScriptStructExportFlags.HasAnyFlags(UhtScriptStructExportFlags.MinimalAPI)))
@@ -1609,7 +1609,7 @@ namespace EpicGames.UHT.Exporters.CodeGen
 			builder.Append("\tstatic void StaticRegisterNatives").Append(classObj.SourceName).Append("(); \\\r\n");
 			if (!classObj.ClassExportFlags.HasAnyFlags(UhtClassExportFlags.NoExport))
 			{
-				builder.Append("\tfriend struct ").Append(info.RegisteredSingletonName).Append("_Statics; \\\r\n");
+				builder.Append("\tfriend struct ::").Append(info.RegisteredSingletonName).Append("_Statics; \\\r\n");
 			}
 			
 			builder.Append('\t').Append("static UClass* GetPrivateStaticClass(); \\\r\n");
