[UBA]
* Added UBT to hook in memory groups when enqueueing actions

diff --git a/Engine/Source/Programs/Shared/EpicGames.UBA/IScheduler.cs b/Engine/Source/Programs/Shared/EpicGames.UBA/IScheduler.cs
--- a/Engine/Source/Programs/Shared/EpicGames.UBA/IScheduler.cs
+++ b/Engine/Source/Programs/Shared/EpicGames.UBA/IScheduler.cs
@@ -114,7 +114,7 @@ namespace EpicGames.UBA
 		/// <summary>
 		/// Queue process.
 		/// </summary>
-		uint EnqueueProcess(ProcessStartInfo info, double weight, bool canDetour, bool canExecuteRemotely, int[]? dependencies, byte[]? knownInputs, uint knownInputsCount, uint cacheBucket, object userData);
+		uint EnqueueProcess(ProcessStartInfo info, double weight, bool canDetour, bool canExecuteRemotely, int[]? dependencies, byte[]? knownInputs, uint knownInputsCount, uint cacheBucket, uint memoryGroup, ulong predictedMemoryUsage, object userData);
 
 		/// <summary>
 		/// Cancel all active processes and skip queued ones
diff --git a/Engine/Source/Programs/Shared/EpicGames.UBA/Impl/SchedulerImpl.cs b/Engine/Source/Programs/Shared/EpicGames.UBA/Impl/SchedulerImpl.cs
--- a/Engine/Source/Programs/Shared/EpicGames.UBA/Impl/SchedulerImpl.cs
+++ b/Engine/Source/Programs/Shared/EpicGames.UBA/Impl/SchedulerImpl.cs
@@ -109,7 +109,7 @@ namespace EpicGames.UBA.Impl
 		static extern byte Scheduler_IsEmpty(nint scheduler);
 
 		[DllImport("UbaHost", CharSet = CharSet.Auto)]
-		static extern uint Scheduler_EnqueueProcess2(nint scheduler, nint info, float weight, byte canDetour, byte canExecuteRemotely, int[]? dependencies, int dependencyCount, byte[]? knownInputs, int knownInputsBytes, uint knownInputsCount, uint cacheBucket);
+		static extern uint Scheduler_EnqueueProcess3(nint scheduler, nint info, float weight, byte canDetour, byte canExecuteRemotely, int[]? dependencies, int dependencyCount, byte[]? knownInputs, int knownInputsBytes, uint knownInputsCount, uint cacheBucket, uint memoryGroup, ulong predictedMemoryUsage);
 
 		[DllImport("UbaHost", CharSet = CharSet.Auto)]
 		static extern float Scheduler_GetProcessWeightThatCanRunRemotelyNow(nint scheduler);
@@ -145,10 +145,10 @@ namespace EpicGames.UBA.Impl
 			Scheduler_Start(_schedulerHandle);
 		}
 
-		public uint EnqueueProcess(ProcessStartInfo info, double weight, bool canDetour, bool canExecuteRemotely, int[]? dependencies, byte[]? knownInputs, uint knownInputsCount, uint cacheBucket, object userData)
+		public uint EnqueueProcess(ProcessStartInfo info, double weight, bool canDetour, bool canExecuteRemotely, int[]? dependencies, byte[]? knownInputs, uint knownInputsCount, uint cacheBucket, uint memoryGroup, ulong predictedMemoryUsage, object userData)
 		{
 			SchedulerProcess process = new(this, info, userData);
-			uint index = Scheduler_EnqueueProcess2(_schedulerHandle, process._startInfoHandle, (float)weight, (byte)(canDetour?1:0), (byte)(canExecuteRemotely?1:0), dependencies, dependencies != null ? dependencies.Length : 0, knownInputs, knownInputs != null ? knownInputs.Length : 0, knownInputsCount, cacheBucket);
+			uint index = Scheduler_EnqueueProcess3(_schedulerHandle, process._startInfoHandle, (float)weight, (byte)(canDetour?1:0), (byte)(canExecuteRemotely?1:0), dependencies, dependencies != null ? dependencies.Length : 0, knownInputs, knownInputs != null ? knownInputs.Length : 0, knownInputsCount, cacheBucket, memoryGroup, predictedMemoryUsage);
 			process._index = index;
 			_processes.Add(process);
 			return index;
diff --git a/Engine/Source/Programs/UnrealBuildTool/Executors/UnrealBuildAccelerator/UBAExecutor.cs b/Engine/Source/Programs/UnrealBuildTool/Executors/UnrealBuildAccelerator/UBAExecutor.cs
--- a/Engine/Source/Programs/UnrealBuildTool/Executors/UnrealBuildAccelerator/UBAExecutor.cs
+++ b/Engine/Source/Programs/UnrealBuildTool/Executors/UnrealBuildAccelerator/UBAExecutor.cs
@@ -229,6 +229,16 @@ namespace UnrealBuildTool
 				{
 					await writer.WriteLineAsync("remote: false");
 				}
+				uint bucket = GetCacheBucket(action);
+				if (bucket != 0)
+				{
+					await writer.WriteLineAsync($"cache: {bucket}");
+				}
+				uint memoryGroup = GetMemoryGroup(action);
+				if (memoryGroup != 0)
+				{
+					await writer.WriteLineAsync($"memgroup: {memoryGroup}");
+				}
 				if (action.PrerequisiteActions.Any())
 				{
 					await writer.WriteAsync("dep: [");
@@ -846,7 +856,9 @@ namespace UnrealBuildTool
 				var (knownInputs, knownInputsCount) = BuildKnownInputs(action);
 
 				uint bucket = GetCacheBucket(action);
-				scheduler.EnqueueProcess(startInfo, action.Weight, enableDetour, canRunRemotely, prereqActionsSortIndex, knownInputs, knownInputsCount, bucket, action);
+				uint memoryGroup = GetMemoryGroup(action);
+				ulong predictedMemoryUsage = 0; // uba would do a better job scheduling if we knew approx memory usage
+				scheduler.EnqueueProcess(startInfo, action.Weight, enableDetour, canRunRemotely, prereqActionsSortIndex, knownInputs, knownInputsCount, bucket, memoryGroup, predictedMemoryUsage, action);
 
 				++totalActions;
 				cacheTotalActions += bucket != 0 ? 0 : 1;
@@ -1060,6 +1072,34 @@ namespace UnrealBuildTool
 			return bucket;
 		}
 
+		uint GetStringHash(string str, uint hash = 5381)
+		{
+			for (int idx = 0; idx < str.Length; idx++)
+			{
+				hash += (hash << 5) + str[idx];
+			}
+			return hash;
+		}
+
+		uint GetMemoryGroup(LinkedAction action)
+		{
+			if (action.ActionType != ActionType.Compile)
+				return 0;
+
+			FileItem? pch = action.ProducedItems.FirstOrDefault(i => i.HasExtension(".pch"));
+
+			if (pch == null)
+			{
+				pch = action.Inner.PrerequisiteItems.FirstOrDefault(i => i.HasExtension(".pch"));
+			}
+
+			if (pch != null)
+			{
+				return GetStringHash(pch.FullName);
+			}
+			return 1;
+		}
+
 		ulong GetActionRootsHandle(LinkedAction action)
 		{
 			using MemoryStream outputsMemory = new(1024);
