[UBA]
* Exposed trace through EpicGames.UBA

diff --git a/Engine/Source/Programs/Shared/EpicGames.UBA/ISessionServer.cs b/Engine/Source/Programs/Shared/EpicGames.UBA/ISessionServer.cs
--- a/Engine/Source/Programs/Shared/EpicGames.UBA/ISessionServer.cs
+++ b/Engine/Source/Programs/Shared/EpicGames.UBA/ISessionServer.cs
@@ -144,6 +144,11 @@ namespace EpicGames.UBA
 		/// </summary>
 		public abstract event RemoteProcessReturnedEventHandler? RemoteProcessReturned;
 
+		/// <summary>
+		/// Add more info about session that show in visualizer
+		/// </summary>
+		public abstract void AddInfo(string info);
+
 		/// <summary>
 		/// Will tell all remote machines that they can disconnect once their active processes are done
 		/// Will also stop listening for new remote machines
@@ -229,16 +234,6 @@ namespace EpicGames.UBA
 		/// <param name="errorCount">Number of errors in session</param>
 		public abstract void UpdateProgress(uint processesTotal, uint processesDone, uint errorCount);
 
-		/// <summary>
-		/// Writes external status to the uba trace stream which can then be visualized by ubavisualizer
-		/// </summary>
-		/// <param name="statusRow">Row of status text. Reuse one index to show one line in visualizer</param>
-		/// <param name="statusColumn">The identation of status name that will be shown in visualizer</param>
-		/// <param name="statusText">The status text that will be shown in visualizer</param>
-		/// <param name="statusType">The status type</param>
-		/// <param name="statusLink">Optional hyperlink that can be used to make text clickable in visualizer</param>
-		public abstract void UpdateStatus(uint statusRow, uint statusColumn, string statusText, LogEntryType statusType, string? statusLink = null);
-
 		/// <summary>
 		/// Set a custom cas key for a process's tracked inputs
 		/// </summary>
@@ -258,25 +253,9 @@ namespace EpicGames.UBA
 		public abstract void PrintSummary();
 
 		/// <summary>
-		/// Begin task. Will create a bar in uba visualizer
+		/// Get the trace for this session
 		/// </summary>
-		/// <param name="description">Description showing on bar in visualizer</param>
-		/// <param name="details">Details showing when hovering over bar in visualizer</param>
-		/// <returns>Task id</returns>
-		public abstract uint TaskBegin(string description, string details);
-
-		/// <summary>
-		/// Add hint to task.. will show with milliseconds since start or previous hint
-		/// </summary>
-		/// <param name="taskId">Task id returned by TaskBegin</param>
-		/// <param name="hint">Hint text</param>
-		public abstract void TaskHint(uint taskId, string hint);
-
-		/// <summary>
-		/// End task.
-		/// </summary>
-		/// <param name="taskId">Task id returned by TaskBegin</param>
-		public abstract void TaskEnd(uint taskId);
+		public abstract ITrace GetTrace();
 
 		/// <summary>
 		/// Create a ISessionServer object
@@ -288,33 +267,4 @@ namespace EpicGames.UBA
 			return new SessionServerImpl(info);
 		}
 	}
-
-	/// <summary>
-	/// Helper type to track task in scope
-	/// </summary>
-	public readonly struct SessionTaskScope : IDisposable
-	{
-		private readonly ISessionServer _server;
-		private readonly uint _taskId;
-
-		/// <summary>
-		/// Will create a bar in uba visualizer
-		/// </summary>
-		/// <param name="server">Session used for this task scope</param>
-		/// <param name="description">Description showing on bar in visualizer</param>
-		/// <param name="details">Details showing when hovering over bar in visualizer</param>
-		public SessionTaskScope(ISessionServer server, string description, string details)
-		{
-			_server = server;
-			_taskId = server.TaskBegin(description, details);
-		}
-
-		/// <summary>
-		/// End task
-		/// </summary>
-		public void Dispose()
-		{
-			_server.TaskEnd(_taskId);
-		}
-	}
 }
diff --git a/Engine/Source/Programs/Shared/EpicGames.UBA/ITrace.cs b/Engine/Source/Programs/Shared/EpicGames.UBA/ITrace.cs
new file mode 100644
--- /dev/null
+++ b/Engine/Source/Programs/Shared/EpicGames.UBA/ITrace.cs
@@ -0,0 +1,85 @@
+// Copyright Epic Games, Inc. All Rights Reserved.
+
+using System;
+
+namespace EpicGames.UBA
+{
+	/// <summary>
+	/// Trace interface
+	/// </summary>
+	public interface ITrace : IBaseInterface
+	{
+		/// <summary>
+		/// Begin task. Will create a bar in uba visualizer
+		/// </summary>
+		/// <param name="description">Description showing on bar in visualizer</param>
+		/// <param name="details">Details showing when hovering over bar in visualizer</param>
+		/// <returns>Task id</returns>
+		public abstract uint TaskBegin(string description, string details);
+
+		/// <summary>
+		/// Add hint to task.. will show with milliseconds since start or previous hint
+		/// </summary>
+		/// <param name="taskId">Task id returned by TaskBegin</param>
+		/// <param name="hint">Hint text</param>
+		public abstract void TaskHint(uint taskId, string hint);
+
+		/// <summary>
+		/// End task.
+		/// </summary>
+		/// <param name="taskId">Task id returned by TaskBegin</param>
+		public abstract void TaskEnd(uint taskId);
+
+		/// <summary>
+		/// Writes external status to the uba trace stream which can then be visualized by ubavisualizer
+		/// </summary>
+		/// <param name="statusRow">Row of status text. Reuse one index to show one line in visualizer</param>
+		/// <param name="statusColumn">The identation of status name that will be shown in visualizer</param>
+		/// <param name="statusText">The status text that will be shown in visualizer</param>
+		/// <param name="statusType">The status type</param>
+		/// <param name="statusLink">Optional hyperlink that can be used to make text clickable in visualizer</param>
+		public abstract void UpdateStatus(uint statusRow, uint statusColumn, string statusText, LogEntryType statusType, string? statusLink = null);
+
+		/// <summary>
+		/// Writes trace to file. Only works if created with ITrace.Create
+		/// </summary>
+		public abstract void CloseAndWrite(string traceFileName);
+
+		/// <summary>
+		/// Creates a new trace
+		/// </summary>
+		public static ITrace Create(string? name = null)
+		{
+			return new TraceImpl(name);
+		}
+	}
+
+	/// <summary>
+	/// Helper type to track task in scope
+	/// </summary>
+	public readonly struct TraceTaskScope : IDisposable
+	{
+		private readonly ITrace _trace;
+		private readonly uint _taskId;
+
+		/// <summary>
+		/// Will create a bar in uba visualizer
+		/// </summary>
+		/// <param name="trace">Trace used for this task scope</param>
+		/// <param name="description">Description showing on bar in visualizer</param>
+		/// <param name="details">Details showing when hovering over bar in visualizer</param>
+		public TraceTaskScope(ITrace trace, string description, string details)
+		{
+			_trace = trace;
+			_taskId = trace.TaskBegin(description, details);
+		}
+
+		/// <summary>
+		/// End task
+		/// </summary>
+		public void Dispose()
+		{
+			_trace.TaskEnd(_taskId);
+		}
+	}
+}
diff --git a/Engine/Source/Programs/Shared/EpicGames.UBA/Impl/SessionServerImpl.cs b/Engine/Source/Programs/Shared/EpicGames.UBA/Impl/SessionServerImpl.cs
--- a/Engine/Source/Programs/Shared/EpicGames.UBA/Impl/SessionServerImpl.cs
+++ b/Engine/Source/Programs/Shared/EpicGames.UBA/Impl/SessionServerImpl.cs
@@ -67,6 +67,7 @@ namespace EpicGames.UBA.Impl
 		delegate void RemoteProcessReturnedCallback(nint handle);
 
 		nint _handle = IntPtr.Zero;
+		readonly ITrace _trace;
 		readonly ISessionServerCreateInfo _info;
 		readonly RemoteProcessSlotAvailableCallback _remoteProcessSlotAvailableCallbackDelegate;
 		readonly RemoteProcessReturnedCallback _remoteProcessReturnedCallbackDelegate;
@@ -100,9 +101,6 @@ namespace EpicGames.UBA.Impl
 		[DllImport("UbaHost", CharSet = CharSet.Auto)]
 		static extern void SessionServer_UpdateProgress(nint server, uint processesTotal, uint processesDone, uint errorCount);
 
-		[DllImport("UbaHost", CharSet = CharSet.Auto)]
-		static extern void SessionServer_UpdateStatus(nint server, uint statusRow, uint statusColumn, string statusText, byte statusType, string? statusLink);
-
 		[DllImport("UbaHost", CharSet = CharSet.Auto)]
 		static extern nint SessionServer_RunProcess(nint server, nint info, bool async, bool enableDetour);
 
@@ -115,6 +113,9 @@ namespace EpicGames.UBA.Impl
 		[DllImport("UbaHost", CharSet = CharSet.Auto)]
 		static extern void SessionServer_SetMaxRemoteProcessCount(nint server, uint count);
 
+		[DllImport("UbaHost", CharSet = CharSet.Auto)]
+		static extern void SessionServer_AddInfo(nint server, string info);
+
 		[DllImport("UbaHost", CharSet = CharSet.Auto)]
 		static extern void SessionServer_DisableRemoteExecution(nint server);
 
@@ -131,13 +132,7 @@ namespace EpicGames.UBA.Impl
 		static extern void SessionServer_CancelAll(nint server);
 
 		[DllImport("UbaHost", CharSet = CharSet.Auto)]
-		static extern uint SessionServer_TaskBegin(nint server, string description, string details);
-
-		[DllImport("UbaHost", CharSet = CharSet.Auto)]
-		static extern void SessionServer_TaskHint(nint server, uint taskId, string hint);
-
-		[DllImport("UbaHost", CharSet = CharSet.Auto)]
-		static extern void SessionServer_TaskEnd(nint server, uint taskId, bool success);
+		static extern nint SessionServer_GetTrace(nint server);
 
 		[DllImport("UbaHost", CharSet = CharSet.Auto)]
 		static extern void SessionServer_Destroy(nint server);
@@ -168,6 +163,8 @@ namespace EpicGames.UBA.Impl
 			SessionServer_SetRemoteProcessAvailable(_handle, _remoteProcessSlotAvailableCallbackDelegate, 0);
 			SessionServer_SetRemoteProcessReturned(_handle, _remoteProcessReturnedCallbackDelegate);
 
+			_trace = new TraceImpl(SessionServer_GetTrace(_handle));
+
 			foreach (KeyValuePair<string, string> kv in Utils.CrossArchitecturePaths)
 			{
 				SessionServer_RegisterCrossArchitectureMapping(_handle, kv.Key, kv.Value);
@@ -225,6 +222,8 @@ namespace EpicGames.UBA.Impl
 
 		public void DisableRemoteExecution() => SessionServer_DisableRemoteExecution(_handle);
 
+		public void AddInfo(string info) => SessionServer_AddInfo(_handle, info);
+
 		public void SetMaxRemoteProcessCount(uint count) => SessionServer_SetMaxRemoteProcessCount(_handle, count);
 
 		public void RefreshDirectories(params string[] directories) => Array.ForEach(directories, (directory) => SessionServer_RefreshDirectory(_handle, directory));
@@ -238,19 +237,13 @@ namespace EpicGames.UBA.Impl
 
 		public void UpdateProgress(uint processesTotal, uint processesDone, uint errorCount) => SessionServer_UpdateProgress(_handle, processesTotal, processesDone, errorCount);
 
-		public void UpdateStatus(uint statusRow, uint statusColumn, string statusText, LogEntryType statusType, string? statusLink) => SessionServer_UpdateStatus(_handle, statusRow, statusColumn, statusText, (byte)statusType, statusLink);
-
 		public void SetCustomCasKeyFromTrackedInputs(string file, string workingDirectory, IProcess process) => SessionServer_SetCustomCasKeyFromTrackedInputs(_handle, process.GetHandle(), file, workingDirectory);
 
 		public void PrintSummary() => SessionServer_PrintSummary(_handle);
 
 		public void CancelAll() => SessionServer_CancelAll(_handle);
 
-		public uint TaskBegin(string description, string details) => SessionServer_TaskBegin(_handle, description, details);
-
-		public void TaskHint(uint taskId, string hint) => SessionServer_TaskHint(_handle, taskId, hint);
-
-		public void TaskEnd(uint taskId) => SessionServer_TaskEnd(_handle, taskId, true);
+		public ITrace GetTrace() => _trace;
 
 		#endregion
 
diff --git a/Engine/Source/Programs/Shared/EpicGames.UBA/Impl/TraceImpl.cs b/Engine/Source/Programs/Shared/EpicGames.UBA/Impl/TraceImpl.cs
new file mode 100644
--- /dev/null
+++ b/Engine/Source/Programs/Shared/EpicGames.UBA/Impl/TraceImpl.cs
@@ -0,0 +1,87 @@
+// Copyright Epic Games, Inc. All Rights Reserved.
+
+using System;
+using System.Runtime.InteropServices;
+
+namespace EpicGames.UBA
+{
+	internal class TraceImpl : ITrace
+	{
+		nint _handle = IntPtr.Zero;
+		readonly bool _ownsHandle = false;
+
+		#region DllImport
+		[DllImport("UbaHost", CharSet = CharSet.Auto)]
+		static extern uint Trace_TaskBegin(nint trace, string description, string details);
+
+		[DllImport("UbaHost", CharSet = CharSet.Auto)]
+		static extern void Trace_TaskHint(nint trace, uint taskId, string hint);
+
+		[DllImport("UbaHost", CharSet = CharSet.Auto)]
+		static extern void Trace_TaskEnd(nint trace, uint taskId, bool success);
+
+		[DllImport("UbaHost", CharSet = CharSet.Auto)]
+		static extern void Trace_UpdateStatus(nint trace, uint statusRow, uint statusColumn, string statusText, byte statusType, string? statusLink);
+
+		[DllImport("UbaHost", CharSet = CharSet.Auto)]
+		static extern nint Trace_Create(string? name);
+
+		[DllImport("UbaHost", CharSet = CharSet.Auto)]
+		static extern void Trace_Destroy(nint server, string? writeFile);
+		#endregion
+
+
+		internal TraceImpl(nint handle)
+		{
+			_handle = handle;
+			_ownsHandle = false;
+		}
+
+		internal TraceImpl(string? name)
+		{
+			_handle = Trace_Create(name);
+		}
+
+		#region IDisposable
+		~TraceImpl() => Dispose(false);
+
+		public void Dispose()
+		{
+			Dispose(true);
+			GC.SuppressFinalize(this);
+		}
+
+		protected virtual void Dispose(bool disposing)
+		{
+			if (disposing)
+			{
+			}
+
+			if (_ownsHandle && _handle != IntPtr.Zero)
+			{
+				Trace_Destroy(_handle, null);
+				_handle = IntPtr.Zero;
+			}
+		}
+		#endregion
+
+		public nint GetHandle() => _handle;
+
+		public uint TaskBegin(string description, string details) => Trace_TaskBegin(_handle, description, details);
+
+		public void TaskHint(uint taskId, string hint) => Trace_TaskHint(_handle, taskId, hint);
+
+		public void TaskEnd(uint taskId) => Trace_TaskEnd(_handle, taskId, true);
+
+		public void UpdateStatus(uint statusRow, uint statusColumn, string statusText, LogEntryType statusType, string? statusLink) => Trace_UpdateStatus(_handle, statusRow, statusColumn, statusText, (byte)statusType, statusLink);
+
+		public void CloseAndWrite(string traceFileName)
+		{
+			if (_ownsHandle && _handle != IntPtr.Zero)
+			{
+				Trace_Destroy(_handle, traceFileName);
+				_handle = IntPtr.Zero;
+			}
+		}
+	}
+}
\ No newline at end of file
diff --git a/Engine/Source/Programs/UnrealBuildTool/Executors/UnrealBuildAccelerator/UBAExecutor.cs b/Engine/Source/Programs/UnrealBuildTool/Executors/UnrealBuildAccelerator/UBAExecutor.cs
--- a/Engine/Source/Programs/UnrealBuildTool/Executors/UnrealBuildAccelerator/UBAExecutor.cs
+++ b/Engine/Source/Programs/UnrealBuildTool/Executors/UnrealBuildAccelerator/UBAExecutor.cs
@@ -175,7 +175,7 @@ namespace UnrealBuildTool
 		{
 			lock (_sessionLock)
 			{
-				_session?.UpdateStatus(statusRow, statusColumn, statusText, statusType, statusLink);
+				_session?.GetTrace().UpdateStatus(statusRow, statusColumn, statusText, statusType, statusLink);
 			}
 		}
 
@@ -392,6 +392,7 @@ namespace UnrealBuildTool
 						{
 							ISessionServer session = ISessionServer.CreateSessionServer(serverCreateInfo);
 							_session = session;
+							ITrace trace = session.GetTrace();
 
 							OnSessionCreated?.Invoke(session);
 
@@ -400,7 +401,7 @@ namespace UnrealBuildTool
 							{
 								if (!UBAConfig.bDisableRemote) // Need to match the InitAsync
 								{
-									using SessionTaskScope s = new(session, "Requesting Horde cache", "");
+									using TraceTaskScope s = new(trace, "Requesting Horde cache", "");
 									foreach (IUBAAgentCoordinator coordinator in _agentCoordinators)
 									{
 										UnrealBuildAcceleratorCacheConfig? cacheConfig = coordinator.RequestCacheServer(new CancellationToken());
@@ -414,10 +415,10 @@ namespace UnrealBuildTool
 
 							if (cacheConfigs.Count > 0)
 							{
-								session.UpdateStatus(1, 1, "Cache", LogEntryType.Info, null);
-								session.UpdateStatus(1, 6, "Connecting...", LogEntryType.Info);
+								trace.UpdateStatus(1, 1, "Cache", LogEntryType.Info, null);
+								trace.UpdateStatus(1, 6, "Connecting...", LogEntryType.Info);
 
-								using SessionTaskScope s = new(session, "Connecting to cache", "");
+								using TraceTaskScope s = new(trace, "Connecting to cache", "");
 								foreach (UnrealBuildAcceleratorCacheConfig cacheConfig in cacheConfigs)
 								{
 									ICacheClient cacheClient = ICacheClient.CreateCacheClient(session, UBAConfig.bReportCacheMissReason, cacheConfig.Crypto, GetCacheClientHint());
@@ -442,7 +443,7 @@ namespace UnrealBuildTool
 										cacheClient.Dispose();
 									}
 								}
-								session.UpdateStatus(1, 6, _cacheClients.Any() ? "Connected" : "Failed to connect", LogEntryType.Info);
+								trace.UpdateStatus(1, 6, _cacheClients.Any() ? "Connected" : "Failed to connect", LogEntryType.Info);
 							}
 
 							ubaStorage = ubaStorageServer;
