[UBA]
* Fixed so SessionServerImpl is disposing TraceImpl in dispose method

diff --git a/Engine/Source/Programs/Shared/EpicGames.UBA/Impl/SessionServerImpl.cs b/Engine/Source/Programs/Shared/EpicGames.UBA/Impl/SessionServerImpl.cs
--- a/Engine/Source/Programs/Shared/EpicGames.UBA/Impl/SessionServerImpl.cs
+++ b/Engine/Source/Programs/Shared/EpicGames.UBA/Impl/SessionServerImpl.cs
@@ -186,6 +186,8 @@ namespace EpicGames.UBA.Impl
 			{
 			}
 
+			_trace.Dispose();
+			
 			if (_handle != IntPtr.Zero)
 			{
 				SessionServer_Destroy(_handle);
diff --git a/Engine/Source/Programs/Shared/EpicGames.UBA/Impl/TraceImpl.cs b/Engine/Source/Programs/Shared/EpicGames.UBA/Impl/TraceImpl.cs
--- a/Engine/Source/Programs/Shared/EpicGames.UBA/Impl/TraceImpl.cs
+++ b/Engine/Source/Programs/Shared/EpicGames.UBA/Impl/TraceImpl.cs
@@ -57,9 +57,12 @@ namespace EpicGames.UBA
 			{
 			}
 
-			if (_ownsHandle && _handle != IntPtr.Zero)
+			if (_handle != IntPtr.Zero)
 			{
-				Trace_Destroy(_handle, null);
+				if (_ownsHandle)
+				{
+					Trace_Destroy(_handle, null);
+				}
 				_handle = IntPtr.Zero;
 			}
 		}
