Fixing code style issue reported by farm

#rnx

diff --git a/Engine/Source/Programs/Shared/EpicGames.Core/UnrealEngineTypes.cs b/Engine/Source/Programs/Shared/EpicGames.Core/UnrealEngineTypes.cs
--- a/Engine/Source/Programs/Shared/EpicGames.Core/UnrealEngineTypes.cs
+++ b/Engine/Source/Programs/Shared/EpicGames.Core/UnrealEngineTypes.cs
@@ -1582,8 +1582,9 @@ namespace EpicGames.Core
 		Module = 1u << 16,                  // This class represents a Verse module
 		Persistable = 1u << 19,             // This class has a <persistable> attribute
 
+		// NOTE: Underscores removed to keep code style check happy.
 		// @TODO: this should be a per-function flag; a class flag is not granular enough
-		Err_Inoperable = 0x40000000u, // One or more of the class's functions contain mis-linked (malformed) bytecode
-		Err_Incomplete = 0x80000000u, // The class layout is malformed (missing super, illformed data-member, etc.)
+		ErrInoperable = 0x40000000u, // One or more of the class's functions contain mis-linked (malformed) bytecode
+		ErrIncomplete = 0x80000000u, // The class layout is malformed (missing super, illformed data-member, etc.)
 	};
 }
