Cloud DDC - Added conversion of float fields when converting compact binary into json

diff --git a/Engine/Source/Programs/Shared/EpicGames.Serialization/CbField.cs b/Engine/Source/Programs/Shared/EpicGames.Serialization/CbField.cs
--- a/Engine/Source/Programs/Shared/EpicGames.Serialization/CbField.cs
+++ b/Engine/Source/Programs/Shared/EpicGames.Serialization/CbField.cs
@@ -2789,6 +2789,17 @@ namespace EpicGames.Serialization
 					writer.WriteStringValue(stringValue);
 				}
 			}
+			else if (field.IsFloat())
+			{
+				if (field._nameLen != 0)
+				{
+					writer.WriteNumber(field.Name.Span, field.AsDouble());
+				}
+				else
+				{
+					writer.WriteNumberValue(field.AsDouble());
+				}
+			}
 			else
 			{
 				throw new NotImplementedException($"Unhandled type {field.GetType()} when attempting to convert to json");
