Horde: Add a named HTTP client for storage with retries

Under high load certain operations may fail against the storage backend (for example S3). "503 Slow down" is a common one which can be retried.

diff --git a/Engine/Source/Programs/Shared/EpicGames.Horde/HordeExtensions.cs b/Engine/Source/Programs/Shared/EpicGames.Horde/HordeExtensions.cs
--- a/Engine/Source/Programs/Shared/EpicGames.Horde/HordeExtensions.cs
+++ b/Engine/Source/Programs/Shared/EpicGames.Horde/HordeExtensions.cs
@@ -32,6 +32,7 @@ namespace EpicGames.Horde
 			serviceCollection.AddLogging();
 
 			serviceCollection.AddHttpClient();
+			serviceCollection.AddHttpClient(HordeHttpClient.StorageHttpClientName).AddPolicyHandler(HordeHttpMessageHandler.CreateDefaultTransientErrorPolicy());
 			serviceCollection.AddSingleton<BundleCache>(sp => new BundleCache(sp.GetRequiredService<IOptions<HordeOptions>>().Value.BundleCache));
 			serviceCollection.AddSingleton<StorageBackendCache>(CreateBackendCache);
 			serviceCollection.AddSingleton<HttpStorageBackendFactory>();
diff --git a/Engine/Source/Programs/Shared/EpicGames.Horde/HordeHttpClient.cs b/Engine/Source/Programs/Shared/EpicGames.Horde/HordeHttpClient.cs
--- a/Engine/Source/Programs/Shared/EpicGames.Horde/HordeHttpClient.cs
+++ b/Engine/Source/Programs/Shared/EpicGames.Horde/HordeHttpClient.cs
@@ -60,7 +60,12 @@ namespace EpicGames.Horde
 		/// Name of clients used for anonymous requests.
 		/// </summary>
 		public const string AnonymousHttpClientName = "HordeAnonymousHttpClient";
-
+		
+		/// <summary>
+		/// Name of clients created for storage operations
+		/// </summary>
+		public const string StorageHttpClientName = "HordeStorageHttpClient";
+		
 		/// <summary>
 		/// Name of clients created from the http client factory for handling upload redirects. Should not contain Horde auth headers.
 		/// </summary>
diff --git a/Engine/Source/Programs/Shared/EpicGames.Horde/Storage/Backends/HttpStorageBackend.cs b/Engine/Source/Programs/Shared/EpicGames.Horde/Storage/Backends/HttpStorageBackend.cs
--- a/Engine/Source/Programs/Shared/EpicGames.Horde/Storage/Backends/HttpStorageBackend.cs
+++ b/Engine/Source/Programs/Shared/EpicGames.Horde/Storage/Backends/HttpStorageBackend.cs
@@ -412,7 +412,7 @@ namespace EpicGames.Horde.Storage.Backends
 		{
 			HttpClient CreateClient()
 			{
-				HttpClient httpClient = _httpClientFactory.CreateClient(HordeHttpClient.HttpClientName);
+				HttpClient httpClient = _httpClientFactory.CreateClient(HordeHttpClient.StorageHttpClientName);
 				if (accessToken != null)
 				{
 					httpClient.DefaultRequestHeaders.Authorization = new AuthenticationHeaderValue("Bearer", accessToken);
