Made ValidateGlobbedFiles 30x faster

diff --git a/Engine/Source/Programs/Shared/EpicGames.Build/System/CompileScriptModules.cs b/Engine/Source/Programs/Shared/EpicGames.Build/System/CompileScriptModules.cs
--- a/Engine/Source/Programs/Shared/EpicGames.Build/System/CompileScriptModules.cs
+++ b/Engine/Source/Programs/Shared/EpicGames.Build/System/CompileScriptModules.cs
@@ -455,6 +455,8 @@ namespace UnrealBuildBase
 
 		private static bool ValidateGlobbedFiles(DirectoryReference ProjectDirectory, List<CsProjBuildRecord.Glob> Globs, HashSet<string> Dependencies, out string ValidationFailureMessage)
 		{
+			string[] AllFiles = [..Directory.EnumerateFiles(ProjectDirectory.FullName, "*", SearchOption.AllDirectories)];
+			
 			// First, evaluate globs
 
 			// Files are grouped by ItemType (e.g. Compile, EmbeddedResource) to ensure that Exclude and
@@ -478,7 +480,17 @@ namespace UnrealBuildBase
 				{
 					matcher.AddExcludePatterns(Glob.Exclude);
 				}
-				TypedFiles.UnionWith(matcher.GetResultsInFullPath(ProjectDirectory.FullName).Select(x => Path.GetRelativePath(ProjectDirectory.FullName, x)));
+				
+				// Not using GetResultsInFullPath, as internally it's super slow
+				// Match returns filenames separated with '/'
+				if (OperatingSystem.IsWindows())
+				{
+					TypedFiles.UnionWith(matcher.Match(ProjectDirectory.FullName, AllFiles).Files.Select(x => x.Path.Replace('/', Path.DirectorySeparatorChar)));
+				}
+				else
+				{
+					TypedFiles.UnionWith(matcher.Match(ProjectDirectory.FullName, AllFiles).Files.Select(x => x.Path));
+				}
 
 				if (Glob.Remove?.Count > 0)
 				{
@@ -493,10 +505,11 @@ namespace UnrealBuildBase
 			// Then, validation that our evaluation matches what we're comparing against
 
 			// Look for extra files that were found.
-			IEnumerable<string> newFiles = Files.Values.SelectMany(x => x).Distinct().Except(Dependencies).Order();
-			ValidationFailureMessage = newFiles.Any() ? $"Found additional files: {String.Join(Environment.NewLine, newFiles)}" : String.Empty;
+			HashSet<string> newFiles = [..Files.Values.SelectMany(x => x)];
+			newFiles.ExceptWith(Dependencies);
+			ValidationFailureMessage = newFiles.Count > 0 ? $"Found additional files: {String.Join(Environment.NewLine, newFiles.Order())}" : String.Empty;
 
-			return !newFiles.Any();
+			return newFiles.Count == 0;
 		}
 
 		private static bool ValidateBuildRecord(CsProjBuildRecord BuildRecord, DirectoryReference ProjectDirectory, out string ValidationFailureMessage, CsProjBuildHook Hook)
