EpicGames.Build: Downgrade another logging messages to Trace

#rnx

diff --git a/Engine/Source/Programs/Shared/EpicGames.Build/System/CompileScriptModules.cs b/Engine/Source/Programs/Shared/EpicGames.Build/System/CompileScriptModules.cs
--- a/Engine/Source/Programs/Shared/EpicGames.Build/System/CompileScriptModules.cs
+++ b/Engine/Source/Programs/Shared/EpicGames.Build/System/CompileScriptModules.cs
@@ -427,7 +427,7 @@ namespace UnrealBuildBase
 					try
 					{
 						BuildRecord = JsonSerializer.Deserialize<CsProjBuildRecord>(FileReference.ReadAllText(JsonFile));
-						Logger.LogDebug("Loaded script module build record {JsonFile}", JsonFile);
+						Logger.LogTrace("Loaded script module build record {JsonFile}", JsonFile);
 					}
 					catch (Exception Ex)
 					{
