Fixed support for enums in namespaces.

#rnx

diff --git a/Engine/Source/Programs/LowLevelTests/FoundationTests/Private/NamespacedTypes.h b/Engine/Source/Programs/LowLevelTests/FoundationTests/Private/NamespacedTypes.h
--- a/Engine/Source/Programs/LowLevelTests/FoundationTests/Private/NamespacedTypes.h
+++ b/Engine/Source/Programs/LowLevelTests/FoundationTests/Private/NamespacedTypes.h
@@ -20,6 +20,25 @@ enum class ENamespacedEnum : uint8
     Three,
 };
 
+UENUM(BlueprintType)
+enum ENamespacedEnum2 : uint8
+{
+	One,
+	Two,
+	Three,
+};
+
+UENUM(BlueprintType)
+namespace ENamespacedEnum3
+{
+	enum Type : uint8
+	{
+		One,
+		Two,
+		Three,
+	};
+}
+
 inline uint32 GetTypeHash(ENamespacedEnum Value)
 {
     return ::GetTypeHash((uint8)Value);
diff --git a/Engine/Source/Programs/Shared/EpicGames.UHT/Exporters/CodeGen/UhtHeaderCodeGeneratorCppFile.cs b/Engine/Source/Programs/Shared/EpicGames.UHT/Exporters/CodeGen/UhtHeaderCodeGeneratorCppFile.cs
--- a/Engine/Source/Programs/Shared/EpicGames.UHT/Exporters/CodeGen/UhtHeaderCodeGeneratorCppFile.cs
+++ b/Engine/Source/Programs/Shared/EpicGames.UHT/Exporters/CodeGen/UhtHeaderCodeGeneratorCppFile.cs
@@ -458,11 +458,11 @@ namespace EpicGames.UHT.Exporters.CodeGen
 
 			if (Module.Module.AlwaysExportEnums)
 			{
-				builder.Append("template<> ").Append(Module.NonAttributedApi).Append("UEnum* StaticEnum<").Append(enumObj.CppType).Append(">()\r\n");
+				builder.Append("template<> ").Append(Module.NonAttributedApi).Append("UEnum* StaticEnum<").Append(enumObj.Namespace.FullSourceName).Append(enumObj.CppType).Append(">()\r\n");
 			}
 			else
 			{
-				builder.Append("template<> ").Append("UEnum* StaticEnum<").Append(enumObj.CppType).Append(">()\r\n");
+				builder.Append("template<> ").Append("UEnum* StaticEnum<").Append(enumObj.Namespace.FullSourceName).Append(enumObj.CppType).Append(">()\r\n");
 			}
 			builder.Append("{\r\n");
 			builder.Append("\treturn ").Append(enumObj.SourceName).Append("_StaticEnum();\r\n");
diff --git a/Engine/Source/Programs/Shared/EpicGames.UHT/Exporters/CodeGen/UhtHeaderCodeGeneratorHFile.cs b/Engine/Source/Programs/Shared/EpicGames.UHT/Exporters/CodeGen/UhtHeaderCodeGeneratorHFile.cs
--- a/Engine/Source/Programs/Shared/EpicGames.UHT/Exporters/CodeGen/UhtHeaderCodeGeneratorHFile.cs
+++ b/Engine/Source/Programs/Shared/EpicGames.UHT/Exporters/CodeGen/UhtHeaderCodeGeneratorHFile.cs
@@ -743,28 +743,32 @@ namespace EpicGames.UHT.Exporters.CodeGen
 			if (enumObj.CppForm == UhtEnumCppForm.EnumClass)
 			{
 				builder.Append("\r\n");
-				builder.Append("enum class ").Append(enumObj.CppType);
-				if (enumObj.UnderlyingType != UhtEnumUnderlyingType.Unspecified)
+				enumObj.Namespace.AppendSingleLine(builder, (StringBuilder builder) =>
 				{
-					builder.Append(" : ").Append(enumObj.UnderlyingType.ToString().ToLower());
-				}
-				builder.Append(";\r\n");
+					builder.Append("enum class ").Append(enumObj.CppType);
+					if (enumObj.UnderlyingType != UhtEnumUnderlyingType.Unspecified)
+					{
+						builder.Append(" : ").Append(enumObj.UnderlyingType.ToString().ToLower());
+					}
+					builder.Append(';');
+				});
 
 				// Add TIsUEnumClass typetraits
-				builder.Append("template<> struct TIsUEnumClass<").Append(enumObj.CppType).Append("> { enum { Value = true }; };\r\n");
+				builder.Append("template<> struct TIsUEnumClass<").Append(enumObj.Namespace.FullSourceName).Append(enumObj.CppType).Append("> { enum { Value = true }; };\r\n");
 			}
 			else if (enumObj.CppForm == UhtEnumCppForm.Regular && enumObj.UnderlyingType != UhtEnumUnderlyingType.Unspecified)
 			{
 				builder.Append("\r\n");
-				builder.Append("enum ").Append(enumObj.CppType);
-				builder.Append(" : ").Append(enumObj.UnderlyingType.ToString().ToLower());
-				builder.Append(";\r\n");
+				enumObj.Namespace.AppendSingleLine(builder, (StringBuilder builder) =>
+				{
+					builder.Append("enum ").Append(enumObj.CppType).Append(" : ").Append(enumObj.UnderlyingType.ToString().ToLower()).Append(';');
+				});
 			}
 			else if (enumObj.CppForm == UhtEnumCppForm.Namespaced && enumObj.UnderlyingType != UhtEnumUnderlyingType.Unspecified)
 			{
 				string[] splitName = enumObj.CppType.Split("::");
 				builder.Append("\r\n");
-				builder.Append("namespace ").Append(splitName[0]).Append(" { enum ").Append(splitName[1]);
+				builder.Append("namespace ").Append(enumObj.Namespace.FullSourceName).Append(splitName[0]).Append(" { enum ").Append(splitName[1]);
 				builder.Append(" : ").Append(enumObj.UnderlyingType.ToString().ToLower());
 				builder.Append("; }\r\n");
 			}
@@ -774,11 +778,11 @@ namespace EpicGames.UHT.Exporters.CodeGen
 				// Forward declare the StaticEnum<> specialization for enum classes
 				if (Module.Module.AlwaysExportEnums || enumObj.EnumExportFlags.HasAnyFlags(UhtEnumExportFlags.MinimalAPI))
 				{
-					builder.Append("template<> ").Append(Module.NonAttributedApi).Append("UEnum* StaticEnum<").Append(enumObj.CppType).Append(">();\r\n");
+					builder.Append("template<> ").Append(Module.NonAttributedApi).Append("UEnum* StaticEnum<").Append(enumObj.Namespace.FullSourceName).Append(enumObj.CppType).Append(">();\r\n");
 				}
 				else
 				{
-					builder.Append("template<> ").Append("UEnum* StaticEnum<").Append(enumObj.CppType).Append(">();\r\n");
+					builder.Append("template<> ").Append("UEnum* StaticEnum<").Append(enumObj.Namespace.FullSourceName).Append(enumObj.CppType).Append(">();\r\n");
 				}
 			}
 
