Horde: Fix possible NullReferenceException during IssueAuditLogEntry equality comparison
Minor refactors to fix other compiler warnings
#rnx

diff --git a/Engine/Source/Programs/Horde/Plugins/Build/HordeServer.Build/Issues/IssueService.cs b/Engine/Source/Programs/Horde/Plugins/Build/HordeServer.Build/Issues/IssueService.cs
--- a/Engine/Source/Programs/Horde/Plugins/Build/HordeServer.Build/Issues/IssueService.cs
+++ b/Engine/Source/Programs/Horde/Plugins/Build/HordeServer.Build/Issues/IssueService.cs
@@ -957,25 +957,26 @@ namespace HordeServer.Issues
 			{
 				foreach (IssueAuditLogEntry e in uniqueEntries)
 				{
-					if (e._messageTemplate != null)
+					object[]? args = e.GetArgs() ?? [];
+
+					#pragma warning disable CA2254 // Template should be a static 
+					if (e.GetLogLevel() == LogLevel.Error)
 					{
-						if (e._logLevel == LogLevel.Error)
-						{
-							auditLogger.LogError(e._messageTemplate, e._args);
-						}
-						else if (e._logLevel == LogLevel.Warning)
-						{
-							auditLogger.LogWarning(e._messageTemplate, e._args);
-						}
-						else if (e._logLevel == LogLevel.Information)
-						{
-							auditLogger.LogInformation(e._messageTemplate, e._args);
-						}
-						else
-						{
-							auditLogger.LogDebug(e._messageTemplate, e._args);
-						}
+						auditLogger.LogError(e.GetMessageTemplate(), args);
+					}
+					else if (e.GetLogLevel() == LogLevel.Warning)
+					{
+						auditLogger.LogWarning(e.GetMessageTemplate(), args);
+					}
+					else if (e.GetLogLevel() == LogLevel.Information)
+					{
+						auditLogger.LogInformation(e.GetMessageTemplate(), args);
+					}
+					else
+					{
+						auditLogger.LogDebug(e.GetMessageTemplate(), args);
 					}
+					#pragma warning restore CA2254 // Template should be a static 
 				}
 			}
 		}
diff --git a/Engine/Source/Programs/Shared/EpicGames.Horde/Issues/IssueEvent.cs b/Engine/Source/Programs/Shared/EpicGames.Horde/Issues/IssueEvent.cs
--- a/Engine/Source/Programs/Shared/EpicGames.Horde/Issues/IssueEvent.cs
+++ b/Engine/Source/Programs/Shared/EpicGames.Horde/Issues/IssueEvent.cs
@@ -179,7 +179,7 @@ namespace EpicGames.Horde.Issues
 				}
 			}
 
-			IssueAuditLogEntry entry = new IssueAuditLogEntry() { _logLevel = logLevel, _eventId = eventId, _messageTemplate = messageTemplate, _args = [.. args], _exception = exception };
+			IssueAuditLogEntry entry = new IssueAuditLogEntry(logLevel, messageTemplate ?? "", [.. args]);
 			_entries.Add(entry);
 		}
 	}
@@ -187,33 +187,49 @@ namespace EpicGames.Horde.Issues
 	/// <summary>
 	/// Buffered log entry stored in a <see cref="IssueAuditLogger"/>
 	/// </summary>
-	public struct IssueAuditLogEntry : IEquatable<IssueAuditLogEntry>
+	/// <remarks>
+	/// Constructor
+	/// </remarks>
+	public readonly struct IssueAuditLogEntry(LogLevel logLevel, string messageTemplate, object[]? args) : IEquatable<IssueAuditLogEntry>
 	{
 		/// <summary>
 		/// Severity level at which to log the entry
 		/// </summary>
-		public LogLevel _logLevel;
+		private readonly LogLevel _logLevel = logLevel;
 
 		/// <summary>
-		/// Id assocaited with the entry
+		/// Message template for structured log entry
 		/// </summary>
-		public EventId _eventId;
+		private readonly string _messageTemplate = messageTemplate;
 
 		/// <summary>
-		/// Message template for structured log entry
+		/// Arguments to be injected in the message template of the entry
 		/// </summary>
-		public string? _messageTemplate;
+		private readonly object[]? _args = args;
 
 		/// <summary>
-		/// Arguments to be injected in the message template of the entry
+		/// Log level accessor
 		/// </summary>
-		public object?[] _args;
+		public readonly LogLevel GetLogLevel()
+		{
+			return _logLevel;
+		}
 
+		/// <summary>
+		/// Message template accessor
+		/// </summary>
+		public readonly string? GetMessageTemplate()
+		{
+			return _messageTemplate;
+		}
 
 		/// <summary>
-		/// The exception related to this entry.
+		/// Args accessor
 		/// </summary>
-		public Exception? _exception;
+		public readonly object[]? GetArgs()
+		{
+			return _args;
+		}
 
 		/// <summary>
 		/// Compares against another IssueAuditLogEntry object for equality
@@ -227,17 +243,28 @@ namespace EpicGames.Horde.Issues
 				return false;
 			}
 
-			if (_args.Length != otherEntry._args.Length)
+			if (_args == null && otherEntry._args == null)
+			{
+				return true;
+			}
+			else if (_args == null || otherEntry._args == null)
 			{
 				return false;
 			}
-
-			for (int i = 0; i < _args.Length; i++)
+			else
 			{
-				if (!_args[i]!.Equals(otherEntry._args[i]))
+				if (_args.Length != otherEntry._args.Length)
 				{
 					return false;
 				}
+
+				for (int i = 0; i < _args.Length; i++)
+				{
+					if (!_args[i].Equals(otherEntry._args[i]))
+					{
+						return false;
+					}
+				}
 			}
 
 			return true;
