Horde: Fix additional possible NullRefException for individual args array elements
#rnx

diff --git a/Engine/Source/Programs/Shared/EpicGames.Horde/Issues/IssueEvent.cs b/Engine/Source/Programs/Shared/EpicGames.Horde/Issues/IssueEvent.cs
--- a/Engine/Source/Programs/Shared/EpicGames.Horde/Issues/IssueEvent.cs
+++ b/Engine/Source/Programs/Shared/EpicGames.Horde/Issues/IssueEvent.cs
@@ -260,7 +260,15 @@ namespace EpicGames.Horde.Issues
 
 				for (int i = 0; i < _args.Length; i++)
 				{
-					if (!_args[i].Equals(otherEntry._args[i]))
+					if (_args[i] == null && otherEntry._args[i] == null)
+					{
+						continue;
+					}
+					else if (_args[i] == null || otherEntry._args[i] == null)
+					{
+						return false;
+					}
+					else if (!_args[i].Equals(otherEntry._args[i]))
 					{
 						return false;
 					}
