Cloud DDC - Dispose of JsonDocument when converting CbObject to Json, was not aware it was disposable. Thank you static code analysis.

diff --git a/Engine/Source/Programs/Shared/EpicGames.Serialization/CbField.cs b/Engine/Source/Programs/Shared/EpicGames.Serialization/CbField.cs
--- a/Engine/Source/Programs/Shared/EpicGames.Serialization/CbField.cs
+++ b/Engine/Source/Programs/Shared/EpicGames.Serialization/CbField.cs
@@ -1945,7 +1945,7 @@ namespace EpicGames.Serialization
 		{
 			CbWriter writer = new CbWriter();
 
-			JsonDocument doc = JsonDocument.ParseValue(ref reader);
+			using JsonDocument doc = JsonDocument.ParseValue(ref reader);
 			CbJsonReader.ReadJsonField(doc.RootElement, null, writer);
 			return writer.ToObject();
 		}
diff --git a/Engine/Source/Programs/Shared/EpicGames.Serialization/CbSerializer.cs b/Engine/Source/Programs/Shared/EpicGames.Serialization/CbSerializer.cs
--- a/Engine/Source/Programs/Shared/EpicGames.Serialization/CbSerializer.cs
+++ b/Engine/Source/Programs/Shared/EpicGames.Serialization/CbSerializer.cs
@@ -219,7 +219,7 @@ namespace EpicGames.Serialization
 		/// <returns></returns>
 		public static CbObject FromJson(string json)
 		{
-			JsonDocument document = JsonDocument.Parse(json);
+			using JsonDocument document = JsonDocument.Parse(json);
 			CbWriter writer = new CbWriter();
 			ReadJsonField(document.RootElement, null, writer);
 
