Cloud DDC - Sending parts as uniform object as the validation code on the C++ side prefers this (its a more optimal format of this object). Implemented a very explicit support for writing uniform objects, fixing so that the writer converts to uniform types if appopriate required a much more intrusive change.

diff --git a/Engine/Source/Programs/Shared/EpicGames.Serialization/CbWriter.cs b/Engine/Source/Programs/Shared/EpicGames.Serialization/CbWriter.cs
--- a/Engine/Source/Programs/Shared/EpicGames.Serialization/CbWriter.cs
+++ b/Engine/Source/Programs/Shared/EpicGames.Serialization/CbWriter.cs
@@ -333,6 +333,22 @@ namespace EpicGames.Serialization
 			EnterScope(CbFieldType.Object, name);
 		}
 
+		/// <summary>
+		/// Begin writing a uniform object field, an object which fields are all the same type
+		/// </summary>
+		/// <param name="name">Name of the field</param>
+		/// <param name="elementType">The type of the uniform field</param>
+		public void BeginUniformObject(CbFieldName name, CbFieldType elementType)
+		{
+			WriteFields();
+
+			Scope parentScope = _openScopes.Peek();
+			parentScope._itemCount++;
+
+			Scope scope = EnterScope(CbFieldType.UniformObject, name);
+			scope._uniformFieldType = elementType;
+		}
+
 		/// <summary>
 		/// End the current object
 		/// </summary>
@@ -452,11 +468,11 @@ namespace EpicGames.Serialization
 				}
 				else
 				{
-					if (elementType != type)
+					if (elementType != (type | CbFieldType.HasFieldName))
 					{
 						throw new CbWriterException($"Mismatched type for uniform object - expected {elementType}, not {type}");
 					}
-					Memory<byte> buffer = Allocate(name.Length);
+					Memory<byte> buffer = Allocate(nameVarIntLength + name.Length);
 					WriteBinaryPayload(buffer.Span, name.Span);
 				}
 
@@ -1100,7 +1116,6 @@ namespace EpicGames.Serialization
 		/// </summary>
 		/// <param name="writer">Writer for output data</param>
 		/// <param name="value">The uuid value</param>
-		/// <exception cref="NotImplementedException"></exception>
 		public static void WriteUuidValue(this ICbWriter writer, Guid value)
 		{
 			Span<byte> buffer = writer.WriteField(CbFieldType.Uuid, default, 16);
diff --git a/Engine/Source/Programs/UnrealCloudDDC/Jupiter/Controllers/BuildsController.cs b/Engine/Source/Programs/UnrealCloudDDC/Jupiter/Controllers/BuildsController.cs
--- a/Engine/Source/Programs/UnrealCloudDDC/Jupiter/Controllers/BuildsController.cs
+++ b/Engine/Source/Programs/UnrealCloudDDC/Jupiter/Controllers/BuildsController.cs
@@ -294,14 +294,15 @@ namespace Jupiter.Controllers
 			}
 
 			// inject the build parts into the object
+			// we know the object fields will all have the same type so make it a uniform object to please compact binary validation
+			newObjectWriter.BeginUniformObject("parts", CbFieldType.ObjectId | CbFieldType.HasFieldName);
 			await foreach ((string partName, CbObjectId partId) in _buildStore.GetBuildPartsAsync(ns, bucket, buildId))
 			{
-				newObjectWriter.BeginObject("parts");
+
 				newObjectWriter.WriteObjectId(partName, partId);
-				newObjectWriter.EndObject();
 			}
-
-			newObjectWriter.EndObject();
+			newObjectWriter.EndObject(); // parts
+			newObjectWriter.EndObject(); // root
 			CbObject newObject = newObjectWriter.ToObject();
 			
 			string responseType = _formatResolver.GetResponseType(Request, format, CustomMediaTypeNames.UnrealCompactBinary);
