EpicGames.Build: Update binary path check to ignore case

#rnx

diff --git a/Engine/Source/Programs/Shared/EpicGames.Build/Unreal.cs b/Engine/Source/Programs/Shared/EpicGames.Build/Unreal.cs
--- a/Engine/Source/Programs/Shared/EpicGames.Build/Unreal.cs
+++ b/Engine/Source/Programs/Shared/EpicGames.Build/Unreal.cs
@@ -39,13 +39,13 @@ namespace UnrealBuildBase
 			// Search up through the directory tree for the deepest instance of the sub-path "Engine/Binaries/DotNET"
 			while (FoundRootDirectory != null)
 			{
-				if (String.Equals("DotNET", FoundRootDirectory.GetDirectoryName()))
+				if (String.Equals("DotNET", FoundRootDirectory.GetDirectoryName(), StringComparison.OrdinalIgnoreCase))
 				{
 					FoundRootDirectory = FoundRootDirectory.ParentDirectory;
-					if (FoundRootDirectory != null && String.Equals("Binaries", FoundRootDirectory.GetDirectoryName()))
+					if (FoundRootDirectory != null && String.Equals("Binaries", FoundRootDirectory.GetDirectoryName(), StringComparison.OrdinalIgnoreCase))
 					{
 						FoundRootDirectory = FoundRootDirectory.ParentDirectory;
-						if (FoundRootDirectory != null && String.Equals("Engine", FoundRootDirectory.GetDirectoryName()))
+						if (FoundRootDirectory != null && String.Equals("Engine", FoundRootDirectory.GetDirectoryName(), StringComparison.OrdinalIgnoreCase))
 						{
 							FoundRootDirectory = FoundRootDirectory.ParentDirectory;
 							break;
@@ -63,13 +63,13 @@ namespace UnrealBuildBase
 				FoundRootDirectory = DirectoryReference.FindCorrectCase(DirectoryReference.FromString(AssemblyLocation)!);
 				while (FoundRootDirectory != null)
 				{
-					if (String.Equals("Programs", FoundRootDirectory.GetDirectoryName()))
+					if (String.Equals("Programs", FoundRootDirectory.GetDirectoryName(), StringComparison.OrdinalIgnoreCase))
 					{
 						FoundRootDirectory = FoundRootDirectory.ParentDirectory;
-						if (FoundRootDirectory != null && String.Equals("Source", FoundRootDirectory.GetDirectoryName()))
+						if (FoundRootDirectory != null && String.Equals("Source", FoundRootDirectory.GetDirectoryName(), StringComparison.OrdinalIgnoreCase))
 						{
 							FoundRootDirectory = FoundRootDirectory.ParentDirectory;
-							if (FoundRootDirectory != null && String.Equals("Engine", FoundRootDirectory.GetDirectoryName()))
+							if (FoundRootDirectory != null && String.Equals("Engine", FoundRootDirectory.GetDirectoryName(), StringComparison.OrdinalIgnoreCase))
 							{
 								FoundRootDirectory = FoundRootDirectory.ParentDirectory;
 								break;
@@ -101,7 +101,7 @@ namespace UnrealBuildBase
 		{
 			// Return the entry assembly location if it is UnrealBuildTool
 			FileReference? entryDll = FileReference.FromString(Assembly.GetEntryAssembly()?.GetOriginalLocation());
-			if (entryDll?.GetFileName().Equals("UnrealBuildTool.dll", StringComparison.Ordinal) == true)
+			if (entryDll?.GetFileName().Equals("UnrealBuildTool.dll", StringComparison.OrdinalIgnoreCase) == true)
 			{
 				return FileReference.FindCorrectCase(entryDll);
 			}
