[Horde] Adds a PUT function for updating devices to the HordeHttpClient

diff --git a/Engine/Source/Programs/Shared/EpicGames.Horde/HordeHttpClient.cs b/Engine/Source/Programs/Shared/EpicGames.Horde/HordeHttpClient.cs
--- a/Engine/Source/Programs/Shared/EpicGames.Horde/HordeHttpClient.cs
+++ b/Engine/Source/Programs/Shared/EpicGames.Horde/HordeHttpClient.cs
@@ -17,6 +17,7 @@ using EpicGames.Horde.Agents.Telemetry;
 using EpicGames.Horde.Artifacts;
 using EpicGames.Horde.Commits;
 using EpicGames.Horde.Dashboard;
+using EpicGames.Horde.Devices;
 using EpicGames.Horde.Jobs;
 using EpicGames.Horde.Jobs.Graphs;
 using EpicGames.Horde.Logs;
@@ -600,5 +601,19 @@ namespace EpicGames.Horde
 			}
 		}
 		#endregion
+
+		#region Devices
+		/// <summary>
+		/// Updates an individual device with the requested fields
+		/// </summary>
+		/// <param name="deviceId">Id of the device to update</param>
+		/// <param name="request">Request object containing the fields to update</param>
+		/// <param name="cancellationToken">Cancellation token for the operation</param>
+		/// <returns>The response message from the http request</returns>
+		public Task<HttpResponseMessage> PutDeviceUpdateAsync(string deviceId, UpdateDeviceRequest request, CancellationToken cancellationToken = default)
+		{
+			return PutAsync(_httpClient, $"/api/v2/devices/{deviceId}", request, cancellationToken);
+		}
+		#endregion
 	}
 }
