EpicGames.Perforce: Add more metrics to ManagedWorkspace

diff --git a/Engine/Source/Programs/Shared/EpicGames.Perforce.Managed/ManagedWorkspace.cs b/Engine/Source/Programs/Shared/EpicGames.Perforce.Managed/ManagedWorkspace.cs
--- a/Engine/Source/Programs/Shared/EpicGames.Perforce.Managed/ManagedWorkspace.cs
+++ b/Engine/Source/Programs/Shared/EpicGames.Perforce.Managed/ManagedWorkspace.cs
@@ -1367,7 +1367,7 @@ namespace EpicGames.Perforce.Managed
 				if (!_options.UseHaveTable)
 				{
 					// If have table is not used, make client is fully reset as it may have been re-used
-					await UpdateHaveTablePathAsync(perforceClient, $"//{clientName}/...#0", cancellationToken);
+					await UpdateHaveTablePathAsync(perforceClient, $"//{clientName}/...#0", _tracer, cancellationToken);
 				}
 
 				status.Progress = $"({timer.Elapsed.TotalSeconds:0.0}s)";
@@ -1520,11 +1520,11 @@ namespace EpicGames.Perforce.Managed
 				// Sync an initial set of files. Either start with a full workspace and remove files, or start with nothing and add files.
 				if (view.Count == 0 || view[0].StartsWith("-", StringComparison.Ordinal))
 				{
-					await UpdateHaveTablePathAsync(perforceClient, $"//{perforceClient.Settings.ClientName}/...@{changeNumber}", cancellationToken);
+					await UpdateHaveTablePathAsync(perforceClient, $"//{perforceClient.Settings.ClientName}/...@{changeNumber}", _tracer, cancellationToken);
 				}
 				else
 				{
-					await UpdateHaveTablePathAsync(perforceClient, $"//{perforceClient.Settings.ClientName}/...#0", cancellationToken);
+					await UpdateHaveTablePathAsync(perforceClient, $"//{perforceClient.Settings.ClientName}/...#0", _tracer, cancellationToken);
 				}
 
 				// Update with the contents of each filter
@@ -1539,7 +1539,7 @@ namespace EpicGames.Perforce.Managed
 					{
 						syncPath = String.Format("//{0}/{1}@{2}", perforceClient.Settings.ClientName, RemoveLeadingSlash(filter), changeNumber);
 					}
-					await UpdateHaveTablePathAsync(perforceClient, syncPath, cancellationToken);
+					await UpdateHaveTablePathAsync(perforceClient, syncPath, _tracer, cancellationToken);
 				}
 
 				scope.Progress = $"({timer.Elapsed.TotalSeconds:0.0}s)";
@@ -1551,11 +1551,15 @@ namespace EpicGames.Perforce.Managed
 		/// </summary>
 		/// <param name="perforceClient">The Perforce client</param>
 		/// <param name="syncPath">Path to sync</param>
+		/// <param name="tracer">Tracer</param>
 		/// <param name="cancellationToken">Cancellation token</param>
 		/// <returns>Async task</returns>
-		private static async Task UpdateHaveTablePathAsync(IPerforceConnection perforceClient, string syncPath, CancellationToken cancellationToken)
+		private static async Task UpdateHaveTablePathAsync(IPerforceConnection perforceClient, string syncPath, Tracer tracer, CancellationToken cancellationToken)
 		{
-			PerforceResponseList<SyncSummaryRecord> responseList = await perforceClient.TrySyncQuietAsync(SyncOptions.KeepWorkspaceFiles, -1, new[] { syncPath }, cancellationToken);
+			using TelemetrySpan span = tracer.StartActiveSpan($"{nameof(ManagedWorkspace)}.{nameof(UpdateHaveTablePathAsync)}");
+			span.SetAttribute("horde.sync_path", syncPath);
+			
+            PerforceResponseList<SyncSummaryRecord> responseList = await perforceClient.TrySyncQuietAsync(SyncOptions.KeepWorkspaceFiles, -1, new[] { syncPath }, cancellationToken);
 			foreach (PerforceResponse<SyncSummaryRecord> response in responseList)
 			{
 				PerforceError? error = response.Error;
@@ -1626,6 +1630,8 @@ namespace EpicGames.Perforce.Managed
 				// reusing as many fragemnts as possible.
 				List<(Utf8String, StreamTreeBuilder)> fragments = new List<(Utf8String, StreamTreeBuilder)>();
 
+				int numFiles = 0;
+
 				// Handler for each returned record
 				FStatIndexedRecord record = new FStatIndexedRecord();
 				void HandleRecord(PerforceRecord rawRecord)
@@ -1647,6 +1653,8 @@ namespace EpicGames.Perforce.Managed
 						throw new InvalidDataException($"Client path returned by Perforce ('{record.ClientFile}') does not begin with client name ('{clientPrefix}')");
 					}
 
+					numFiles++;
+
 					// Duplicate the client path. If we reference into the raw record, we'll prevent all the raw P4 output from being garbage collected.
 					Utf8String clientFile = record.ClientFile.Clone();
 
@@ -1729,7 +1737,7 @@ namespace EpicGames.Perforce.Managed
 				await perforceClient.RecordCommandAsync("fstat", arguments, null, HandleRecord, cancellationToken);
 
 				// Output the elapsed time
-				scope.Progress = $"({timer.Elapsed.TotalSeconds:0.0}s)";
+				scope.Progress = $"({timer.Elapsed.TotalSeconds:0.0}s {numFiles} files)";
 			}
 
 			return new StreamSnapshotFromMemory(builder);
@@ -1785,6 +1793,8 @@ namespace EpicGames.Perforce.Managed
 				StreamRecord streamRecord = await perforceClient.GetStreamAsync(streamName, true, cancellationToken);
 				PerforceViewMap viewMap = PerforceViewMap.Parse(streamRecord.View);
 
+				int numFiles = 0;
+
 				// Use Horde's additional filtering taking place after stream view mapping
 				PerforceViewFilter viewFilter = PerforceViewFilter.Parse(view);
 
@@ -1810,6 +1820,7 @@ namespace EpicGames.Perforce.Managed
 						Md5Hash md5Hash = Md5Hash.Parse(record.Digest);
 						FileContentId fileContentId = new(md5Hash, record.HeadType);
 						builder.AddFile(clientFile, new StreamFile(record.DepotFile, record.FileSize, fileContentId, record.HeadRev));
+						numFiles++;
 					}
 					else
 					{
@@ -1829,7 +1840,7 @@ namespace EpicGames.Perforce.Managed
 
 				await perforceClient.RecordCommandAsync("fstat", arguments, null, HandleRecord, cancellationToken);
 
-				scope.Progress = $"({timer.Elapsed.TotalSeconds:0.0}s)";
+				scope.Progress = $"({timer.Elapsed.TotalSeconds:0.0}s {numFiles} files)";
 			}
 
 			return new StreamSnapshotFromMemory(builder);
