Fixed up EInternalObjectFlags in UnrealEngineTypes.cs to match ones in ObjectMacros.h

diff --git a/Engine/Source/Programs/Shared/EpicGames.Core/UnrealEngineTypes.cs b/Engine/Source/Programs/Shared/EpicGames.Core/UnrealEngineTypes.cs
--- a/Engine/Source/Programs/Shared/EpicGames.Core/UnrealEngineTypes.cs
+++ b/Engine/Source/Programs/Shared/EpicGames.Core/UnrealEngineTypes.cs
@@ -678,72 +678,70 @@ namespace EpicGames.Core
 		/// </summary>
 		ReachabilityFlag2 = 1 << 16,
 
+		/// <summary>
+		/// Object was constructed in an AutoRTFM transaction that has been aborted
+		/// </summary>
+		AutoRTFMConstructionAborted = 1 << 17,
+
 		/// <summary>
 		/// Object is ready to be imported by another package during loading
 		/// </summary>
-		LoaderImport = 1 << 19,
+		LoaderImport = 1 << 20,
 
 		/// <summary>
 		/// Garbage from logical point of view and should not be referenced. This flag is mirrored in EObjectFlags as RF_Garbage for performance
 		/// </summary>
-		Garbage = 1 << 20,
+		Garbage = 1 << 21,
 
 		/// <summary>
 		/// Object is being asynchronously loaded.
 		/// </summary>
-		AsyncLoadingPhase1 = 1 << 21,
+		AsyncLoadingPhase1 = 1 << 22,
 
 		/// <summary>
 		/// External reference to object in cluster exists
 		/// </summary>
-		ReachableInCluster = 1 << 22,
+		ReachableInCluster = 1 << 23,
 
 		/// <summary>
 		/// Root of a cluster
 		/// </summary>
-		ClusterRoot = 1 << 23,
+		ClusterRoot = 1 << 24,
 
 		/// <summary>
 		/// Native (UClass only).
 		/// </summary>
-		Native = 1 << 24,
+		Native = 1 << 25,
 
 		/// <summary>
 		/// Object exists only on a different thread than the game thread.
 		/// </summary>
-		Async = 1 << 25,
+		Async = 1 << 26,
 
 		/// <summary>
 		/// Object is being asynchronously loaded.
 		/// </summary>
-		AsyncLoadingPhase2 = 1 << 26,
+		AsyncLoadingPhase2 = 1 << 27,
 
 		/// <summary>
 		/// Object is not reachable on the object graph.
 		/// </summary>
-		Unreachable = 1 << 27,
+		Unreachable = 1 << 28,
 
 		/// <summary>
 		/// Object currently has ref-counts associated with it.
 		/// </summary>
-		RefCounted = 1 << 28,
+		RefCounted = 1 << 29,
 
 		/// <summary>
 		/// Object will not be garbage collected, even if unreferenced.
 		/// </summary>
-		RootSet = 1 << 29,
+		RootSet = 1 << 30,
 
 		/// <summary>
 		/// Object didn't have its class constructor called yet (only the UObjectBase one to initialize its most basic members)
 		/// </summary>
-		PendingConstruction = 1 << 30,
-
-		/// <summary>
-		/// Object was constructed in an AutoRTFM transaction that has been aborted
-		/// </summary>
-		AutoRTFMConstructionAborted = 1 << 31,
-
-		AsyncLoading = AsyncLoadingPhase1 | AsyncLoadingPhase2,
+		PendingConstruction = 1 << 31,
 	};
 
 	/// <summary>
