pass down the interactive boolean to the request

diff --git a/Engine/Source/Programs/Shared/EpicGames.Horde/HordeClient.cs b/Engine/Source/Programs/Shared/EpicGames.Horde/HordeClient.cs
--- a/Engine/Source/Programs/Shared/EpicGames.Horde/HordeClient.cs
+++ b/Engine/Source/Programs/Shared/EpicGames.Horde/HordeClient.cs
@@ -363,7 +363,7 @@ namespace EpicGames.Horde
 			using (HttpClient httpClient = CreateAuthenticatedHttpClient())
 			{
 				using HttpRequestMessage request = new HttpRequestMessage(HttpMethod.Get, "api/v1/dashboard/challenge");
-				request.Options.Set(HordeHttpAuthHandler.AllowInteractiveLogin, true);
+				request.Options.Set(HordeHttpAuthHandler.AllowInteractiveLogin, allowLogin);
 
 				using HttpResponseMessage response = await httpClient.SendAsync(request, cancellationToken);
 				return response.IsSuccessStatusCode;
