Horde: Renamed Type structured logging field to avoid overlap with audit log channel.

diff --git a/Engine/Source/Programs/Horde/Plugins/Build/HordeServer.Build/Issues/IssueService.cs b/Engine/Source/Programs/Horde/Plugins/Build/HordeServer.Build/Issues/IssueService.cs
--- a/Engine/Source/Programs/Horde/Plugins/Build/HordeServer.Build/Issues/IssueService.cs
+++ b/Engine/Source/Programs/Horde/Plugins/Build/HordeServer.Build/Issues/IssueService.cs
@@ -814,7 +814,7 @@ namespace HordeServer.Issues
 					issue.Type = $"{issue.Type}:{group}";
 					foreach (IssueEvent e in issue.Events)
 					{
-						e.AuditLogger?.LogDebug("Fingerprint: '{Type}' modified by node annotation: '{Group}'", issue.Type, group);
+						e.AuditLogger?.LogDebug("Fingerprint: '{IssueType}' modified by node annotation: '{Group}'", issue.Type, group);
 					}
 				}
 			}
diff --git a/Engine/Source/Programs/Shared/EpicGames.Horde/Issues/Handlers/CompileIssueHandler.cs b/Engine/Source/Programs/Shared/EpicGames.Horde/Issues/Handlers/CompileIssueHandler.cs
--- a/Engine/Source/Programs/Shared/EpicGames.Horde/Issues/Handlers/CompileIssueHandler.cs
+++ b/Engine/Source/Programs/Shared/EpicGames.Horde/Issues/Handlers/CompileIssueHandler.cs
@@ -73,7 +73,7 @@ namespace EpicGames.Horde.Issues.Handlers
 					issue.Keys.AddSourceFiles(issueEvent);
 					issue.Metadata.Add(CompileTypeAnnotation, compileType);
 					_issues.Add(issue);
-					issueEvent.AuditLogger?.LogDebug("{Type} issue added for event: '{Event}'", fingerprintType, issueEvent.Render());
+					issueEvent.AuditLogger?.LogDebug("{IssueType} issue added for event: '{Event}'", fingerprintType, issueEvent.Render());
 
 					return true;
 				}
diff --git a/Engine/Source/Programs/Shared/EpicGames.Horde/Issues/Handlers/ContentIssueHandler.cs b/Engine/Source/Programs/Shared/EpicGames.Horde/Issues/Handlers/ContentIssueHandler.cs
--- a/Engine/Source/Programs/Shared/EpicGames.Horde/Issues/Handlers/ContentIssueHandler.cs
+++ b/Engine/Source/Programs/Shared/EpicGames.Horde/Issues/Handlers/ContentIssueHandler.cs
@@ -32,7 +32,7 @@ namespace EpicGames.Horde.Issues.Handlers
 					issue.Events.Add(issueEvent);
 					issue.Keys.AddAssets(issueEvent);
 					_issues.Add(issue);
-					issueEvent.AuditLogger?.LogDebug("{Type} issue added for event: '{Event}'", issue.Type, issueEvent.Render());
+					issueEvent.AuditLogger?.LogDebug("{IssueType} issue added for event: '{Event}'", issue.Type, issueEvent.Render());
 
 					return true;
 				}
diff --git a/Engine/Source/Programs/Shared/EpicGames.Horde/Issues/Handlers/CopyrightIssueHandler.cs b/Engine/Source/Programs/Shared/EpicGames.Horde/Issues/Handlers/CopyrightIssueHandler.cs
--- a/Engine/Source/Programs/Shared/EpicGames.Horde/Issues/Handlers/CopyrightIssueHandler.cs
+++ b/Engine/Source/Programs/Shared/EpicGames.Horde/Issues/Handlers/CopyrightIssueHandler.cs
@@ -26,7 +26,7 @@ namespace EpicGames.Horde.Issues.Handlers
 				issue.Events.Add(issueEvent);
 				issue.Keys.AddSourceFiles(issueEvent);
 				_issues.Add(issue);
-				issueEvent.AuditLogger?.LogDebug("{Type} issue added for event: '{Event}'", issue.Type, issueEvent.Render());
+				issueEvent.AuditLogger?.LogDebug("{IssueType} issue added for event: '{Event}'", issue.Type, issueEvent.Render());
 
 				return true;
 			}
diff --git a/Engine/Source/Programs/Shared/EpicGames.Horde/Issues/Handlers/ExternalIssueHandler.cs b/Engine/Source/Programs/Shared/EpicGames.Horde/Issues/Handlers/ExternalIssueHandler.cs
--- a/Engine/Source/Programs/Shared/EpicGames.Horde/Issues/Handlers/ExternalIssueHandler.cs
+++ b/Engine/Source/Programs/Shared/EpicGames.Horde/Issues/Handlers/ExternalIssueHandler.cs
@@ -44,7 +44,7 @@ namespace EpicGames.Horde.Issues.Handlers
 					group.Events.Add(issueEvent);
 
 					_issues.Add(group);
-					issueEvent.AuditLogger?.LogDebug("{Type} issue added for event: '{Event}'", group.Type, issueEvent.Render());
+					issueEvent.AuditLogger?.LogDebug("{IssueType} issue added for event: '{Event}'", group.Type, issueEvent.Render());
 					return true;
 				}
 			}
diff --git a/Engine/Source/Programs/Shared/EpicGames.Horde/Issues/Handlers/LocalizationIssueHandler.cs b/Engine/Source/Programs/Shared/EpicGames.Horde/Issues/Handlers/LocalizationIssueHandler.cs
--- a/Engine/Source/Programs/Shared/EpicGames.Horde/Issues/Handlers/LocalizationIssueHandler.cs
+++ b/Engine/Source/Programs/Shared/EpicGames.Horde/Issues/Handlers/LocalizationIssueHandler.cs
@@ -275,7 +275,7 @@ namespace EpicGames.Horde.Issues.Handlers
 					if(issueGroup.Keys.Count > 0 && !_issues.Contains(issueGroup))
 					{
 						_issues.Add(issueGroup);
-						issueEvent.AuditLogger?.LogDebug("{Type} issue added for event: '{Event}'", issueGroup.Type, issueEvent.Render());
+						issueEvent.AuditLogger?.LogDebug("{IssueType} issue added for event: '{Event}'", issueGroup.Type, issueEvent.Render());
 					}
 
 					return true;
diff --git a/Engine/Source/Programs/Shared/EpicGames.Horde/Issues/Handlers/PerforceCaseIssueHandler.cs b/Engine/Source/Programs/Shared/EpicGames.Horde/Issues/Handlers/PerforceCaseIssueHandler.cs
--- a/Engine/Source/Programs/Shared/EpicGames.Horde/Issues/Handlers/PerforceCaseIssueHandler.cs
+++ b/Engine/Source/Programs/Shared/EpicGames.Horde/Issues/Handlers/PerforceCaseIssueHandler.cs
@@ -36,7 +36,7 @@ namespace EpicGames.Horde.Issues.Handlers
 				issue.Events.Add(issueEvent);
 				issue.Keys.AddDepotPaths(issueEvent);
 				_issues.Add(issue);
-				issueEvent.AuditLogger?.LogDebug("{Type} issue added for event: '{Event}'", issue.Type, issueEvent.Render());
+				issueEvent.AuditLogger?.LogDebug("{IssueType} issue added for event: '{Event}'", issue.Type, issueEvent.Render());
 
 				return true;
 			}
diff --git a/Engine/Source/Programs/Shared/EpicGames.Horde/Issues/Handlers/SanitizerIssueHandler.cs b/Engine/Source/Programs/Shared/EpicGames.Horde/Issues/Handlers/SanitizerIssueHandler.cs
--- a/Engine/Source/Programs/Shared/EpicGames.Horde/Issues/Handlers/SanitizerIssueHandler.cs
+++ b/Engine/Source/Programs/Shared/EpicGames.Horde/Issues/Handlers/SanitizerIssueHandler.cs
@@ -75,7 +75,7 @@ namespace EpicGames.Horde.Issues.Handlers
 				issue.Events.Add(issueEvent);
 
 				_issues.Add(issue);
-				issueEvent.AuditLogger?.LogDebug("{Type} issue added for event: '{Event}'", issue.Type, issueEvent.Render());
+				issueEvent.AuditLogger?.LogDebug("{IssueType} issue added for event: '{Event}'", issue.Type, issueEvent.Render());
 				return true;
 			}
 			return false;
diff --git a/Engine/Source/Programs/Shared/EpicGames.Horde/Issues/Handlers/ScopedIssueHandler.cs b/Engine/Source/Programs/Shared/EpicGames.Horde/Issues/Handlers/ScopedIssueHandler.cs
--- a/Engine/Source/Programs/Shared/EpicGames.Horde/Issues/Handlers/ScopedIssueHandler.cs
+++ b/Engine/Source/Programs/Shared/EpicGames.Horde/Issues/Handlers/ScopedIssueHandler.cs
@@ -75,7 +75,7 @@ namespace EpicGames.Horde.Issues.Handlers
 				issue.Metadata.Add(NodeName, _context.NodeName);
 				issue.Metadata.Add(ScopeName, scope);
 				_issues.Add(issue);
-				logEvent.AuditLogger?.LogDebug("{Type} issue added for event: '{Event}'", issue.Type, logEvent.Render());
+				logEvent.AuditLogger?.LogDebug("{IssueType} issue added for event: '{Event}'", issue.Type, logEvent.Render());
 
 				return true;
 			}
diff --git a/Engine/Source/Programs/Shared/EpicGames.Horde/Issues/Handlers/ShaderIssueHandler.cs b/Engine/Source/Programs/Shared/EpicGames.Horde/Issues/Handlers/ShaderIssueHandler.cs
--- a/Engine/Source/Programs/Shared/EpicGames.Horde/Issues/Handlers/ShaderIssueHandler.cs
+++ b/Engine/Source/Programs/Shared/EpicGames.Horde/Issues/Handlers/ShaderIssueHandler.cs
@@ -32,7 +32,7 @@ namespace EpicGames.Horde.Issues.Handlers
 					issue.Events.Add(logEvent);
 					issue.Keys.AddSourceFiles(logEvent);
 					_issues.Add(issue);
-					logEvent.AuditLogger?.LogDebug("{Type} issue added for event: '{Event}'", issue.Type, logEvent.Render());
+					logEvent.AuditLogger?.LogDebug("{IssueType} issue added for event: '{Event}'", issue.Type, logEvent.Render());
 
 					return true;
 				}
diff --git a/Engine/Source/Programs/Shared/EpicGames.Horde/Issues/Handlers/SymbolIssueHandler.cs b/Engine/Source/Programs/Shared/EpicGames.Horde/Issues/Handlers/SymbolIssueHandler.cs
--- a/Engine/Source/Programs/Shared/EpicGames.Horde/Issues/Handlers/SymbolIssueHandler.cs
+++ b/Engine/Source/Programs/Shared/EpicGames.Horde/Issues/Handlers/SymbolIssueHandler.cs
@@ -55,7 +55,7 @@ namespace EpicGames.Horde.Issues.Handlers
 					if (issue.Keys.Count > 0)
 					{
 						_issues.Add(issue);
-						issueEvent.AuditLogger?.LogDebug("{Type} issue added for event: '{Event}'", issue.Type, issueEvent.Render());
+						issueEvent.AuditLogger?.LogDebug("{IssueType} issue added for event: '{Event}'", issue.Type, issueEvent.Render());
 						return true;
 					}
 					if (_issues.Count > 0)
diff --git a/Engine/Source/Programs/Shared/EpicGames.Horde/Issues/Handlers/UnacceptableWordsIssueHandler.cs b/Engine/Source/Programs/Shared/EpicGames.Horde/Issues/Handlers/UnacceptableWordsIssueHandler.cs
--- a/Engine/Source/Programs/Shared/EpicGames.Horde/Issues/Handlers/UnacceptableWordsIssueHandler.cs
+++ b/Engine/Source/Programs/Shared/EpicGames.Horde/Issues/Handlers/UnacceptableWordsIssueHandler.cs
@@ -26,7 +26,7 @@ namespace EpicGames.Horde.Issues.Handlers
 				issue.Events.Add(issueEvent);
 				issue.Keys.AddSourceFiles(issueEvent);
 				_issues.Add(issue);
-				issueEvent.AuditLogger?.LogDebug("{Type} issue added for event: '{Event}'", issue.Type, issueEvent.Render());
+				issueEvent.AuditLogger?.LogDebug("{IssueType} issue added for event: '{Event}'", issue.Type, issueEvent.Render());
 
 				return true;
 			}
