Horde: Added issue audit logging to new World Leak issue handler

diff --git a/Engine/Source/Programs/Shared/EpicGames.Horde/Issues/Handlers/WorldLeakIssueHandler.cs b/Engine/Source/Programs/Shared/EpicGames.Horde/Issues/Handlers/WorldLeakIssueHandler.cs
--- a/Engine/Source/Programs/Shared/EpicGames.Horde/Issues/Handlers/WorldLeakIssueHandler.cs
+++ b/Engine/Source/Programs/Shared/EpicGames.Horde/Issues/Handlers/WorldLeakIssueHandler.cs
@@ -76,6 +76,7 @@ namespace EpicGames.Horde.Issues.Handlers
 				issue.Events.Add(issueEvent);
 
 				_issues.Add(issue);
+				issueEvent.AuditLogger?.LogDebug("{IssueType} issue added for event: '{Event}'", fingerprint, issueEvent.Render());
 				return true;
 			}
 			return false;
