[Horde] Add GET endpoints for fetching devices to the Horde HTTP client

diff --git a/Engine/Source/Programs/Shared/EpicGames.Horde/Devices/DeviceMessages.cs b/Engine/Source/Programs/Shared/EpicGames.Horde/Devices/DeviceMessages.cs
--- a/Engine/Source/Programs/Shared/EpicGames.Horde/Devices/DeviceMessages.cs
+++ b/Engine/Source/Programs/Shared/EpicGames.Horde/Devices/DeviceMessages.cs
@@ -390,7 +390,7 @@ namespace EpicGames.Horde.Devices
 		/// <summary>
 		/// Device response constructor
 		/// </summary>
-		public GetDeviceResponse(string id, string platformId, string poolId, string name, bool enabled, string? address, string? modelId, string? modifiedByUser, string? notes, DateTime? problemTime, DateTime? maintenanceTime, List<GetDeviceUtilizationResponse>? utilization, string? checkedOutByUser = null, DateTime? checkOutTime = null, DateTime? checkOutExpirationTime = null, bool? cleanRequired = null, DateTime? cleanStartTime = null, DateTime? cleanFinishTime = null, List<string>? tags = null)
+		public GetDeviceResponse(string id, string platformId, string poolId, string name, bool enabled, string? address, string? modelId, string? modifiedByUser, string? notes, DateTime? problemTime, DateTime? maintenanceTime, List<GetDeviceUtilizationResponse>? utilization, string? checkedOutByUserId = null, DateTime? checkOutTime = null, DateTime? checkOutExpirationTime = null, bool? cleanRequired = null, DateTime? cleanStartTime = null, DateTime? cleanFinishTime = null, List<string>? tags = null)
 		{
 			Id = id;
 			Name = name;
@@ -404,7 +404,7 @@ namespace EpicGames.Horde.Devices
 			ProblemTime = problemTime;
 			MaintenanceTime = maintenanceTime;
 			Utilization = utilization;
-			CheckedOutByUserId = checkedOutByUser;
+			CheckedOutByUserId = checkedOutByUserId;
 			CheckOutTime = checkOutTime;
 			CheckOutExpirationTime = checkOutExpirationTime;
 			CleanRequired = cleanRequired;
diff --git a/Engine/Source/Programs/Shared/EpicGames.Horde/HordeHttpClient.cs b/Engine/Source/Programs/Shared/EpicGames.Horde/HordeHttpClient.cs
--- a/Engine/Source/Programs/Shared/EpicGames.Horde/HordeHttpClient.cs
+++ b/Engine/Source/Programs/Shared/EpicGames.Horde/HordeHttpClient.cs
@@ -603,6 +603,28 @@ namespace EpicGames.Horde
 		#endregion
 
 		#region Devices
+
+		/// <summary>
+		/// Retrieves information about all devices
+		/// </summary>
+		/// <param name="cancellationToken">Cancellation token for the operation</param>
+		/// <returns>Collection of GetDeviceResponse containing information about all devices</returns>
+		public Task<List<GetDeviceResponse>> GetDevicesAsync(CancellationToken cancellationToken = default)
+		{
+			return GetAsync<List<GetDeviceResponse>>(_httpClient, $"/api/v2/devices", cancellationToken);
+		}
+
+		/// <summary>
+		/// Retrieves information about the specified device
+		/// </summary>
+		/// <param name="deviceId">Id of the device to query</param>
+		/// <param name="cancellationToken">Cancellation token for the operation</param>
+		/// <returns>GetDeviceResponse containing information about the device</returns>
+		public Task<GetDeviceResponse> GetDeviceAsync(string deviceId, CancellationToken cancellationToken = default)
+		{
+			return GetAsync<GetDeviceResponse>(_httpClient, $"/api/v2/devices/{deviceId}", cancellationToken);
+		}
+
 		/// <summary>
 		/// Updates an individual device with the requested fields
 		/// </summary>
