Horde: Remove unused Password property from EpicGames.Horde.Streams.IWorkspaceType

diff --git a/Engine/Source/Programs/Horde/Plugins/Build/HordeServer.Build/Streams/StreamConfig.cs b/Engine/Source/Programs/Horde/Plugins/Build/HordeServer.Build/Streams/StreamConfig.cs
--- a/Engine/Source/Programs/Horde/Plugins/Build/HordeServer.Build/Streams/StreamConfig.cs
+++ b/Engine/Source/Programs/Horde/Plugins/Build/HordeServer.Build/Streams/StreamConfig.cs
@@ -750,9 +750,6 @@ namespace HordeServer.Streams
 		/// <inheritdoc/>
 		public string? UserName { get; set; }
 
-		/// <inheritdoc/>
-		public string? Password { get; set; }
-
 		/// <inheritdoc/>
 		public string? Identifier { get; set; }
 
diff --git a/Engine/Source/Programs/Shared/EpicGames.Horde/Streams/IStream.cs b/Engine/Source/Programs/Shared/EpicGames.Horde/Streams/IStream.cs
--- a/Engine/Source/Programs/Shared/EpicGames.Horde/Streams/IStream.cs
+++ b/Engine/Source/Programs/Shared/EpicGames.Horde/Streams/IStream.cs
@@ -324,11 +324,6 @@ namespace EpicGames.Horde.Streams
 		/// </summary>
 		string? UserName { get; }
 
-		/// <summary>
-		/// Password to use to log into the workspace
-		/// </summary>
-		string? Password { get; }
-
 		/// <summary>
 		/// Identifier to distinguish this workspace from other workspaces. Defaults to the workspace type name.
 		/// </summary>
diff --git a/Engine/Source/Programs/Shared/EpicGames.Horde/Streams/StreamMessages.cs b/Engine/Source/Programs/Shared/EpicGames.Horde/Streams/StreamMessages.cs
--- a/Engine/Source/Programs/Shared/EpicGames.Horde/Streams/StreamMessages.cs
+++ b/Engine/Source/Programs/Shared/EpicGames.Horde/Streams/StreamMessages.cs
@@ -323,9 +323,6 @@ namespace EpicGames.Horde.Streams
 		/// <inheritdoc/>
 		public string? UserName { get; set; }
 
-		/// <inheritdoc/>
-		string? IWorkspaceType.Password => throw new NotSupportedException();
-
 		/// <inheritdoc/>
 		public string? Identifier { get; set; }
 
