[EpicGames.Build]
* DirectoryItem is now caching both files and directories in one call instead of two. Since file system kernel call stand for 99% of the time it is much faster to do one query

diff --git a/Engine/Source/Programs/Shared/EpicGames.Build/System/DirectoryItem.cs b/Engine/Source/Programs/Shared/EpicGames.Build/System/DirectoryItem.cs
--- a/Engine/Source/Programs/Shared/EpicGames.Build/System/DirectoryItem.cs
+++ b/Engine/Source/Programs/Shared/EpicGames.Build/System/DirectoryItem.cs
@@ -27,14 +27,9 @@ namespace UnrealBuildBase
 		Lazy<DirectoryInfo> Info;
 
 		/// <summary>
-		/// Cached map of name to subdirectory item
+		/// Cached maps of name to subdirectory and name to file
 		/// </summary>
-		Lazy<Dictionary<string, DirectoryItem>> CachedDirectories;
-
-		/// <summary>
-		/// Cached map of name to file
-		/// </summary>
-		Lazy<Dictionary<string, FileItem>> CachedFiles;
+		Lazy<(Dictionary<string, DirectoryItem>, Dictionary<string, FileItem>)> Cache;
 
 		/// <summary>
 		/// Global map of location to item
@@ -49,8 +44,7 @@ namespace UnrealBuildBase
 		private DirectoryItem(DirectoryReference Location, DirectoryInfo Info)
 		{
 			this.Location = Location;
-			CachedDirectories = new(ScanDirectories);
-			CachedFiles = new(ScanFiles);
+			Cache = new(Scan);
 
 			// For some reason we need to call an extra Refresh on Linux/Mac to not get wrong results from "Exists"
 			this.Info = OperatingSystem.IsWindows()
@@ -65,12 +59,12 @@ namespace UnrealBuildBase
 		/// <summary>
 		/// Accessor for map of name to subdirectory item
 		/// </summary>
-		Dictionary<string, DirectoryItem> Directories => CachedDirectories.Value;
+		Dictionary<string, DirectoryItem> Directories => Cache.Value.Item1;
 
 		/// <summary>
 		/// Accessor for map of name to file
 		/// </summary>
-		Dictionary<string, FileItem> Files => CachedFiles.Value;
+		Dictionary<string, FileItem> Files => Cache.Value.Item2;
 
 		/// <summary>
 		/// The name of this directory
@@ -157,22 +151,18 @@ namespace UnrealBuildBase
 				return Info;
 			});
 
-			if (CachedDirectories.IsValueCreated)
+			if (Cache.IsValueCreated)
 			{
-				foreach (DirectoryItem SubDirectory in CachedDirectories.Value.Values)
+				(Dictionary<string, DirectoryItem> dirs, Dictionary<string, FileItem> files) = Cache.Value;
+				foreach (DirectoryItem SubDirectory in dirs.Values)
 				{
 					SubDirectory.ResetCachedInfo();
 				}
-				CachedDirectories = new(ScanDirectories);
-			}
-
-			if (CachedFiles.IsValueCreated)
-			{
-				foreach (FileItem File in CachedFiles.Value.Values)
+				foreach (FileItem File in files.Values)
 				{
 					File.ResetCachedInfo();
 				}
-				CachedFiles = new(ScanFiles);
+				Cache = new(Scan);
 			}
 		}
 
@@ -200,51 +190,15 @@ namespace UnrealBuildBase
 					Info.Refresh();
 					return Info;
 				});
-				Item.CachedDirectories = new(Item.ScanDirectories);
-				Item.CachedFiles = new(Item.ScanFiles);
+				Item.Cache = new(Item.Scan);
 			});
 			FileItem.ResetAllCachedInfo_SLOW();
 		}
 
-		/// <summary>
-		/// Scans the directory for directories, used for lazy initialization.
-		/// </summary>
-		/// <returns></returns>
-		Dictionary<string, DirectoryItem> ScanDirectories()
-		{
-			Dictionary<string, DirectoryItem>? NewDirectories = null;
-			if (Info.Value.Exists)
-			{
-				try
-				{
-					IEnumerable<DirectoryInfo> DirectoryInfos = Info.Value.EnumerateDirectories();
-					NewDirectories = new(DirectoryInfos.Count(), DirectoryReference.Comparer);
-					foreach (DirectoryInfo SubDirectoryInfo in DirectoryInfos)
-					{
-						if (NewDirectories.TryGetValue(SubDirectoryInfo.Name, out DirectoryItem? value))
-						{
-							throw new Exception($"Trying to add {SubDirectoryInfo.FullName} as '{SubDirectoryInfo.Name}' yet exists as {value.FullName}");
-						}
-
-						NewDirectories.Add(SubDirectoryInfo.Name, DirectoryItem.GetItemByDirectoryInfo(SubDirectoryInfo));
-					}
-				}
-				catch (SecurityException)
-				{
-				}
-				catch (UnauthorizedAccessException)
-				{
-				}
-			}
-
-			return NewDirectories ?? new(DirectoryReference.Comparer);
-		}
-
-
 		/// <summary>
 		/// Caches the subdirectories of this directories
 		/// </summary>
-		public bool CacheDirectories() => CachedDirectories.Value != null;
+		public bool CacheDirectories() => Cache.Value.Item1 != null;
 
 		/// <summary>
 		/// Enumerates all the subdirectories
@@ -283,40 +237,62 @@ namespace UnrealBuildBase
 		}
 
 		/// <summary>
-		/// Scans the directory for files, used for lazy initialization.
+		/// Scans the directory for directories and files, used for lazy initialization.
 		/// </summary>
 		/// <returns></returns>
-		Dictionary<string, FileItem> ScanFiles()
+		(Dictionary<string, DirectoryItem> , Dictionary<string, FileItem>) Scan()
 		{
-			Dictionary<string, FileItem>? NewFiles = null;
-			if (Info.Value.Exists)
+			try
 			{
-				try
+				// We want to turn enumerator to a list here since using EnumerateFiles.Count and then enumerate EnumerateFiles cause two iterations of findfirst/findnext
+				// Also, we don't check if exists to minimize kernel call count
+				List<FileSystemInfo> infos = Info.Value.EnumerateFileSystemInfos().ToList();
+
+				int dirCount = 0;
+				foreach (FileSystemInfo info in infos)
 				{
-					IEnumerable<FileInfo> FileInfos = Info.Value.EnumerateFiles();
-					NewFiles = new(FileInfos.Count(), FileReference.Comparer);
-					foreach (FileInfo FileInfo in Info.Value.EnumerateFiles())
+					if (info.Attributes.HasFlag(FileAttributes.Directory))
 					{
-						FileItem FileItem = FileItem.GetItemByFileInfo(FileInfo);
-						FileItem.UpdateCachedDirectory(this);
-						NewFiles[FileInfo.Name] = FileItem;
+						++dirCount;
 					}
 				}
-				catch (SecurityException)
-				{
-				}
-				catch (UnauthorizedAccessException)
+
+				Dictionary<string, DirectoryItem>? newDirs = new(dirCount, FileReference.Comparer);
+				Dictionary<string, FileItem>? newFiles = new(infos.Count - dirCount, FileReference.Comparer);
+
+				foreach (FileSystemInfo info in infos)
 				{
+					if (info.Attributes.HasFlag(FileAttributes.Directory))
+					{
+						newDirs.Add(info.Name, DirectoryItem.GetItemByDirectoryInfo((DirectoryInfo)info));
+					}
+					else
+					{
+						FileItem FileItem = FileItem.GetItemByFileInfo((FileInfo)info);
+						FileItem.UpdateCachedDirectory(this);
+						newFiles.Add(info.Name, FileItem);
+					}
 				}
+
+				return (newDirs, newFiles);
+			}
+			catch (DirectoryNotFoundException)
+			{
+			}
+			catch (SecurityException)
+			{
+			}
+			catch (UnauthorizedAccessException)
+			{
 			}
 
-			return NewFiles ?? new(FileReference.Comparer);
+			return ([], []);
 		}
 
 		/// <summary>
 		/// Caches the files in this directory
 		/// </summary>
-		public bool CacheFiles() => CachedFiles.Value != null;
+		public bool CacheFiles() => Cache.Value.Item2 != null;
 
 		/// <summary>
 		/// Enumerates all the files
