[EpicGames.Build]
* Fixed bug where helper count could end up being -1 (on single core machines)

diff --git a/Engine/Source/Programs/Shared/EpicGames.Build/System/Parallel2.cs b/Engine/Source/Programs/Shared/EpicGames.Build/System/Parallel2.cs
--- a/Engine/Source/Programs/Shared/EpicGames.Build/System/Parallel2.cs
+++ b/Engine/Source/Programs/Shared/EpicGames.Build/System/Parallel2.cs
@@ -130,7 +130,7 @@ namespace UnrealBuildBase
 			internal Exception? Exception;
 		}
 
-		public static readonly int MaxHelperCount = Environment.ProcessorCount - 1;
+		public static readonly int MaxHelperCount = Math.Max(Environment.ProcessorCount - 1, 0);
 	}
 
 	/// <summary>
@@ -245,6 +245,6 @@ namespace UnrealBuildBase
 		private int _workerCount;
 		private Exception? _exception;
 
-		public static readonly int MaxHelperCount = Environment.ProcessorCount - 1;
+		public static readonly int MaxHelperCount = Math.Max(Environment.ProcessorCount - 1, 0);
 	}
 }
\ No newline at end of file
