[UBT]
* Added Lazy<> around value in ConcurrentDictionary in places where creating the value is actually expensive. A GetOrAdd does not prevent the value from being created multiple times (even though only one end up in the dictionary)
* Changed code to use DirectoryItem instead of DirectoryReference in Unreal.GetExtensionDirs which is faster

diff --git a/Engine/Source/Programs/Shared/EpicGames.Build/Unreal.cs b/Engine/Source/Programs/Shared/EpicGames.Build/Unreal.cs
--- a/Engine/Source/Programs/Shared/EpicGames.Build/Unreal.cs
+++ b/Engine/Source/Programs/Shared/EpicGames.Build/Unreal.cs
@@ -374,7 +374,7 @@ namespace UnrealBuildBase
 		private static bool DataDrivenPlatformInfoIniIsPresent(string PlatformName) => IniPresentForPlatform.Value.Contains(PlatformName);
 
 		// cached dictionary of BaseDir to extension directories
-		private static ConcurrentDictionary<DirectoryReference, Tuple<List<DirectoryReference>, List<DirectoryReference>>> CachedExtensionDirectories = new();
+		private static ConcurrentDictionary<DirectoryReference, Lazy<(List<DirectoryReference>, List<DirectoryReference>)>> CachedExtensionDirectories = new();
 
 		/// <summary>
 		/// Finds all the extension directories for the given base directory. This includes platform extensions and restricted folders.
@@ -386,29 +386,27 @@ namespace UnrealBuildBase
 		/// <returns>List of extension directories, including the given base directory</returns>
 		public static List<DirectoryReference> GetExtensionDirs(DirectoryReference BaseDir, bool bIncludePlatformDirectories = true, bool bIncludeRestrictedDirectories = true, bool bIncludeBaseDirectory = true)
 		{
-			Tuple<List<DirectoryReference>, List<DirectoryReference>> CachedDirs = CachedExtensionDirectories.GetOrAdd(BaseDir, _ =>
+			(List<DirectoryReference>, List<DirectoryReference>) CachedDirs = CachedExtensionDirectories.GetOrAdd(BaseDir, new Lazy<(List<DirectoryReference>, List<DirectoryReference>)>(() =>
 			{
-				Tuple<List<DirectoryReference>, List<DirectoryReference>> NewCachedDirs = Tuple.Create(new List<DirectoryReference>(), new List<DirectoryReference>());
+				(List<DirectoryReference>, List<DirectoryReference>) NewCachedDirs = new([], []);
 
-				DirectoryReference PlatformExtensionBaseDir = DirectoryReference.Combine(BaseDir, "Platforms");
-				if (DirectoryReference.Exists(PlatformExtensionBaseDir))
+				DirectoryItem BaseDirItem = DirectoryItem.GetItemByDirectoryReference(BaseDir);
+				if (BaseDirItem.TryGetDirectory("Platforms", out DirectoryItem? PlatformExtensionBaseDir))
 				{
-					NewCachedDirs.Item1.AddRange(DirectoryReference.EnumerateDirectories(PlatformExtensionBaseDir));
+					NewCachedDirs.Item1.AddRange(PlatformExtensionBaseDir.EnumerateDirectories().Select(d => d.Location));
 				}
 
-				DirectoryReference RestrictedBaseDir = DirectoryReference.Combine(BaseDir, "Restricted");
-				if (DirectoryReference.Exists(RestrictedBaseDir))
+				if (BaseDirItem.TryGetDirectory("Restricted", out DirectoryItem? RestrictedBaseDir))
 				{
-					IEnumerable<DirectoryReference> RestrictedDirs = DirectoryReference.EnumerateDirectories(RestrictedBaseDir);
-					NewCachedDirs.Item2.AddRange(RestrictedDirs);
+					IEnumerable<DirectoryItem> RestrictedDirs = RestrictedBaseDir.EnumerateDirectories();
+					NewCachedDirs.Item2.AddRange(RestrictedDirs.Select(d => d.Location));
 
 					// also look for nested platforms in the restricted
-					foreach (DirectoryReference RestrictedDir in RestrictedDirs)
+					foreach (DirectoryItem RestrictedDir in RestrictedDirs)
 					{
-						DirectoryReference RestrictedPlatformExtensionBaseDir = DirectoryReference.Combine(RestrictedDir, "Platforms");
-						if (DirectoryReference.Exists(RestrictedPlatformExtensionBaseDir))
+						if (RestrictedDir.TryGetDirectory("Platforms", out DirectoryItem? RestrictedPlatformExtensionBaseDir))
 						{
-							NewCachedDirs.Item1.AddRange(DirectoryReference.EnumerateDirectories(RestrictedPlatformExtensionBaseDir));
+							NewCachedDirs.Item1.AddRange(RestrictedPlatformExtensionBaseDir.EnumerateDirectories().Select(d => d.Location));
 						}
 					}
 				}
@@ -420,7 +418,7 @@ namespace UnrealBuildBase
 					NewCachedDirs.Item1.RemoveAll(x => !DataDrivenPlatformInfoIniIsPresent(x.GetDirectoryName()));
 				}
 				return NewCachedDirs;
-			});
+			})).Value;
 
 			// now return what the caller wanted (always include BaseDir)
 			List<DirectoryReference> ExtensionDirs = [];
diff --git a/Engine/Source/Programs/Shared/EpicGames.Core/CommandLineArguments.cs b/Engine/Source/Programs/Shared/EpicGames.Core/CommandLineArguments.cs
--- a/Engine/Source/Programs/Shared/EpicGames.Core/CommandLineArguments.cs
+++ b/Engine/Source/Programs/Shared/EpicGames.Core/CommandLineArguments.cs
@@ -2,6 +2,7 @@
 
 using System;
 using System.Collections;
+using System.Collections.Concurrent;
 using System.Collections.Generic;
 using System.ComponentModel;
 using System.Diagnostics;
@@ -858,59 +859,48 @@ namespace EpicGames.Core
 			return obj;
 		}
 
-		private static readonly Dictionary<Type, List<ArgumentTarget>> s_cachedArgumentTargetsForType = new Dictionary<Type, List<ArgumentTarget>>();
+		private static readonly ConcurrentDictionary<Type, Lazy<List<ArgumentTarget>>> s_cachedArgumentTargetsForType = [];
 		static List<ArgumentTarget> GetArgumentTargetsForType(Type? targetType)
 		{
-			List<ArgumentTarget> targets = [];
 			if (targetType == null || targetType == typeof(object))
 			{
-				return targets;
+				return [];
 			}
 
-			lock (s_cachedArgumentTargetsForType)
+			return s_cachedArgumentTargetsForType.GetOrAdd(targetType, static targetType => new Lazy<List<ArgumentTarget>>(() =>
 			{
-				List<ArgumentTarget>? cachedTargets;
-				if (s_cachedArgumentTargetsForType.TryGetValue(targetType, out cachedTargets)) 
+				List<ArgumentTarget> targets = [];
+				foreach (FieldInfo fieldInfo in targetType.GetFields(BindingFlags.Instance | BindingFlags.GetField | BindingFlags.Public | BindingFlags.NonPublic | BindingFlags.DeclaredOnly))
 				{
-					return cachedTargets;
-				}
-			}
-
-			foreach (FieldInfo fieldInfo in targetType.GetFields(BindingFlags.Instance | BindingFlags.GetField | BindingFlags.Public | BindingFlags.NonPublic | BindingFlags.DeclaredOnly))
-			{
-				IEnumerable<CommandLineAttribute> attributes = fieldInfo.GetCustomAttributes<CommandLineAttribute>();
-				if (attributes.Any())
-				{
-					targets.Add(new ArgumentTarget(fieldInfo, fieldInfo.FieldType, fieldInfo.SetValue, fieldInfo.GetValue, attributes.ToArray()));
+					IEnumerable<CommandLineAttribute> attributes = fieldInfo.GetCustomAttributes<CommandLineAttribute>();
+					if (attributes.Any())
+					{
+						targets.Add(new ArgumentTarget(fieldInfo, fieldInfo.FieldType, fieldInfo.SetValue, fieldInfo.GetValue, attributes.ToArray()));
+					}
 				}
-			}
-			foreach (PropertyInfo propertyInfo in targetType.GetProperties(BindingFlags.Instance | BindingFlags.GetProperty | BindingFlags.Public | BindingFlags.NonPublic | BindingFlags.DeclaredOnly))
-			{
-				IEnumerable<CommandLineAttribute> attributes = propertyInfo.GetCustomAttributes<CommandLineAttribute>();
-				if (attributes.Any())
+				foreach (PropertyInfo propertyInfo in targetType.GetProperties(BindingFlags.Instance | BindingFlags.GetProperty | BindingFlags.Public | BindingFlags.NonPublic | BindingFlags.DeclaredOnly))
 				{
-					targets.Add(new ArgumentTarget(propertyInfo, propertyInfo.PropertyType, propertyInfo.SetValue, propertyInfo.GetValue, attributes.ToArray()));
+					IEnumerable<CommandLineAttribute> attributes = propertyInfo.GetCustomAttributes<CommandLineAttribute>();
+					if (attributes.Any())
+					{
+						targets.Add(new ArgumentTarget(propertyInfo, propertyInfo.PropertyType, propertyInfo.SetValue, propertyInfo.GetValue, attributes.ToArray()));
+					}
 				}
-			}
 
-			List<ArgumentTarget> baseTargets = GetArgumentTargetsForType(targetType.BaseType);
-			if (baseTargets.Count > 0)
-			{
-				if (targets.Count > 0)
+				List<ArgumentTarget> baseTargets = GetArgumentTargetsForType(targetType.BaseType);
+				if (baseTargets.Count > 0)
 				{
-					targets.AddRange(baseTargets);
-				}
-				else
-				{
-					targets = baseTargets;
+					if (targets.Count > 0)
+					{
+						targets.AddRange(baseTargets);
+					}
+					else
+					{
+						targets = baseTargets;
+					}
 				}
-			}
-
-			lock (s_cachedArgumentTargetsForType)
-			{
-				s_cachedArgumentTargetsForType[targetType] = targets;
-			}
-			return targets;
+				return targets;
+			})).Value;
 		}
 
 		/// <summary>
diff --git a/Engine/Source/Programs/UnrealBuildTool/System/ConfigCache.cs b/Engine/Source/Programs/UnrealBuildTool/System/ConfigCache.cs
--- a/Engine/Source/Programs/UnrealBuildTool/System/ConfigCache.cs
+++ b/Engine/Source/Programs/UnrealBuildTool/System/ConfigCache.cs
@@ -248,17 +248,17 @@ namespace UnrealBuildTool
 		/// <summary>
 		/// Cache of individual config files
 		/// </summary>
-		static ConcurrentDictionary<FileReference, ConfigFile?> LocationToConfigFile = new ConcurrentDictionary<FileReference, ConfigFile?>();
+		static ConcurrentDictionary<FileReference, Lazy<ConfigFile?>> LocationToConfigFile = [];
 
 		/// <summary>
 		/// Cache of config hierarchies by project
 		/// </summary>
-		static ConcurrentDictionary<ConfigHierarchyKey, ConfigHierarchy> HierarchyKeyToHierarchy = new ConcurrentDictionary<ConfigHierarchyKey, ConfigHierarchy>();
+		static ConcurrentDictionary<ConfigHierarchyKey, Lazy<ConfigHierarchy>> HierarchyKeyToHierarchy = [];
 
 		/// <summary>
 		/// Cache of config fields by type
 		/// </summary>
-		static ConcurrentDictionary<Type, List<ConfigMember>> TypeToConfigMembers = new ConcurrentDictionary<Type, List<ConfigMember>>();
+		static ConcurrentDictionary<Type, Lazy<List<ConfigMember>>> TypeToConfigMembers = [];
 
 		/// <summary>
 		/// Attempts to read a config file (or retrieve it from the cache)
@@ -268,7 +268,7 @@ namespace UnrealBuildTool
 		/// <returns>True if the file exists and was read, false otherwise</returns>
 		internal static bool TryReadFile(FileReference Location, [NotNullWhen(true)] out ConfigFile? ConfigFile)
 		{
-			ConfigFile = LocationToConfigFile.GetOrAdd(Location, _ =>
+			ConfigFile = LocationToConfigFile.GetOrAdd(Location, static Location => new Lazy<ConfigFile?>(() =>
 			{
 				if (FileReference.Exists(Location))
 				{
@@ -277,7 +277,7 @@ namespace UnrealBuildTool
 
 				// Note: it's important to keep the null value here so we don't keep checking for files that don't exist
 				return null;
-			});
+			})).Value;
 			return ConfigFile != null;
 		}
 
@@ -369,7 +369,7 @@ namespace UnrealBuildTool
 
 			// Try to get the cached hierarchy with this key
 			Logger.LogInformation($"Requested Hierarchy: {Type},{ProjectDir},{Platform},{CustomConfig}");
-			return HierarchyKeyToHierarchy.GetOrAdd(Key, _ =>
+			return HierarchyKeyToHierarchy.GetOrAdd(Key, new Lazy<ConfigHierarchy>(() =>
 			{
 				// Find all the input files
 				List<ConfigFile> Files = new List<ConfigFile>();
@@ -408,7 +408,7 @@ namespace UnrealBuildTool
 
 				// Create the hierarchy
 				return new ConfigHierarchy(Files, ReadFiles);
-			});
+			})).Value;
 		}
 
 		/// <summary>
@@ -418,7 +418,7 @@ namespace UnrealBuildTool
 		/// <returns>List of config fields for the given type</returns>
 		static List<ConfigMember> FindConfigMembersForType(Type TargetObjectType)
 		{
-			return TypeToConfigMembers.GetOrAdd(TargetObjectType, _ =>
+			return TypeToConfigMembers.GetOrAdd(TargetObjectType, static TargetObjectType => new Lazy<List<ConfigMember>>(() =>
 			{
 				List<ConfigMember> Members = new List<ConfigMember>();
 				if (TargetObjectType.BaseType != null)
@@ -434,7 +434,7 @@ namespace UnrealBuildTool
 					ProcessConfigTypeMember<PropertyInfo>(TargetObjectType, PropertyInfo, Members, (PropertyInfo, Attribute) => new ConfigProperty(PropertyInfo, Attribute));
 				}
 				return Members;
-			});
+			})).Value;
 		}
 
 		static void ProcessConfigTypeMember<MEMBER>(Type TargetType, MEMBER MemberInfo, List<ConfigMember> Members, Func<MEMBER, ConfigFileAttribute, ConfigMember> CreateConfigMember)
diff --git a/Engine/Source/Programs/UnrealBuildTool/System/ConfigHierarchy.cs b/Engine/Source/Programs/UnrealBuildTool/System/ConfigHierarchy.cs
--- a/Engine/Source/Programs/UnrealBuildTool/System/ConfigHierarchy.cs
+++ b/Engine/Source/Programs/UnrealBuildTool/System/ConfigHierarchy.cs
@@ -1,6 +1,7 @@
 // Copyright Epic Games, Inc. All Rights Reserved.
 
 using System;
+using System.Collections.Concurrent;
 using System.Collections.Generic;
 using System.Diagnostics.CodeAnalysis;
 using System.Linq;
@@ -188,12 +189,7 @@ namespace UnrealBuildTool
 		/// <summary>
 		/// Cache of requested config sections
 		/// </summary>
-		readonly Dictionary<string, ConfigHierarchySection> NameToSection = new(StringComparer.InvariantCultureIgnoreCase);
-
-		/// <summary>
-		/// Lock for NameToSection
-		/// </summary>
-		readonly System.Threading.ReaderWriterLockSlim NameToSectionLock = new();
+		readonly ConcurrentDictionary<string, Lazy<ConfigHierarchySection>> NameToSection = new(StringComparer.InvariantCultureIgnoreCase);
 
 		/// <summary>
 		/// Construct a config hierarchy from the given files
@@ -210,7 +206,6 @@ namespace UnrealBuildTool
 		/// <inheritdoc/>
 		public void Dispose()
 		{
-			NameToSectionLock.Dispose();
 		}
 		#endregion
 
@@ -227,46 +222,22 @@ namespace UnrealBuildTool
 		/// <returns>The merged config section</returns>
 		public ConfigHierarchySection FindSection(string SectionName)
 		{
-			ConfigHierarchySection? Section;
-			try
+			return NameToSection.GetOrAdd(SectionName, new Lazy<ConfigHierarchySection>(() =>
 			{
-				// Acquire a read lock and do a quick check for the config section
-				NameToSectionLock.EnterUpgradeableReadLock();
-				if (!NameToSection.TryGetValue(SectionName, out Section))
+				// Find all the raw sections from the file hierarchy
+				List<ConfigFileSection> RawSections = new List<ConfigFileSection>();
+				foreach (ConfigFile File in Files)
 				{
-					try
-					{
-						// Acquire a write lock and add the config section if another thread didn't just complete it
-						NameToSectionLock.EnterWriteLock();
-						if (!NameToSection.TryGetValue(SectionName, out Section))
-						{
-							// Find all the raw sections from the file hierarchy
-							List<ConfigFileSection> RawSections = new List<ConfigFileSection>();
-							foreach (ConfigFile File in Files)
-							{
-								ConfigFileSection? RawSection;
-								if (File.TryGetSection(SectionName, out RawSection))
-								{
-									RawSections.Add(RawSection);
-								}
-							}
-
-							// Merge them together and add it to the cache
-							Section = new ConfigHierarchySection(RawSections);
-							NameToSection.Add(SectionName, Section);
-						}
-					}
-					finally
+					ConfigFileSection? RawSection;
+					if (File.TryGetSection(SectionName, out RawSection))
 					{
-						NameToSectionLock.ExitWriteLock();
+						RawSections.Add(RawSection);
 					}
 				}
-			}
-			finally
-			{
-				NameToSectionLock.ExitUpgradeableReadLock();
-			}
-			return Section;
+
+				// Merge them together and add it to the cache
+				return new ConfigHierarchySection(RawSections);
+			})).Value;
 		}
 
 		/// <summary>
