[EpicGames.Build]
* Changed PluginsBase.EnumeratePlugins to return list instead of IEnumerable since that makes it possible to query size and run things in parallel at the call site

diff --git a/Engine/Source/Programs/Shared/EpicGames.Build/System/EnumeratePlugins.cs b/Engine/Source/Programs/Shared/EpicGames.Build/System/EnumeratePlugins.cs
--- a/Engine/Source/Programs/Shared/EpicGames.Build/System/EnumeratePlugins.cs
+++ b/Engine/Source/Programs/Shared/EpicGames.Build/System/EnumeratePlugins.cs
@@ -38,7 +38,7 @@ namespace UnrealBuildBase
 		/// </summary>
 		/// <param name="ProjectFile">Path to the project file</param>
 		/// <returns>List of project files</returns>
-		public static IEnumerable<FileReference> EnumeratePlugins(FileReference? ProjectFile)
+		public static List<FileReference> EnumeratePlugins(FileReference? ProjectFile)
 		{
 			List<DirectoryReference> BaseDirs = [.. Unreal.GetExtensionDirs(Unreal.EngineDirectory, "Plugins")];
 			if(ProjectFile != null)
@@ -53,7 +53,7 @@ namespace UnrealBuildBase
 		/// Find paths to all the plugins under a given parent directory (recursively)
 		/// </summary>
 		/// <param name="ParentDirectory">Parent directory to look in. Plugins will be found in any *subfolders* of this directory.</param>
-		public static IEnumerable<FileReference> EnumeratePlugins(DirectoryReference ParentDirectory)
+		public static List<FileReference> EnumeratePlugins(DirectoryReference ParentDirectory)
 		{
 			return PluginFileCache.GetOrAdd(ParentDirectory, _ =>
 			{
