[EpicGames.Build]
* Initialized FileItem UniqueSourceFileMap with quite a large number to reduce rehashing/resizing
* Optimized FileItem.GetItemByFileReference

diff --git a/Engine/Source/Programs/Shared/EpicGames.Build/System/FileItem.cs b/Engine/Source/Programs/Shared/EpicGames.Build/System/FileItem.cs
--- a/Engine/Source/Programs/Shared/EpicGames.Build/System/FileItem.cs
+++ b/Engine/Source/Programs/Shared/EpicGames.Build/System/FileItem.cs
@@ -36,7 +36,7 @@ namespace UnrealBuildBase
 		/// <summary>
 		/// A case-insensitive dictionary that's used to map each unique file name to a single FileItem object.
 		/// </summary>
-		static ConcurrentDictionary<FileReference, FileItem> UniqueSourceFileMap = [];
+		static ConcurrentDictionary<FileReference, FileItem> UniqueSourceFileMap = new(-1, 100000);// [];
 
 		/// <summary>
 		/// Constructor
@@ -146,9 +146,7 @@ namespace UnrealBuildBase
 		/// <param name="Location">Location of the file</param>
 		/// <returns>The FileItem that represents the given a full file path.</returns>
 		public static FileItem GetItemByFileReference(FileReference Location)
-			=> UniqueSourceFileMap.TryGetValue(Location, out FileItem? Result)
-				? Result
-				: UniqueSourceFileMap.GetOrAdd(Location, new FileItem(Location, Location.ToFileInfo()));
+			=> UniqueSourceFileMap.GetOrAdd(Location, static loc => new FileItem(loc, loc.ToFileInfo()));
 
 		/// <summary>
 		/// Deletes the file.
