[EpicGames.Core]
* Fixed bug in CsProjectInfo where recursion of directories with wildcards were wrong. once a \**\ was found in a path it just dropped all parts of the wildcard after it and just searched everything. So a  path\**\SomeDir\*.cs would just ignore the SomeDir part
* Optimized CsProjectInfo to run wildcard path search in parallel

diff --git a/Engine/Source/Programs/Shared/EpicGames.Core/CsProjectInfo.cs b/Engine/Source/Programs/Shared/EpicGames.Core/CsProjectInfo.cs
--- a/Engine/Source/Programs/Shared/EpicGames.Core/CsProjectInfo.cs
+++ b/Engine/Source/Programs/Shared/EpicGames.Core/CsProjectInfo.cs
@@ -1,6 +1,7 @@
 ﻿// Copyright Epic Games, Inc. All Rights Reserved.
 
 using System;
+using System.Collections.Concurrent;
 using System.Collections.Generic;
 using System.Diagnostics.CodeAnalysis;
 using System.IO;
@@ -627,6 +628,29 @@ namespace EpicGames.Core
 			}
 		}
 
+		/// <summary>
+		/// Finds all directories but filters out certain folders
+		/// </summary>
+		static void GetDirectories(DirectoryReference dir, ConcurrentBag<DirectoryReference> list, SearchOption option, ThreadPoolWorkQueue queue)
+		{
+			foreach (DirectoryReference child in DirectoryReference.EnumerateDirectories(dir))
+			{
+				string name = child.GetDirectoryName();
+				if (name == "Intermediate" || name == "Binaries")
+				{
+					continue;
+				}
+
+				list.Add(child);
+				if (option != SearchOption.AllDirectories)
+				{
+					continue;
+				}
+
+				queue.Enqueue(() => { GetDirectories(child, list, option, queue); });
+			}
+		}
+
 		/// recursive helper used by the function below that will append RemainingComponents one by one to ExistingPath, 
 		/// expanding wildcards as necessary. The complete list of files that match the complete path is returned out OutFoundFiles
 		static void ProcessPathComponents(DirectoryReference existingPath, IEnumerable<string> remainingComponents, List<FileReference> outFoundFiles)
@@ -663,15 +687,55 @@ namespace EpicGames.Core
 					// ** means all directories, no matter how deep
 					SearchOption option = currentComponent == "**" ? SearchOption.AllDirectories : SearchOption.TopDirectoryOnly;
 
-					IEnumerable<DirectoryReference> directories = DirectoryReference.EnumerateDirectories(existingPath, currentComponent, option);
-
-					// if we searched all directories regardless of depth, the rest of the components other than the last (file) are irrelevant
-					if (option == SearchOption.AllDirectories)
+					IEnumerable<DirectoryReference> directories;
+					if (option == SearchOption.TopDirectoryOnly)
+					{
+						directories = DirectoryReference.EnumerateDirectories(existingPath, currentComponent, option);
+					}
+					else
 					{
-						remainingComponents = [remainingComponents.Last()];
+						// Let's go wide because we potentially will traverse a lot of directories
+						// and we also want to filter out Intermediate, Binaries, etc
 
-						// ** includes files in the current directory too
-						directories = directories.Concat([existingPath]);
+						ConcurrentBag<DirectoryReference> bag = new();
+						using (ThreadPoolWorkQueue queue = new ThreadPoolWorkQueue())
+						{
+							GetDirectories(existingPath, bag, option, queue);
+						}
+
+						// Add current path
+						bag.Add(existingPath);
+
+						// Consume potential folders after **   (like ..\Engine\**\Build\Gauntlet\**\*.cs
+						string subPath = "";
+						while (remainingComponents.Any())
+						{
+							string component = remainingComponents.First();
+							if (component[0] == '*')
+								break;
+							subPath += Path.DirectorySeparatorChar + component;
+							remainingComponents = remainingComponents.Skip(1);
+						}
+
+						// Match folders against end of Directories found
+						if (subPath.Length > 0)
+						{
+							List<DirectoryReference> matchingDirs = [];
+
+							foreach (DirectoryReference dir in bag)
+							{
+								if (dir.FullName.EndsWith(subPath, StringComparison.Ordinal))
+								{
+									matchingDirs.Add(dir);
+								}
+							}
+							directories = matchingDirs;
+						}
+						else
+						{
+							// .. otherwise just keep going (remaining component is likely *.cs)
+							directories = bag;
+						}
 					}
 
 					foreach (DirectoryReference dir in directories)
