[EpicGames.Core/Epicgames.Build]
* Moved CsProjectInfo to EpicGames.Build to be able to use cached directory infos when

diff --git a/Engine/Source/Programs/Shared/EpicGames.Build/CsProjectInfo.cs b/Engine/Source/Programs/Shared/EpicGames.Build/CsProjectInfo.cs
new file mode 100644
--- /dev/null
+++ b/Engine/Source/Programs/Shared/EpicGames.Build/CsProjectInfo.cs
@@ -0,0 +1,1341 @@
+﻿// Copyright Epic Games, Inc. All Rights Reserved.
+
+using System;
+using System.Collections.Concurrent;
+using System.Collections.Generic;
+using System.Diagnostics.CodeAnalysis;
+using System.IO;
+using System.Linq;
+using System.Runtime.InteropServices;
+using System.Text.RegularExpressions;
+using System.Xml;
+
+namespace EpicGames.Core
+{
+	/// <summary>
+	/// Exception parsing a csproj file
+	/// </summary>
+	public sealed class CsProjectParseException : Exception
+	{
+		internal CsProjectParseException(string? message, Exception? exception = null) : base(message, exception)
+		{
+		}
+	}
+
+	/// <summary>
+	/// Basic information from a preprocessed C# project file. Supports reading a project file, expanding simple conditions in it, parsing property values, assembly references and references to other projects.
+	/// </summary>
+	public class CsProjectInfo
+	{
+		/// <summary>
+		/// Evaluated properties from the project file
+		/// </summary>
+		public Dictionary<string, string> Properties { get; }
+
+		/// <summary>
+		/// Mapping of referenced assemblies to their 'CopyLocal' (aka 'Private') setting.
+		/// </summary>
+		public Dictionary<FileReference, bool> References { get; } = [];
+
+		/// <summary>
+		/// Mapping of referenced projects to their 'CopyLocal' (aka 'Private') setting.
+		/// </summary>
+		public Dictionary<FileReference, bool> ProjectReferences { get; } = [];
+
+		/// <summary>
+		/// List of compile references in the project.
+		/// </summary>
+		public List<FileReference> CompileReferences { get; } = [];
+
+		/// <summary>
+		/// Mapping of content IF they are flagged Always or Newer
+		/// </summary>
+		public Dictionary<FileReference, bool> ContentReferences { get; } = [];
+
+		/// <summary>
+		/// Path to the CSProject file
+		/// </summary>
+		public FileReference ProjectPath { get; }
+
+		/// <summary>
+		/// Constructor
+		/// </summary>
+		/// <param name="inProperties">Initial mapping of property names to values</param>
+		/// <param name="inProjectPath"></param>
+		CsProjectInfo(Dictionary<string, string> inProperties, FileReference inProjectPath)
+		{
+			ProjectPath = inProjectPath;
+			Properties = new Dictionary<string, string>(inProperties);
+			Properties.TryAdd("MSBuildProjectFile", inProjectPath.FullName);
+			Properties.TryAdd("MSBuildThisFileDirectory", inProjectPath.Directory.FullName);
+		}
+
+		/// <summary>
+		/// Get the ouptut file for this project
+		/// </summary>
+		/// <param name="file">If successful, receives the assembly path</param>
+		/// <returns>True if the output file was found</returns>
+		public bool TryGetOutputFile([NotNullWhen(true)] out FileReference? file)
+		{
+			DirectoryReference? outputDir;
+			if(!TryGetOutputDir(out outputDir))
+			{
+				file = null;
+				return false;
+			}
+
+			string? assemblyName;
+			if(!TryGetAssemblyName(out assemblyName))
+			{
+				file = null;
+				return false;
+			}
+
+			file = FileReference.Combine(outputDir, assemblyName + ".dll");
+			return true;
+		}
+
+		/// <summary>
+		/// Returns whether or not this project is a modern dotnet project
+		/// </summary>
+		/// <returns>True if this is a netcore project</returns>
+		private bool IsNetCoreProject()
+		{
+			string? framework;
+			return Properties.TryGetValue("TargetFramework", out framework)
+				&& (framework.StartsWith("netcoreapp", StringComparison.Ordinal) || Regex.Match(framework, @"net\d+\.0").Success);
+		}
+
+		/// <summary>
+		/// Resolve the project's output directory
+		/// </summary>
+		/// <param name="baseDirectory">Base directory to resolve relative paths to</param>
+		/// <returns>The configured output directory</returns>
+		public DirectoryReference GetOutputDir(DirectoryReference baseDirectory)
+		{
+			string? outputPath;
+			if (Properties.TryGetValue("OutputPath", out outputPath))
+			{
+				return DirectoryReference.Combine(baseDirectory, outputPath);
+			}
+			else if (IsNetCoreProject())
+			{
+				string configuration = Properties.TryGetValue("Configuration", out string? value) ? value : "Development";
+				return !Properties.TryGetValue("AppendTargetFrameworkToOutputPath", out string? appendFrameworkStr) || appendFrameworkStr.Equals("true", StringComparison.OrdinalIgnoreCase)
+					? DirectoryReference.Combine(baseDirectory, "bin", configuration, Properties["TargetFramework"])
+					: DirectoryReference.Combine(baseDirectory, "bin", configuration);
+			}
+			else
+			{
+				return baseDirectory;
+			}
+		}
+
+		/// <summary>
+		/// Resolve the project's output directory
+		/// </summary>
+		/// <param name="outputDir">If successful, receives the output directory</param>
+		/// <returns>True if the output directory was found</returns>
+		public bool TryGetOutputDir([NotNullWhen(true)] out DirectoryReference? outputDir)
+		{
+			string? outputPath;
+			if (Properties.TryGetValue("OutputPath", out outputPath))
+			{
+				outputDir = DirectoryReference.Combine(ProjectPath.Directory, outputPath);
+				return true;
+			}
+			else if (IsNetCoreProject())
+			{
+				string configuration = Properties.ContainsKey("Configuration") ? Properties["Configuration"] : "Development";
+				outputDir = DirectoryReference.Combine(ProjectPath.Directory, "bin", configuration, Properties["TargetFramework"]);
+				return true;
+			}
+			else
+			{
+				outputDir = null;
+				return false;
+			}
+		}
+
+		/// <summary>
+		/// Returns the assembly name used by this project
+		/// </summary>
+		/// <returns></returns>
+		public bool TryGetAssemblyName([NotNullWhen(true)] out string? assemblyName)
+		{
+			if (Properties.TryGetValue("AssemblyName", out assemblyName))
+			{
+				return true;
+			}
+			else if (IsNetCoreProject())
+			{
+				assemblyName = ProjectPath.GetFileNameWithoutExtension();
+				return true;
+			}
+			return false;
+		}
+
+		/// <summary>
+		/// Finds all build products from this project. This includes content and other assemblies marked to be copied local.
+		/// </summary>
+		/// <param name="outputDir">The output directory</param>
+		/// <param name="buildProducts">Receives the set of build products</param>
+		/// <param name="projectFileToInfo">Map of project file to information, to resolve build products from referenced projects copied locally</param>
+		public void FindBuildProducts(DirectoryReference outputDir, HashSet<FileReference> buildProducts, Dictionary<FileReference, CsProjectInfo> projectFileToInfo)
+		{
+			// Add the standard build products
+			FindCompiledBuildProducts(outputDir, buildProducts);
+
+			// Add the referenced assemblies which are marked to be copied into the output directory. This only happens for the main project, and does not happen for referenced projects.
+			foreach(KeyValuePair<FileReference, bool> reference in References)
+			{
+				if (reference.Value)
+				{
+					FileReference outputFile = FileReference.Combine(outputDir, reference.Key.GetFileName());
+					AddReferencedAssemblyAndSupportFiles(outputFile, buildProducts);
+				}
+			}
+
+			// Copy the build products for any referenced projects. Note that this does NOT operate recursively.
+			foreach(KeyValuePair<FileReference, bool> projectReference in ProjectReferences)
+			{
+				CsProjectInfo? otherProjectInfo;
+				if(projectFileToInfo.TryGetValue(projectReference.Key, out otherProjectInfo))
+				{
+					otherProjectInfo.FindCompiledBuildProducts(outputDir, buildProducts);
+				}
+			}
+
+			// Add any copied content. This DOES operate recursively.
+			FindCopiedContent(outputDir, buildProducts, projectFileToInfo);
+		}
+
+		/// <summary>
+		/// Determines all the compiled build products (executable, etc...) directly built from this project.
+		/// </summary>
+		/// <param name="outputDir">The output directory</param>
+		/// <param name="buildProducts">Receives the set of build products</param>
+		public void FindCompiledBuildProducts(DirectoryReference outputDir, HashSet<FileReference> buildProducts)
+		{
+			string? outputType, assemblyName;
+			if (Properties.TryGetValue("OutputType", out outputType) && TryGetAssemblyName(out assemblyName))
+			{
+				switch (outputType)
+				{
+					case "Exe":
+					case "WinExe":
+						string executableExtension = RuntimeInformation.IsOSPlatform(OSPlatform.Windows) ? ".exe" : "";
+						buildProducts.Add(FileReference.Combine(outputDir, assemblyName + executableExtension));
+						// dotnet outputs a apphost executable and a dll with the actual assembly
+						AddOptionalBuildProduct(FileReference.Combine(outputDir, assemblyName + ".dll"), buildProducts);
+						AddOptionalBuildProduct(FileReference.Combine(outputDir, assemblyName + ".pdb"), buildProducts);
+						AddOptionalBuildProduct(FileReference.Combine(outputDir, assemblyName + ".exe.config"), buildProducts);
+						AddOptionalBuildProduct(FileReference.Combine(outputDir, assemblyName + ".exe.mdb"), buildProducts);
+						break;
+					case "Library":
+						buildProducts.Add(FileReference.Combine(outputDir, assemblyName + ".dll"));
+						AddOptionalBuildProduct(FileReference.Combine(outputDir, assemblyName + ".pdb"), buildProducts);
+						AddOptionalBuildProduct(FileReference.Combine(outputDir, assemblyName + ".dll.config"), buildProducts);
+						AddOptionalBuildProduct(FileReference.Combine(outputDir, assemblyName + ".dll.mdb"), buildProducts);
+						break;
+				}
+			}
+		}
+
+		/// <summary>
+		/// Finds all content which will be copied into the output directory for this project. This includes content from any project references as "copy local" recursively (though MSBuild only traverses a single reference for actual binaries, in such cases)
+		/// </summary>
+		/// <param name="outputDir">The output directory</param>
+		/// <param name="outputFiles">Receives the set of build products</param>
+		/// <param name="projectFileToInfo">Map of project file to information, to resolve build products from referenced projects copied locally</param>
+		private void FindCopiedContent(DirectoryReference outputDir, HashSet<FileReference> outputFiles, Dictionary<FileReference, CsProjectInfo> projectFileToInfo)
+		{
+			// Copy any referenced projects too.
+			foreach(KeyValuePair<FileReference, bool> projectReference in ProjectReferences)
+			{
+				CsProjectInfo? otherProjectInfo;
+				if(projectFileToInfo.TryGetValue(projectReference.Key, out otherProjectInfo))
+				{
+					otherProjectInfo.FindCopiedContent(outputDir, outputFiles, projectFileToInfo);
+				}
+			}
+
+			// Add the content which is copied to the output directory
+			foreach (KeyValuePair<FileReference, bool> contentReference in ContentReferences)
+			{
+				FileReference contentFile = contentReference.Key;
+				if (contentReference.Value)
+				{
+					outputFiles.Add(FileReference.Combine(outputDir, contentFile.GetFileName()));
+				}
+			}
+		}
+
+		/// <summary>
+		/// Adds the given file and any additional build products to the output set
+		/// </summary>
+		/// <param name="outputFile">The assembly to add</param>
+		/// <param name="outputFiles">Set to receive the file and support files</param>
+		public static void AddReferencedAssemblyAndSupportFiles(FileReference outputFile, HashSet<FileReference> outputFiles)
+		{
+			outputFiles.Add(outputFile);
+
+			FileReference symbolFile = outputFile.ChangeExtension(".pdb");
+			if (FileReference.Exists(symbolFile))
+			{
+				outputFiles.Add(symbolFile);
+			}
+
+			FileReference documentationFile = outputFile.ChangeExtension(".xml");
+			if (FileReference.Exists(documentationFile))
+			{
+				outputFiles.Add(documentationFile);
+			}
+		}
+
+		/// <summary>
+		/// Determines if a TargetFramework is a .NET core framework
+		/// </summary>
+		/// <returns>True if the TargetFramework is a .NET core framework</returns>
+		static bool IsDotNETCoreFramework(string targetFramework)
+		{
+			if (targetFramework.ToLower().Contains("netstandard", StringComparison.Ordinal) || targetFramework.ToLower().Contains("netcoreapp", StringComparison.Ordinal))
+			{
+				return true;
+			}
+			else if (targetFramework.StartsWith("net", StringComparison.OrdinalIgnoreCase))
+			{
+				string[] versionSplit = targetFramework.Substring(3).Split('.');
+				if (versionSplit.Length >= 1 && Int32.TryParse(versionSplit[0], out int majorVersion))
+				{
+					return majorVersion >= 5;
+				}
+			}
+			return false;
+		}
+
+		/// <summary>
+		/// Determines if this project is a .NET core project
+		/// </summary>
+		/// <returns>True if the project is a .NET core project</returns>
+		public bool IsDotNETCoreProject()
+		{
+			if (Properties.TryGetValue("TargetFramework", out string? targetFramework))
+			{
+				return IsDotNETCoreFramework(targetFramework);
+			}
+
+			if (Properties.TryGetValue("TargetFrameworks", out string? targetFrameworks))
+			{
+				return targetFramework?.Split(';', StringSplitOptions.RemoveEmptyEntries).Any(x => IsDotNETCoreFramework(x.Trim())) ?? false;
+			}
+
+			return false;
+		}
+
+		/// <summary>
+		/// Adds a build product to the output list if it exists
+		/// </summary>
+		/// <param name="buildProduct">The build product to add</param>
+		/// <param name="buildProducts">List of output build products</param>
+		public static void AddOptionalBuildProduct(FileReference buildProduct, HashSet<FileReference> buildProducts)
+		{
+			if (FileReference.Exists(buildProduct))
+			{
+				buildProducts.Add(buildProduct);
+			}
+		}
+
+		/// <summary>
+		/// Parses supported elements from the children of the provided note. May recurse
+		/// for conditional elements.
+		/// </summary>
+		/// <param name="node"></param>
+		/// <param name="projectInfo"></param>
+		static void ParseNode(XmlNode node, CsProjectInfo projectInfo)
+		{
+			foreach (XmlElement element in node.ChildNodes.OfType<XmlElement>())
+			{
+				switch (element.Name)
+				{
+					case "Import":
+						if (EvaluateCondition(element, projectInfo))
+						{
+							ParseImportProject(element, projectInfo);
+						}
+						break;
+					case "PropertyGroup":
+						if (EvaluateCondition(element, projectInfo))
+						{
+							ParsePropertyGroup(element, projectInfo);
+						}
+						break;
+					case "ItemGroup":
+						if (EvaluateCondition(element, projectInfo))
+						{
+							ParseItemGroup(projectInfo.ProjectPath.Directory, element, projectInfo);
+						}
+						break;
+					case "Choose":
+					case "When":
+						if (EvaluateCondition(element, projectInfo))
+						{
+							ParseNode(element, projectInfo);
+						}
+						break;
+				}
+			}
+		}
+
+		/// <summary>
+		/// Reads project information for the given file.
+		/// </summary>
+		/// <param name="file">The project file to read</param>
+		/// <param name="properties">Initial set of property values</param>
+		/// <returns>The parsed project info</returns>
+		public static CsProjectInfo Read(FileReference file, Dictionary<string, string> properties)
+		{
+			CsProjectInfo? project;
+			if(!TryRead(file, properties, out project))
+			{
+				throw new Exception(String.Format("Unable to read '{0}'", file));
+			}
+			return project;
+		}
+
+		/// <summary>
+		/// Attempts to read project information for the given file.
+		/// </summary>
+		/// <param name="file">The project file to read</param>
+		/// <param name="properties">Initial set of property values</param>
+		/// <param name="outProjectInfo">If successful, the parsed project info</param>
+		/// <returns>True if the project was read successfully, false otherwise</returns>
+		public static bool TryRead(FileReference file, Dictionary<string, string> properties, [NotNullWhen(true)] out CsProjectInfo? outProjectInfo)
+		{
+			// Read the project file
+			XmlDocument document = new XmlDocument();
+			document.Load(file.FullName);
+
+			// Check the root element is the right type
+			//			HashSet<FileReference> ProjectBuildProducts = new HashSet<FileReference>();
+			if (document.DocumentElement!.Name != "Project")
+			{
+				outProjectInfo = null;
+				return false;
+			}
+
+			// Parse the basic structure of the document, updating properties and recursing into other referenced projects as we go
+			CsProjectInfo projectInfo = new CsProjectInfo(properties, file);
+
+			// Parse elements in the root node
+			try
+			{
+				ParseNode(document.DocumentElement, projectInfo);
+			}
+			catch (Exception ex)
+			{
+				throw new CsProjectParseException($"Error parsing {file}: {ex.Message}", ex);
+			}
+
+			// Return the complete project
+			outProjectInfo = projectInfo;
+			return true;
+		}
+
+		/// <summary>
+		/// Parses a 'Import' emenebt
+		/// </summary>
+		/// <param name="element">The element</param>
+		/// <param name="projectInfo">Dictionary mapping property names to values</param>
+		static void ParseImportProject(XmlElement element, CsProjectInfo projectInfo)
+		{
+			string importProjectStr = element.GetAttribute("Project");
+			foreach (Match match in Regex.Matches(importProjectStr, @"(\$\(.*?\))").OfType<Match>())
+			{
+				importProjectStr = importProjectStr.Replace(match.Value, projectInfo.Properties.GetValueOrDefault(match.Value.Substring(2, match.Length - 3), String.Empty), StringComparison.Ordinal);
+			}
+			FileReference importProject = Path.IsPathFullyQualified(importProjectStr) ? new FileReference(importProjectStr) : FileReference.Combine(projectInfo.ProjectPath.Directory, importProjectStr);
+			if (!FileReference.Exists(importProject))
+			{
+				return;
+			}
+			XmlDocument document = new XmlDocument();
+			document.Load(importProject.FullName);
+			if (document.DocumentElement?.Name == "Project")
+			{
+				projectInfo.Properties.TryGetValue("MSBuildThisFileDirectory", out string? msBuildThisFileDirectory);
+				projectInfo.Properties.Remove("MSBuildThisFileDirectory");
+				projectInfo.Properties.Add("MSBuildThisFileDirectory", importProject.Directory.FullName);
+				try
+				{
+					ParseNode(document.DocumentElement, projectInfo);
+				}
+				catch (Exception)
+				{
+					// Ignore issues in imported .props files
+				}
+				projectInfo.Properties.Remove("MSBuildThisFileDirectory");
+				if (msBuildThisFileDirectory != null)
+				{
+					projectInfo.Properties.Add("MSBuildThisFileDirectory", msBuildThisFileDirectory);
+				}
+			}
+		}
+
+		/// <summary>
+		/// Parses a 'PropertyGroup' element.
+		/// </summary>
+		/// <param name="parentElement">The parent 'PropertyGroup' element</param>
+		/// <param name="projectInfo">Dictionary mapping property names to values</param>
+		static void ParsePropertyGroup(XmlElement parentElement, CsProjectInfo projectInfo)
+		{
+			// We need to know the overridden output type and output path for the selected configuration.
+			foreach (XmlElement element in parentElement.ChildNodes.OfType<XmlElement>())
+			{
+				// Common properties used by UnrealEngine.csproj.props, manually handle because the parsing is awful
+				switch (element.Name)
+				{
+					case "IsLinux": projectInfo.Properties[element.Name] = OperatingSystem.IsLinux().ToString(); continue;
+					case "IsMacOS": projectInfo.Properties[element.Name] = OperatingSystem.IsMacOS().ToString(); continue;
+					case "IsWindows": projectInfo.Properties[element.Name] = OperatingSystem.IsWindows().ToString(); continue;
+					case "IsX64": projectInfo.Properties[element.Name] = (RuntimeInformation.OSArchitecture == Architecture.X64).ToString(); continue;
+					case "IsArm64": projectInfo.Properties[element.Name] = (RuntimeInformation.OSArchitecture == Architecture.Arm64).ToString(); continue;
+					case "OSArchitecture": projectInfo.Properties[element.Name] = RuntimeInformation.OSArchitecture.ToString(); continue;
+					case "EngineDirectory": projectInfo.Properties.TryAdd(element.Name, String.Empty); continue;
+					case "EngineDir":
+					case "EnginePath":
+						{
+							projectInfo.Properties[element.Name] = String.Empty;
+							if (projectInfo.Properties.TryGetValue("EngineDirectory", out string? engineDirectory) && !String.IsNullOrEmpty(engineDirectory))
+							{
+								projectInfo.Properties[element.Name] = engineDirectory;
+							}
+							continue;
+						}
+					case "IsEngineProject":
+					{
+						projectInfo.Properties[element.Name] = "False";
+						if (projectInfo.Properties.TryGetValue("EngineDirectory", out string? engineDirectory) && !String.IsNullOrEmpty(engineDirectory))
+						{
+							projectInfo.Properties[element.Name] = projectInfo.ProjectPath.IsUnderDirectory(new DirectoryReference(engineDirectory)).ToString();
+						}
+						continue;
+					}
+					case "IsAutomationProject": projectInfo.Properties[element.Name] = projectInfo.ProjectPath.FullName.EndsWith(".automation.csproj", StringComparison.OrdinalIgnoreCase).ToString(); continue;
+					case "HasAssemblyInfo": projectInfo.Properties[element.Name] = FileReference.Exists(FileReference.Combine(projectInfo.ProjectPath.Directory, "Properties", "AssemblyInfo.cs")).ToString(); continue;
+					default: break;
+				}
+
+				if (EvaluateCondition(element, projectInfo))
+				{
+					projectInfo.Properties[element.Name] = ExpandProperties(element.InnerText, projectInfo.Properties);
+				}
+			}
+		}
+
+		/// <summary>
+		/// Parses an 'ItemGroup' element.
+		/// </summary>
+		/// <param name="baseDirectory">Base directory to resolve relative paths against</param>
+		/// <param name="parentElement">The parent 'ItemGroup' element</param>
+		/// <param name="projectInfo">Project info object to be updated</param>
+		static void ParseItemGroup(DirectoryReference baseDirectory, XmlElement parentElement, CsProjectInfo projectInfo)
+		{
+			// Parse any external assembly references
+			foreach (XmlElement itemElement in parentElement.ChildNodes.OfType<XmlElement>())
+			{
+				switch (itemElement.Name)
+				{
+					case "Reference":
+						// Reference to an external assembly
+						if (EvaluateCondition(itemElement, projectInfo))
+						{
+							ParseReference(baseDirectory, itemElement, projectInfo.References);
+						}
+						break;
+					case "ProjectReference":
+						// Reference to another project
+						if (EvaluateCondition(itemElement, projectInfo))
+						{
+							ParseProjectReference(baseDirectory, itemElement, projectInfo.Properties, projectInfo.ProjectReferences);
+						}
+						break;
+					case "Compile":
+						// Reference to a file
+						if (EvaluateCondition(itemElement, projectInfo))
+						{
+							ParseCompileReference(baseDirectory, itemElement, projectInfo.CompileReferences);
+						}
+						break;
+					case "Content":
+					case "None":
+						// Reference to a file
+						if (EvaluateCondition(itemElement, projectInfo))
+						{
+							ParseContent(baseDirectory, itemElement, projectInfo.ContentReferences);
+						}
+						break;
+				}
+			}
+		}
+
+		/// <summary>
+		/// Parses an assembly reference from a given 'Reference' element
+		/// </summary>
+		/// <param name="baseDirectory">Directory to resolve relative paths against</param>
+		/// <param name="parentElement">The parent 'Reference' element</param>
+		/// <param name="references">Dictionary of project files to a bool indicating whether the assembly should be copied locally to the referencing project.</param>
+		static void ParseReference(DirectoryReference baseDirectory, XmlElement parentElement, Dictionary<FileReference, bool> references)
+		{
+			string? hintPath = UnescapeString(GetChildElementString(parentElement, "HintPath", null));
+			if (!String.IsNullOrEmpty(hintPath))
+			{
+				// Don't include embedded assemblies; they aren't referenced externally by the compiled executable
+				bool bEmbedInteropTypes = GetChildElementBoolean(parentElement, "EmbedInteropTypes", false);
+				if(!bEmbedInteropTypes)
+				{
+					if (!OperatingSystem.IsWindows() && hintPath.Contains('\\', StringComparison.Ordinal))
+					{
+						hintPath = hintPath.Replace('\\', Path.DirectorySeparatorChar);
+					}
+					FileReference assemblyFile = Path.IsPathFullyQualified(hintPath) ? new FileReference(hintPath) : FileReference.Combine(baseDirectory, hintPath);
+					bool bPrivate = GetChildElementBoolean(parentElement, "Private", true);
+					references.Add(assemblyFile, bPrivate);
+				}
+			}
+		}
+
+		/// <summary>
+		/// Parses a project reference from a given 'ProjectReference' element
+		/// </summary>
+		/// <param name="baseDirectory">Directory to resolve relative paths against</param>
+		/// <param name="parentElement">The parent 'ProjectReference' element</param>
+		/// <param name="properties">Dictionary of properties for parsing the file</param>
+		/// <param name="projectReferences">Dictionary of project files to a bool indicating whether the outputs of the project should be copied locally to the referencing project.</param>
+		static void ParseProjectReference(DirectoryReference baseDirectory, XmlElement parentElement, Dictionary<string, string> properties, Dictionary<FileReference, bool> projectReferences)
+		{
+			string? includePath = UnescapeString(parentElement.GetAttribute("Include"));
+			if (!String.IsNullOrEmpty(includePath))
+			{
+				includePath = ExpandProperties(includePath, properties);
+				if (!OperatingSystem.IsWindows() && includePath.Contains('\\', StringComparison.Ordinal))
+				{
+					includePath = includePath.Replace('\\', Path.DirectorySeparatorChar);
+				}
+				FileReference projectFile = Path.IsPathFullyQualified(includePath) ? new FileReference(includePath) : FileReference.Combine(baseDirectory, includePath);
+				bool bPrivate = GetChildElementBoolean(parentElement, "Private", true);
+				projectReferences[projectFile] = bPrivate;
+			}
+		}
+
+		/// <summary>
+		/// Finds all directories but filters out certain folders
+		/// </summary>
+		static void GetDirectories(DirectoryReference dir, ConcurrentBag<DirectoryReference> list, SearchOption option, ThreadPoolWorkQueue queue)
+		{
+			foreach (DirectoryReference child in DirectoryReference.EnumerateDirectories(dir))
+			{
+				string name = child.GetDirectoryName();
+				if (name == "Intermediate" || name == "Binaries")
+				{
+					continue;
+				}
+
+				list.Add(child);
+				if (option != SearchOption.AllDirectories)
+				{
+					continue;
+				}
+
+				queue.Enqueue(() => { GetDirectories(child, list, option, queue); });
+			}
+		}
+
+		/// recursive helper used by the function below that will append RemainingComponents one by one to ExistingPath, 
+		/// expanding wildcards as necessary. The complete list of files that match the complete path is returned out OutFoundFiles
+		static void ProcessPathComponents(DirectoryReference existingPath, IEnumerable<string> remainingComponents, List<FileReference> outFoundFiles)
+		{
+			if (!remainingComponents.Any())
+			{
+				return;
+			}
+
+			// take a look at the first component
+			string currentComponent = remainingComponents.First();
+			remainingComponents = remainingComponents.Skip(1);
+
+			// If no other components then this is either a file pattern or a greedy pattern
+			if (!remainingComponents.Any())
+			{
+				// ** means include all files under this tree, so enumerate them all
+				if (currentComponent.Contains("**", StringComparison.Ordinal))
+				{
+					outFoundFiles.AddRange(DirectoryReference.EnumerateFiles(existingPath, "*", SearchOption.AllDirectories));
+				}
+				else
+				{
+					// easy, a regular path with a file that may or may not be a wildcard
+					outFoundFiles.AddRange(DirectoryReference.EnumerateFiles(existingPath, currentComponent));
+				}
+			}
+			else
+			{
+				// new component contains a wildcard, and based on the above we know there are more entries so find 
+				// matching directories
+				if (currentComponent.Contains('*', StringComparison.Ordinal))
+				{
+					// ** means all directories, no matter how deep
+					SearchOption option = currentComponent == "**" ? SearchOption.AllDirectories : SearchOption.TopDirectoryOnly;
+
+					IEnumerable<DirectoryReference> directories;
+					if (option == SearchOption.TopDirectoryOnly)
+					{
+						directories = DirectoryReference.EnumerateDirectories(existingPath, currentComponent, option);
+					}
+					else
+					{
+						// Let's go wide because we potentially will traverse a lot of directories
+						// and we also want to filter out Intermediate, Binaries, etc
+
+						ConcurrentBag<DirectoryReference> bag = new();
+						using (ThreadPoolWorkQueue queue = new ThreadPoolWorkQueue())
+						{
+							GetDirectories(existingPath, bag, option, queue);
+						}
+
+						// Add current path
+						bag.Add(existingPath);
+
+						// Consume potential folders after **   (like ..\Engine\**\Build\Gauntlet\**\*.cs
+						string subPath = "";
+						while (remainingComponents.Any())
+						{
+							string component = remainingComponents.First();
+							if (component[0] == '*')
+								break;
+							subPath += Path.DirectorySeparatorChar + component;
+							remainingComponents = remainingComponents.Skip(1);
+						}
+
+						// Match folders against end of Directories found
+						if (subPath.Length > 0)
+						{
+							List<DirectoryReference> matchingDirs = [];
+
+							foreach (DirectoryReference dir in bag)
+							{
+								if (dir.FullName.EndsWith(subPath, StringComparison.Ordinal))
+								{
+									matchingDirs.Add(dir);
+								}
+							}
+							directories = matchingDirs;
+						}
+						else
+						{
+							// .. otherwise just keep going (remaining component is likely *.cs)
+							directories = bag;
+						}
+					}
+
+					foreach (DirectoryReference dir in directories)
+					{
+						ProcessPathComponents(dir, remainingComponents, outFoundFiles);
+					}
+				}
+				else
+				{
+					// add this component to our path and recurse.
+					existingPath = DirectoryReference.Combine(existingPath, currentComponent);
+
+					// but... we can just take all the next components that don't have wildcards in them instead of recursing
+					// into each one!
+					IEnumerable<string> nonWildCardComponents = remainingComponents.TakeWhile(c => !c.Contains('*', StringComparison.Ordinal));
+					remainingComponents = remainingComponents.Skip(nonWildCardComponents.Count());
+
+					existingPath = DirectoryReference.Combine(existingPath, nonWildCardComponents.ToArray());
+
+					if (Directory.Exists(existingPath.FullName))
+					{
+						ProcessPathComponents(existingPath, remainingComponents, outFoundFiles);
+					}
+				}
+			}
+		}
+
+		/// <summary>
+		/// Finds all files in the provided path, which may be a csproj wildcard specification.
+		/// E.g. The following are all valid
+		/// Foo/Bar/Item.cs
+		/// Foo/Bar/*.cs
+		/// Foo/*/Item.cs
+		/// Foo/*/*.cs
+		/// Foo/**
+		/// (the last means include all files under the path).
+		/// </summary>
+		/// <param name="inPath">Path specifier to process</param>
+		/// <returns></returns>
+		static IEnumerable<FileReference> FindMatchingFiles(FileReference inPath)
+		{
+			List<FileReference> foundFiles = [];		
+
+			// split off the drive root
+			string driveRoot = Path.GetPathRoot(inPath.FullName)!;
+
+			// break the rest of the path into components
+			string[] pathComponents = inPath.FullName.Substring(driveRoot.Length).Split([Path.DirectorySeparatorChar, Path.AltDirectorySeparatorChar]);
+
+			// Process all the components recursively
+			ProcessPathComponents(new DirectoryReference(driveRoot), pathComponents, foundFiles);
+
+			return foundFiles;
+		}
+
+		/// <summary>
+		/// Parses a a compile file path from a given 'Compile' element
+		/// </summary>
+		/// <param name="baseDirectory">Directory to resolve relative paths against</param>
+		/// <param name="parentElement">The parent 'ProjectReference' element</param>
+		/// <param name="compileReferences">List of source files.</param>
+		static void ParseCompileReference(DirectoryReference baseDirectory, XmlElement parentElement, List<FileReference> compileReferences)
+		{
+			string? includePath = UnescapeString(parentElement.GetAttribute("Include"));
+			if (!String.IsNullOrEmpty(includePath))
+			{
+				if (!OperatingSystem.IsWindows() && includePath.Contains('\\', StringComparison.Ordinal))
+				{
+					includePath = includePath.Replace('\\', Path.DirectorySeparatorChar);
+				}
+				FileReference sourceFile = Path.IsPathFullyQualified(includePath) ? new FileReference(includePath) : FileReference.Combine(baseDirectory, includePath);
+
+				if (sourceFile.FullName.Contains('*', StringComparison.Ordinal))
+				{
+					compileReferences.AddRange(FindMatchingFiles(sourceFile));
+				}
+				else
+				{
+					compileReferences.Add(sourceFile);
+				}
+			}
+		}
+
+		/// <summary>
+		/// Parses a file path from a given 'Content' element
+		/// </summary>
+		/// <param name="baseDirectory">Directory to resolve relative paths against</param>
+		/// <param name="parentElement">The parent 'Content' element</param>
+		/// <param name="contents">Dictionary of project files to a bool indicating whether the assembly should be copied locally to the referencing project.</param>
+		static void ParseContent(DirectoryReference baseDirectory, XmlElement parentElement, Dictionary<FileReference, bool> contents)
+		{
+			string? includePath = UnescapeString(parentElement.GetAttribute("Include"));
+			if (!String.IsNullOrEmpty(includePath))
+			{
+				string? copyTo = GetChildElementString(parentElement, "CopyToOutputDirectory", null);
+				bool shouldCopy = !String.IsNullOrEmpty(copyTo) && (copyTo.Equals("Always", StringComparison.OrdinalIgnoreCase) || copyTo.Equals("PreserveNewest", StringComparison.OrdinalIgnoreCase));
+				if (!OperatingSystem.IsWindows() && includePath.Contains('\\', StringComparison.Ordinal))
+				{
+					includePath = includePath.Replace('\\', Path.DirectorySeparatorChar);
+				}
+				FileReference contentFile = Path.IsPathFullyQualified(includePath) ? new FileReference(includePath) : FileReference.Combine(baseDirectory, includePath);
+				contents.TryAdd(contentFile, shouldCopy);
+			}
+		}
+
+		/// <summary>
+		/// Reads the inner text of a child XML element
+		/// </summary>
+		/// <param name="parentElement">The parent element to check</param>
+		/// <param name="name">Name of the child element</param>
+		/// <param name="defaultValue">Default value to return if the child element is missing</param>
+		/// <returns>The contents of the child element, or default value if it's not present</returns>
+		static string? GetChildElementString(XmlElement parentElement, string name, string? defaultValue)
+		{
+			XmlElement? childElement = parentElement.ChildNodes.OfType<XmlElement>().FirstOrDefault(x => x.Name == name);
+			if (childElement == null)
+			{
+				return defaultValue;
+			}
+			else
+			{
+				return childElement.InnerText ?? defaultValue;
+			}
+		}
+
+		/// <summary>
+		/// Read a child XML element with the given name, and parse it as a boolean.
+		/// </summary>
+		/// <param name="parentElement">Parent element to check</param>
+		/// <param name="name">Name of the child element to look for</param>
+		/// <param name="defaultValue">Default value to return if the element is missing or not a valid bool</param>
+		/// <returns>The parsed boolean, or the default value</returns>
+		static bool GetChildElementBoolean(XmlElement parentElement, string name, bool defaultValue)
+		{
+			string? value = GetChildElementString(parentElement, name, null);
+			if (value == null)
+			{
+				return defaultValue;
+			}
+			else if (value.Equals("True", StringComparison.OrdinalIgnoreCase))
+			{
+				return true;
+			}
+			else if (value.Equals("False", StringComparison.OrdinalIgnoreCase))
+			{
+				return false;
+			}
+			else
+			{
+				return defaultValue;
+			}
+		}
+
+		record class ConditionContext(DirectoryReference BaseDir);
+
+		/// <summary>
+		/// Evaluate whether the optional MSBuild condition on an XML element evaluates to true. Currently only supports 'ABC' == 'DEF' style expressions, but can be expanded as needed.
+		/// </summary>
+		/// <param name="element">The XML element to check</param>
+		/// <param name="projectInfo">Dictionary mapping from property names to values.</param>
+		/// <returns></returns>
+		static bool EvaluateCondition(XmlElement element, CsProjectInfo projectInfo)
+		{
+			// Read the condition attribute. If it's not present, assume it evaluates to true.
+			string condition = element.GetAttribute("Condition");
+			if (String.IsNullOrEmpty(condition))
+			{
+				return true;
+			}
+
+			// Expand all the properties
+			condition = ExpandProperties(condition, projectInfo.Properties);
+
+			// Parse literal true/false values
+			bool outResult;
+			if (Boolean.TryParse(condition, out outResult))
+			{
+				return outResult;
+			}
+
+			// Tokenize the condition
+			string[] tokens = Tokenize(condition);
+
+			// Try to evaluate it. We only support a very limited class of condition expressions at the moment, but it's enough to parse standard projects
+			int tokenIdx = 0;
+
+			ConditionContext context = new ConditionContext(projectInfo.ProjectPath.Directory);
+			try
+			{
+				bool bResult = CoerceToBool(EvaluateLogicalAnd(context, tokens, ref tokenIdx));
+				if (tokenIdx != tokens.Length)
+				{
+					throw new CsProjectParseException("Unexpected tokens at end of condition");
+				}
+				return bResult;
+			}
+			catch (CsProjectParseException ex)
+			{
+				throw new CsProjectParseException($"{ex.Message} while parsing {element} in project file {projectInfo.ProjectPath}", ex);
+			}
+		}
+
+		static string EvaluateLogicalAnd(ConditionContext context, string[] tokens, ref int tokenIdx)
+		{
+			string result = EvaluateLogicalOr(context, tokens, ref tokenIdx);
+			while (tokenIdx < tokens.Length && tokens[tokenIdx].Equals("And", StringComparison.OrdinalIgnoreCase))
+			{
+				tokenIdx++;
+				string rhs = EvaluateEquality(context, tokens, ref tokenIdx);
+				result = (CoerceToBool(result) && CoerceToBool(rhs)).ToString();
+			}
+			return result;
+		}
+
+		static string EvaluateLogicalOr(ConditionContext context, string[] tokens, ref int tokenIdx)
+		{
+			string result = EvaluateEquality(context, tokens, ref tokenIdx);
+			while (tokenIdx < tokens.Length && tokens[tokenIdx].Equals("Or", StringComparison.OrdinalIgnoreCase))
+			{
+				tokenIdx++;
+				string rhs = EvaluateEquality(context, tokens, ref tokenIdx);
+				result = (CoerceToBool(result) || CoerceToBool(rhs)).ToString();
+			}
+			return result;
+		}
+
+		static string EvaluateEquality(ConditionContext context, string[] tokens, ref int tokenIdx)
+		{
+			// Otherwise try to parse an equality or inequality expression
+			string lhs = EvaluateValue(context, tokens, ref tokenIdx);
+			if (tokenIdx < tokens.Length)
+			{
+				if (tokens[tokenIdx] == "==")
+				{
+					tokenIdx++;
+					string rhs = EvaluateValue(context, tokens, ref tokenIdx);
+					return lhs.Equals(rhs, StringComparison.OrdinalIgnoreCase).ToString();
+				}
+				else if (tokens[tokenIdx] == "!=")
+				{
+					tokenIdx++;
+					string rhs = EvaluateValue(context, tokens, ref tokenIdx);
+					return lhs.Equals(rhs, StringComparison.OrdinalIgnoreCase).ToString();
+				}
+			}
+			return lhs;
+		}
+
+		static string EvaluateValue(ConditionContext context, string[] tokens, ref int tokenIdx)
+		{
+			// Handle Exists('Platform\Windows\Gauntlet.TargetDeviceWindows.cs')
+			if (tokens[tokenIdx].Equals("Exists", StringComparison.OrdinalIgnoreCase))
+			{
+				if (tokenIdx + 3 >= tokens.Length || !tokens[tokenIdx + 1].Equals("(", StringComparison.Ordinal) || !tokens[tokenIdx + 3].Equals(")", StringComparison.Ordinal))
+				{
+					throw new CsProjectParseException("Invalid 'Exists' expression", null);
+				}
+
+				// remove all quotes, apostrophes etc that are either tokens or wrap tokens (The Tokenize() function is a bit suspect).
+				string path = tokens[tokenIdx + 2].Trim('\'', '(', ')', '{', '}', '[', ']');
+				tokenIdx += 4;
+
+				FileSystemReference dependency = DirectoryReference.Combine(context.BaseDir, path);
+				bool exists = File.Exists(dependency.FullName) || Directory.Exists(dependency.FullName);
+				return exists.ToString();
+			}
+
+			// Handle negation
+			if (tokens[tokenIdx].Equals("!", StringComparison.Ordinal))
+			{
+				tokenIdx++;
+				bool value = CoerceToBool(EvaluateValue(context, tokens, ref tokenIdx));
+				return (!value).ToString();
+			}
+
+			// Handle subexpressions
+			if (tokens[tokenIdx].Equals("(", StringComparison.Ordinal))
+			{
+				tokenIdx++;
+
+				string result = EvaluateLogicalAnd(context, tokens, ref tokenIdx);
+				if (!tokens[tokenIdx].Equals(")", StringComparison.Ordinal))
+				{
+					throw new CsProjectParseException("Missing ')'", null);
+				}
+				tokenIdx++;
+
+				return result;
+			}
+
+			return tokens[tokenIdx++];
+		}
+
+		static bool CoerceToBool(string value)
+		{
+			return !value.Equals("false", StringComparison.OrdinalIgnoreCase) && !value.Equals("0", StringComparison.Ordinal);
+		}
+
+		/// <summary>
+		/// Expand MSBuild properties within a string. If referenced properties are not in this dictionary, the process' environment variables are expanded. Unknown properties are expanded to an empty string.
+		/// </summary>
+		/// <param name="text">The input string to expand</param>
+		/// <param name="properties">Dictionary mapping from property names to values.</param>
+		/// <returns>String with all properties expanded.</returns>
+		static string ExpandProperties(string text, Dictionary<string, string> properties)
+		{
+			string newText = text;
+			for (int idx = newText.IndexOf("$(", StringComparison.Ordinal); idx != -1; idx = newText.IndexOf("$(", idx, StringComparison.Ordinal))
+			{
+				// Find the end of the variable name, accounting for changes in scope
+				int endIdx = idx + 2;
+				for(int depth = 1; depth > 0; endIdx++)
+				{
+					if(endIdx == newText.Length)
+					{
+						throw new Exception("Encountered end of string while expanding properties");
+					}
+					else if(newText[endIdx] == '(')
+					{
+						depth++;
+					}
+					else if(newText[endIdx] == ')')
+					{
+						depth--;
+					}
+				}
+
+				// Convert the property name to tokens
+				string[] tokens = Tokenize(newText.Substring(idx + 2, (endIdx - 1) - (idx + 2)));
+
+				// Make sure the first token is a valid property name
+				if(tokens.Length == 0 || !(Char.IsLetter(tokens[0][0]) || tokens[0][0] == '_' || tokens[0][0] == '[' ))
+				{
+					throw new Exception(String.Format("Invalid property name '{0}' in .csproj file", tokens[0]));
+				}
+
+				// Find the value for it, either from the dictionary or the environment block
+				string value;
+				if (properties.TryGetValue(tokens[0], out string? retrievedValue))
+				{
+					value = retrievedValue;
+				}
+				else
+				{
+					value = Environment.GetEnvironmentVariable(tokens[0]) ?? "";
+				}
+
+				// Evaluate any functions within it
+				int tokenIdx = 1;
+				while(tokenIdx + 3 < tokens.Length && tokens[tokenIdx] == "." && tokens[tokenIdx + 2] == "(")
+				{
+					// Read the method name
+					string methodName = tokens[tokenIdx + 1];
+
+					// Skip to the first argument
+					tokenIdx += 3;
+
+					// Parse any arguments
+					List<object> arguments = [];
+					if(tokens[tokenIdx] != ")")
+					{
+						arguments.Add(ParseArgument(tokens[tokenIdx]));
+						tokenIdx++;
+
+						while(tokenIdx + 1 < tokens.Length && tokens[tokenIdx] == ",")
+						{
+							arguments.Add(ParseArgument(tokens[tokenIdx + 2]));
+							tokenIdx += 2;
+						}
+
+						if(tokens[tokenIdx] != ")")
+						{
+							throw new Exception("Missing closing parenthesis in condition");
+						}
+					}
+
+					// Skip over the closing parenthesis
+					tokenIdx++;
+
+					// Execute the method
+					try
+					{
+						value = typeof(string).InvokeMember(methodName, System.Reflection.BindingFlags.Instance | System.Reflection.BindingFlags.Public | System.Reflection.BindingFlags.InvokeMethod, Type.DefaultBinder, value, [.. arguments])!.ToString()!;
+					}
+					catch(Exception ex)
+					{
+						throw new Exception(String.Format("Unable to evaluate condition '{0}'", text), ex);
+					}
+				}
+
+				if (tokenIdx < tokens.Length && tokens[tokenIdx] == ":")
+				{
+					tokenIdx = tokens.Length;
+				}
+
+				// Make sure there's nothing left over
+				if(tokenIdx != tokens.Length)
+				{
+					throw new Exception(String.Format("Unable to parse token '{0}'", newText));
+				}
+
+				// Replace the variable with its value
+				newText = newText.Substring(0, idx) + value + newText.Substring(endIdx);
+
+				// Make sure we skip over the expanded variable; we don't want to recurse on it.
+				idx += value.Length;
+			}
+			return newText;
+		}
+
+		/// <summary>
+		/// Parse an argument into a framework type
+		/// </summary>
+		/// <param name="token">The token to parse</param>
+		/// <returns>The argument object</returns>
+		static object ParseArgument(string token)
+		{
+			// Try to parse a string
+			if(token.Length > 2 && token[0] == '\'' && token[^1] == '\'')
+			{
+				return token.Substring(1, token.Length - 2);
+			}
+
+			// Try to parse an integer
+			int value;
+			if(Int32.TryParse(token, out value))
+			{
+				return value;
+			}
+
+			// Otherwise throw an exception
+			throw new Exception(String.Format("Unable to parse token '{0}' into a .NET framework type", token));
+		}
+
+		/// <summary>
+		/// Split an MSBuild condition into tokens
+		/// </summary>
+		/// <param name="condition">The condition expression</param>
+		/// <returns>Array of the parsed tokens</returns>
+		static string[] Tokenize(string condition)
+		{
+			List<string> tokens = [];
+			for (int idx = 0; idx < condition.Length; )
+			{
+				if(Char.IsWhiteSpace(condition[idx]))
+				{
+					// Whitespace
+					idx++;
+				}
+				else if (idx + 1 < condition.Length && condition[idx] == '=' && condition[idx + 1] == '=')
+				{
+					// "==" operator
+					idx += 2;
+					tokens.Add("==");
+				}
+				else if (idx + 1 < condition.Length && condition[idx] == '!' && condition[idx + 1] == '=')
+				{
+					// "!=" operator
+					idx += 2;
+					tokens.Add("!=");
+				}
+				else if (condition[idx] == '\'')
+				{
+					// Quoted string
+					int startIdx = idx++;
+					for(;;idx++)
+					{
+						if(idx == condition.Length)
+						{
+							throw new Exception(String.Format("Missing end quote in condition string ('{0}')", condition));
+						}
+						if(condition[idx] == '\'')
+						{
+							break;
+						}
+					}
+					idx++;
+					tokens.Add(condition.Substring(startIdx, idx - startIdx));
+				}
+				else if (condition[idx] == '[')
+				{
+#pragma warning disable IDE0059 // Unnecessary assignment of a value
+					// static property function invoke
+					// format: [Class]::Property
+					// alternatively: [Class]::Method()
+					// we consider the entire invocation to be a single token
+					int startIdx = idx++;
+					int classEndIdx = 0;
+					int methodEndIdx = 0;
+					int methodArgsEndIdx = 0;
+					for (; ; idx++)
+					{
+						while (idx < condition.Length && (Char.IsLetterOrDigit(condition[idx]) || condition[idx] == '_'))
+						{
+							idx++;
+						}
+
+						if (idx == condition.Length)
+						{
+							throw new Exception(String.Format("Found end of condition when searching for end of static property function for condition string ('{0}')", condition));
+						}
+						if (condition[idx] == ']')
+						{
+							classEndIdx = idx;
+							idx++;
+							break;
+						}
+					}
+
+					// skip ::
+					if (condition[idx] != ':')
+					{
+						throw new Exception(String.Format("Unexpected format of static property function, expected :: after class declaration in condition string ('{0}')", condition));
+					}
+					idx += 2;
+
+					while (idx < condition.Length && (Char.IsLetterOrDigit(condition[idx]) || condition[idx] == '_'))
+					{
+						idx++;
+					}
+
+					methodEndIdx = idx;
+
+					if (idx < condition.Length && condition[idx] == '(')
+					{
+						// a method invoke
+						for (; ; idx++)
+						{
+							while (idx < condition.Length && (Char.IsLetterOrDigit(condition[idx]) || condition[idx] == '_'))
+							{
+								idx++;
+							}
+
+							if (idx == condition.Length)
+							{
+								throw new Exception(String.Format("Found end of condition when searching for ) to indicate end of arguments to static property function for condition string ('{0}')", condition));
+							}
+							if (condition[idx] == ')')
+							{
+								methodArgsEndIdx = idx;
+								idx++;
+								break;
+							}
+						}
+						idx++;
+					}
+#pragma warning restore IDE0059 // Unnecessary assignment of a value
+
+					tokens.Add(condition.Substring(startIdx, idx - startIdx));
+				}
+				else if(Char.IsLetterOrDigit(condition[idx]) || condition[idx] == '_')
+				{
+					// Identifier or number
+					int startIdx = idx++;
+					while(idx < condition.Length && (Char.IsLetterOrDigit(condition[idx]) || condition[idx] == '_'))
+					{
+						idx++;
+					}
+					tokens.Add(condition.Substring(startIdx, idx - startIdx));
+				}
+				else
+				{
+					// Other token; assume a single character.
+					string token = condition.Substring(idx++, 1);
+					tokens.Add(token);
+				}
+			}
+			return [.. tokens];
+		}
+
+		/// <summary>
+		/// Un-escape an MSBuild string (see https://msdn.microsoft.com/en-us/library/bb383819.aspx)
+		/// </summary>
+		/// <param name="text">String to remove escape characters from</param>
+		/// <returns>Unescaped string</returns>
+		static string? UnescapeString(string? text)
+		{
+			const string HexChars = "0123456789abcdef";
+
+			string? newText = text;
+			if(newText != null)
+			{
+				for(int idx = 0; idx + 2 < newText.Length; idx++)
+				{
+					if(newText[idx] == '%')
+					{
+						int upperDigitIdx = HexChars.IndexOf(Char.ToLowerInvariant(newText[idx + 1]), StringComparison.Ordinal);
+						if(upperDigitIdx != -1)
+						{
+							int lowerDigitIdx = HexChars.IndexOf(Char.ToLowerInvariant(newText[idx + 2]), StringComparison.Ordinal);
+							if(lowerDigitIdx != -1)
+							{
+								char newChar = (char)((upperDigitIdx << 4) | lowerDigitIdx);
+								newText = newText.Substring(0, idx) + newChar + newText.Substring(idx + 3);
+							}
+						}
+					}
+				}
+			}
+			return newText;
+		}
+	}
+}
diff --git a/Engine/Source/Programs/Shared/EpicGames.Core/CsProjectInfo.cs b/Engine/Source/Programs/Shared/EpicGames.Core/CsProjectInfo.cs
deleted file mode 100644
