[EpicGames.Build]
* Changed CsProjectInfo to use DirectoryItem instead of DirectoryReference when traversing wildcard search. This to be able to utilize potentially cached results

diff --git a/Engine/Source/Programs/Shared/EpicGames.Build/CsProjectInfo.cs b/Engine/Source/Programs/Shared/EpicGames.Build/CsProjectInfo.cs
--- a/Engine/Source/Programs/Shared/EpicGames.Build/CsProjectInfo.cs
+++ b/Engine/Source/Programs/Shared/EpicGames.Build/CsProjectInfo.cs
@@ -9,6 +9,7 @@ using System.Linq;
 using System.Runtime.InteropServices;
 using System.Text.RegularExpressions;
 using System.Xml;
+using UnrealBuildBase;
 
 namespace EpicGames.Core
 {
@@ -631,17 +632,17 @@ namespace EpicGames.Core
 		/// <summary>
 		/// Finds all directories but filters out certain folders
 		/// </summary>
-		static void GetDirectories(DirectoryReference dir, ConcurrentBag<DirectoryReference> list, SearchOption option, ThreadPoolWorkQueue queue)
+		static void GetDirectories(DirectoryItem dir, ConcurrentBag<DirectoryReference> list, SearchOption option, ParallelQueue queue)
 		{
-			foreach (DirectoryReference child in DirectoryReference.EnumerateDirectories(dir))
+			foreach (DirectoryItem child in dir.EnumerateDirectories())
 			{
-				string name = child.GetDirectoryName();
+				string name = child.Name;
 				if (name == "Intermediate" || name == "Binaries")
 				{
 					continue;
 				}
 
-				list.Add(child);
+				list.Add(child.Location);
 				if (option != SearchOption.AllDirectories)
 				{
 					continue;
@@ -698,10 +699,9 @@ namespace EpicGames.Core
 						// and we also want to filter out Intermediate, Binaries, etc
 
 						ConcurrentBag<DirectoryReference> bag = new();
-						using (ThreadPoolWorkQueue queue = new ThreadPoolWorkQueue())
-						{
-							GetDirectories(existingPath, bag, option, queue);
-						}
+						ParallelQueue queue = new();
+						GetDirectories(DirectoryItem.GetItemByDirectoryReference(existingPath), bag, option, queue);
+						queue.Drain();
 
 						// Add current path
 						bag.Add(existingPath);
