[EpicGames.Shared]
* Changed PluginsBase to use Lazy for value in concurrent dictionary to make sure code is really only run once
* Changed from ThreadPoolWorkQueue to ParallelQueue to improve perf

diff --git a/Engine/Source/Programs/Shared/EpicGames.Build/System/EnumeratePlugins.cs b/Engine/Source/Programs/Shared/EpicGames.Build/System/EnumeratePlugins.cs
--- a/Engine/Source/Programs/Shared/EpicGames.Build/System/EnumeratePlugins.cs
+++ b/Engine/Source/Programs/Shared/EpicGames.Build/System/EnumeratePlugins.cs
@@ -15,7 +15,7 @@ namespace UnrealBuildBase
 		/// <summary>
 		/// Cache of plugin filenames under each directory
 		/// </summary>
-		static ConcurrentDictionary<DirectoryReference, List<FileReference>> PluginFileCache = new ConcurrentDictionary<DirectoryReference, List<FileReference>>();
+		static ConcurrentDictionary<DirectoryReference, Lazy<List<FileReference>>> PluginFileCache = new();
 
 		/// <summary>
 		/// Cache of all UBT plugins that have been built.  When generating projects, there is a good chance that a UBT project
@@ -29,7 +29,7 @@ namespace UnrealBuildBase
 		/// </summary>
 		public static void InvalidateCache_SLOW()
 		{
-			PluginFileCache = new ConcurrentDictionary<DirectoryReference, List<FileReference>>();
+			PluginFileCache = new();
 			DirectoryItem.ResetAllCachedInfo_SLOW();
 		}
 
@@ -55,23 +55,22 @@ namespace UnrealBuildBase
 		/// <param name="ParentDirectory">Parent directory to look in. Plugins will be found in any *subfolders* of this directory.</param>
 		public static List<FileReference> EnumeratePlugins(DirectoryReference ParentDirectory)
 		{
-			return PluginFileCache.GetOrAdd(ParentDirectory, _ =>
+			return PluginFileCache.GetOrAdd(ParentDirectory, new Lazy<List<FileReference>>(() =>
 			{
 				List<FileReference> FileNames = [];
 
 				DirectoryItem ParentDirectoryItem = DirectoryItem.GetItemByDirectoryReference(ParentDirectory);
 				if (ParentDirectoryItem.Exists)
 				{
-					using (ThreadPoolWorkQueue Queue = new ThreadPoolWorkQueue())
-					{
-						EnumeratePluginsInternal(ParentDirectoryItem, FileNames, Queue);
-					}
+					ParallelQueue Queue = new();
+					EnumeratePluginsInternal(ParentDirectoryItem, FileNames, Queue);
+					Queue.Drain();
 				}
 
 				// Sort the filenames to ensure that the plugin order is deterministic; otherwise response files will change with each build.
 				FileNames = [.. FileNames.OrderBy(x => x.FullName, StringComparer.OrdinalIgnoreCase)];
 				return FileNames;
-			});
+			})).Value;
 		}
 
 		/// <summary>
@@ -80,7 +79,7 @@ namespace UnrealBuildBase
 		/// <param name="ParentDirectory">Parent directory to look in. Plugins will be found in any *subfolders* of this directory.</param>
 		/// <param name="FileNames">List of filenames. Will have all the discovered .uplugin files appended to it.</param>
 		/// <param name="Queue">Queue for tasks to be executed</param>
-		static void EnumeratePluginsInternal(DirectoryItem ParentDirectory, List<FileReference> FileNames, ThreadPoolWorkQueue Queue)
+		static void EnumeratePluginsInternal(DirectoryItem ParentDirectory, List<FileReference> FileNames, ParallelQueue Queue)
 		{
 			foreach (DirectoryItem ChildDirectory in ParentDirectory.EnumerateDirectories())
 			{
