[UBT]
* Changed so DirectoryItem.Scan of files allow but ignore files with the same casing in the same dir. The reason we hit this now and not before is that cached files are now happening for directories which did not cache before (like in linux sdk). This is still an optimization because doing a separate EnumerateFiles and EnumerateDirs is much more expensive that just Enumerate all for both.

This fixes solution generation on linux

diff --git a/Engine/Source/Programs/Shared/EpicGames.Build/System/DirectoryItem.cs b/Engine/Source/Programs/Shared/EpicGames.Build/System/DirectoryItem.cs
--- a/Engine/Source/Programs/Shared/EpicGames.Build/System/DirectoryItem.cs
+++ b/Engine/Source/Programs/Shared/EpicGames.Build/System/DirectoryItem.cs
@@ -269,8 +269,13 @@ namespace UnrealBuildBase
 					else
 					{
 						FileItem FileItem = FileItem.GetItemByFileInfo((FileInfo)info);
-						FileItem.UpdateCachedDirectory(this);
-						newFiles.Add(info.Name, FileItem);
+
+						// There are folders in linux sdk that has files with same name but different casing.
+						// Ideally FileReference.Comparer should be case sensitive on linux/mac but I don't dare changing that right now
+						if (newFiles.TryAdd(info.Name, FileItem))
+						{
+							FileItem.UpdateCachedDirectory(this);
+						}
 					}
 				}
 
