[EpicGames.Build]
* Fixed so Rules.FindAllRulesSourceFiles runs more in parallel and also use ParallelQueue instead of ThreadPoolWorkQueue

diff --git a/Engine/Source/Programs/Shared/EpicGames.Build/System/Rules.cs b/Engine/Source/Programs/Shared/EpicGames.Build/System/Rules.cs
--- a/Engine/Source/Programs/Shared/EpicGames.Build/System/Rules.cs
+++ b/Engine/Source/Programs/Shared/EpicGames.Build/System/Rules.cs
@@ -6,6 +6,7 @@ using System.Collections.Generic;
 using System.IO;
 using System.Linq;
 using System.Text;
+using System.Threading.Tasks;
 
 namespace UnrealBuildBase
 {
@@ -86,16 +87,19 @@ namespace UnrealBuildBase
 			}
 
 			// Find all the plugin source and tests directories
-			foreach (DirectoryReference RootFolder in RootFolders)
+			Parallel2.ForEach(RootFolders, (RootFolder) =>
 			{
 				DirectoryReference PluginsFolder = DirectoryReference.Combine(RootFolder, "Plugins");
 				foreach (FileReference PluginFile in PluginsBase.EnumeratePlugins(PluginsFolder))
 				{
-					Folders.Add(DirectoryReference.Combine(PluginFile.Directory, "Source"));
-					Folders.Add(DirectoryReference.Combine(PluginFile.Directory, "Tests"));
-					Folders.Add(DirectoryReference.Combine(PluginFile.Directory, "Build"));	// this is to pick up plugin's C# automation projects					
+					lock (Folders)
+					{
+						Folders.Add(DirectoryReference.Combine(PluginFile.Directory, "Source"));
+						Folders.Add(DirectoryReference.Combine(PluginFile.Directory, "Tests"));
+						Folders.Add(DirectoryReference.Combine(PluginFile.Directory, "Build")); // this is to pick up plugin's C# automation projects					
+					}
 				}
-			}
+			});
 
 			// Add all the extra plugin folders
 			if (ForeignPlugins != null)
@@ -163,37 +167,35 @@ namespace UnrealBuildBase
 			List<(DirectoryReference, RulesFileCache)> Caches = new List<(DirectoryReference, RulesFileCache)>(Directories.Count());
 			lock (RootFolderToRulesFileCache)
 			{
-				using (ThreadPoolWorkQueue Queue = new ThreadPoolWorkQueue())
+				ParallelQueue queue = new();
+				foreach (DirectoryReference Directory in Directories)
 				{
-					foreach (DirectoryReference Directory in Directories)
+					RulesFileCache? Cache;
+					if (!RootFolderToRulesFileCache.TryGetValue(Directory, out Cache))
 					{
-						RulesFileCache? Cache;
-						if (!RootFolderToRulesFileCache.TryGetValue(Directory, out Cache))
-						{
-							Cache = new RulesFileCache();
-							Queue.Enqueue(() => FindAllRulesFilesRecursively(DirectoryItem.GetItemByDirectoryReference(Directory), Cache, Queue));
-						}
-						Caches.Add((Directory, Cache));
+						Cache = new RulesFileCache();
+						queue.Enqueue(() => FindAllRulesFilesRecursively(DirectoryItem.GetItemByDirectoryReference(Directory), Cache, queue));
 					}
+					Caches.Add((Directory, Cache));
 				}
+				queue.Drain();
 
-				using (ThreadPoolWorkQueue SortQueue = new ThreadPoolWorkQueue())
+				ParallelQueue sortQueue = new();
+				foreach ((DirectoryReference Directory, RulesFileCache Cache) in Caches)
 				{
-					foreach ((DirectoryReference Directory, RulesFileCache Cache) in Caches)
+					if (!RootFolderToRulesFileCache.ContainsKey(Directory))
 					{
-						if (!RootFolderToRulesFileCache.ContainsKey(Directory))
+						sortQueue.Enqueue(() =>
 						{
-							SortQueue.Enqueue(() =>
-							{
-								Cache.ModuleRules.Sort((A, B) => A.FullName.CompareTo(B.FullName));
-								Cache.TargetRules.Sort((A, B) => A.FullName.CompareTo(B.FullName));
-								Cache.AutomationModules.Sort((A, B) => A.FullName.CompareTo(B.FullName));
-								Cache.UbtPlugins.Sort((A, B) => A.FullName.CompareTo(B.FullName));
-							});
-							RootFolderToRulesFileCache[Directory] = Cache;
-						}
+							Cache.ModuleRules.Sort((A, B) => A.FullName.CompareTo(B.FullName));
+							Cache.TargetRules.Sort((A, B) => A.FullName.CompareTo(B.FullName));
+							Cache.AutomationModules.Sort((A, B) => A.FullName.CompareTo(B.FullName));
+							Cache.UbtPlugins.Sort((A, B) => A.FullName.CompareTo(B.FullName));
+						});
+						RootFolderToRulesFileCache[Directory] = Cache;
 					}
 				}
+				sortQueue.Drain();
 			}
 			return Caches;
 		}
@@ -252,7 +254,7 @@ namespace UnrealBuildBase
 		/// <param name="Directory">The root directory to search from</param>
 		/// <param name="Cache">Receives all the discovered rules files</param>
 		/// <param name="Queue">Queue for adding additional tasks to</param>
-		private static void FindAllRulesFilesRecursively(DirectoryItem Directory, RulesFileCache Cache, ThreadPoolWorkQueue Queue)
+		private static void FindAllRulesFilesRecursively(DirectoryItem Directory, RulesFileCache Cache, ParallelQueue Queue)
 		{
 			// Scan all the files in this directory
 			bool bSearchSubFolders = true;
