[EpicGames.Core]
* Optimized JsonObject reading.

diff --git a/Engine/Source/Programs/Shared/EpicGames.Core/JsonObject.cs b/Engine/Source/Programs/Shared/EpicGames.Core/JsonObject.cs
--- a/Engine/Source/Programs/Shared/EpicGames.Core/JsonObject.cs
+++ b/Engine/Source/Programs/Shared/EpicGames.Core/JsonObject.cs
@@ -10,6 +10,7 @@ using System.Collections;
 using System.IO;
 using System.Text;
 using System.Text.Encodings.Web;
+using System.Buffers;
 
 namespace EpicGames.Core
 {
@@ -105,9 +106,16 @@ namespace EpicGames.Core
 		/// Construct a JSON object from the OrderedDictionary obtained from reading a file on disk or parsing valid json text.
 		/// </summary>
 		/// <param name="inRawObject">Raw object parsed from disk</param>
-		private JsonObject(OrderedDictionary inRawObject)
+		private JsonObject(OrderedDictionary inRawObject, bool createCopy)
 		{
-			_rawOrderedObject = inRawObject.CreateDeepCopy();
+			if (createCopy)
+			{
+				_rawOrderedObject = inRawObject.CreateDeepCopy();
+			}
+			else
+			{
+				_rawOrderedObject = inRawObject;
+			}
 		}
 
 		/// <summary>
@@ -174,13 +182,22 @@ namespace EpicGames.Core
 			switch(element.ValueKind)
 			{
 				case JsonValueKind.Array:
-					return element.EnumerateArray().Select(x => ParseElement(x)).ToArray();
+					{
+						int len = element.GetArrayLength();
+						object?[] elementArray = new object[len];
+						int index = 0;
+						foreach (JsonElement x in element.EnumerateArray())
+						{
+							elementArray[index++] = ParseElement(x);
+						}
+						return elementArray;
+					}
 				case JsonValueKind.Number:
 					return element.GetDouble();
 				case JsonValueKind.Object:
 					OrderedDictionary dictionary = new OrderedDictionary(StringComparer.InvariantCultureIgnoreCase);
 					foreach (JsonProperty property in element.EnumerateObject())
-						{
+					{
 						dictionary.Add(property.Name, ParseElement(property.Value));
 					}
 					return dictionary;
@@ -237,6 +254,38 @@ namespace EpicGames.Core
 			return returnObject;
 		}
 
+		static JsonDocument ParseUtf8File(string path)
+		{
+			using FileStream fs = new(path, FileMode.Open, FileAccess.Read, FileShare.Read, bufferSize: 1, FileOptions.SequentialScan);
+
+			int len = checked((int)fs.Length);
+			byte[] buf = len <= 128 * 1024 ? new byte[len] : ArrayPool<byte>.Shared.Rent(len);
+			try
+			{
+				fs.ReadExactly(buf.AsSpan(0, len));
+				int off = (len >= 3 && buf[0] == 0xEF && buf[1] == 0xBB && buf[2] == 0xBF) ? 3 : 0; // Skip UTF-8 BOM if present
+
+				try
+				{
+					return JsonDocument.Parse(new ReadOnlyMemory<byte>(buf, off, len - off), new() { AllowTrailingCommas = true, CommentHandling = JsonCommentHandling.Skip });
+				}
+				catch (JsonException)
+				{
+					throw;
+				}
+				catch (Exception ex)
+				{
+					throw new JsonException($"Unable to parse {path}: {ex.Message}", path, null, null, ex);
+				}
+			}
+			finally
+			{
+				if (buf.Length > 128 * 1024)
+				{
+					ArrayPool<byte>.Shared.Return(buf);
+				}
+			}
+		}
 		/// <summary>
 		/// Read a JSON file from disk and construct a JsonObject from it
 		/// </summary>
@@ -244,19 +293,8 @@ namespace EpicGames.Core
 		/// <returns>New JsonObject instance</returns>
 		public static JsonObject Read(FileReference file)
 		{
-			string text = FileReference.ReadAllText(file);
-			try
-			{
-				return Parse(text);
-			}
-			catch (JsonException)
-			{
-				throw;
-			}
-			catch (Exception ex)
-			{
-				throw new JsonException($"Unable to parse {file}: {ex.Message}", file.FullName, null, null, ex );
-			}
+			JsonDocument doc = ParseUtf8File(file.FullName);
+			return new JsonObject(doc.RootElement);
 		}
 
 		/// <summary>
@@ -267,14 +305,23 @@ namespace EpicGames.Core
 		/// <returns>True if the file was read, false otherwise</returns>
 		public static bool TryRead(FileReference fileName, [NotNullWhen(true)] out JsonObject? result)
 		{
+			result = null;
+
 			if (!FileReference.Exists(fileName))
 			{
-				result = null;
 				return false;
 			}
 
-			string text = FileReference.ReadAllText(fileName);
-			return TryParse(text, out result);
+			try
+			{
+				JsonDocument doc = ParseUtf8File(fileName.FullName);
+				result = new JsonObject(doc.RootElement);
+				return true;
+			}
+			catch (Exception)
+			{
+				return false;
+			}
 		}
 
 		/// <summary>
@@ -600,7 +647,7 @@ namespace EpicGames.Core
 			object? rawValue;
 			if (_rawOrderedObject.TryGetValue(fieldName, out rawValue) && (rawValue is OrderedDictionary orderedDictValue))
 			{
-				result = new JsonObject(orderedDictValue);
+				result = new JsonObject(orderedDictValue, false);
 				return true;
 			}
 			else
@@ -636,7 +683,7 @@ namespace EpicGames.Core
 			object? rawValue;
 			if (_rawOrderedObject.TryGetValue(fieldName, out rawValue) && (rawValue is IEnumerable<object> enumValue) && enumValue.All(x => x is OrderedDictionary))
 			{
-				result = enumValue.Select(x => new JsonObject((OrderedDictionary)x)).ToArray();
+				result = enumValue.Select(x => new JsonObject((OrderedDictionary)x, false)).ToArray();
 				return true;
 			}
 			else
