[EpicGames.Core]
* Fixed xml comment tag

diff --git a/Engine/Source/Programs/Shared/EpicGames.Core/JsonObject.cs b/Engine/Source/Programs/Shared/EpicGames.Core/JsonObject.cs
--- a/Engine/Source/Programs/Shared/EpicGames.Core/JsonObject.cs
+++ b/Engine/Source/Programs/Shared/EpicGames.Core/JsonObject.cs
@@ -106,6 +106,7 @@ namespace EpicGames.Core
 		/// Construct a JSON object from the OrderedDictionary obtained from reading a file on disk or parsing valid json text.
 		/// </summary>
 		/// <param name="inRawObject">Raw object parsed from disk</param>
+		/// <param name="createCopy">raw object should be copied or not</param>
 		private JsonObject(OrderedDictionary inRawObject, bool createCopy)
 		{
 			if (createCopy)
