[EpicGames.Build]
* Added Parallel2.For

diff --git a/Engine/Source/Programs/Shared/EpicGames.Build/System/Parallel2.cs b/Engine/Source/Programs/Shared/EpicGames.Build/System/Parallel2.cs
--- a/Engine/Source/Programs/Shared/EpicGames.Build/System/Parallel2.cs
+++ b/Engine/Source/Programs/Shared/EpicGames.Build/System/Parallel2.cs
@@ -32,7 +32,7 @@ namespace UnrealBuildBase
 		{
 			using EventWaitHandle Handle = new(false, EventResetMode.ManualReset);
 
-			Context<TSource> context = new() { Enumerator = source.GetEnumerator(), Handle = Handle };
+			ForEachContext<TSource> context = new() { Enumerator = source.GetEnumerator(), Handle = Handle };
 
 			void WorkerLoop()
 			{
@@ -120,7 +120,84 @@ namespace UnrealBuildBase
 			}
 		}
 
-		private class Context<TSource>
+		/// <summary>
+		/// Read information about Parallel2.ForEach
+		/// </summary>
+		/// <param name="from">Start index</param>
+		/// <param name="to">End index</param>
+		/// <param name="body">action to execute on each index</param>
+		/// <param name="helperCount">Helper count. -1 means it will use as many cores as needed</param>
+		/// <param name="useTasks">Use tasks or threads</param>
+		public static void For(int from, int to, Action<int> body, int helperCount = -1, bool useTasks = true)
+		{
+			using EventWaitHandle Handle = new(false, EventResetMode.ManualReset);
+
+			ForContext context = new() { Handle = Handle, Index = from, End = to, Finished = from };
+
+			void WorkerLoop()
+			{
+				while (true)
+				{
+					int index = Interlocked.Increment(ref context.Index) - 1;
+					if (index >= context.End)
+					{
+						return;
+					}
+
+					try
+					{
+						body(index);
+					}
+					catch (Exception e)
+					{
+						Interlocked.CompareExchange(ref context.Exception, e, null);
+					}
+
+					if (Interlocked.Increment(ref context.Finished) == context.End)
+					{
+						context.Handle.Set();
+						return;
+					}
+				}
+			}
+
+			if (helperCount == -1)
+			{
+				helperCount = MaxHelperCount - 1;
+			}
+
+			if (helperCount > 0)
+			{
+				helperCount = Math.Min(to - from, helperCount);
+				if (helperCount > 0)
+				{
+					--helperCount;
+				}
+			}
+
+			for (int i = 0; i != helperCount; ++i)
+			{
+				if (useTasks)
+				{
+					Task.Run(WorkerLoop);
+				}
+				else
+				{
+					(new Thread(WorkerLoop)).Start();
+				}
+			}
+
+			WorkerLoop();
+
+			context.Handle.WaitOne();
+
+			if (context.Exception != null)
+			{
+				throw context.Exception;
+			}
+		}
+
+		private class ForEachContext<TSource>
 		{
 			internal required IEnumerator<TSource> Enumerator;
 			internal required EventWaitHandle Handle;
@@ -130,6 +207,15 @@ namespace UnrealBuildBase
 			internal Exception? Exception;
 		}
 
+		private class ForContext
+		{
+			internal required EventWaitHandle Handle;
+			internal int Index;
+			internal int End;
+			internal int Finished;
+			internal Exception? Exception;
+		}
+
 		public static readonly int MaxHelperCount = Math.Max(Environment.ProcessorCount - 1, 0);
 	}
 
