EpicGames.Perforce: Add a secondary release stream to fixtures

diff --git a/Engine/Source/Programs/Shared/EpicGames.Perforce.Fixture/BasePerforceFixtureTest.cs b/Engine/Source/Programs/Shared/EpicGames.Perforce.Fixture/BasePerforceFixtureTest.cs
--- a/Engine/Source/Programs/Shared/EpicGames.Perforce.Fixture/BasePerforceFixtureTest.cs
+++ b/Engine/Source/Programs/Shared/EpicGames.Perforce.Fixture/BasePerforceFixtureTest.cs
@@ -30,7 +30,7 @@ public abstract class BasePerforceFixtureTest : IDisposable
 	/// Temporary scratch dir, deleted after every test.
 	/// Provides a location for (managed) workspace root or cache files. 
 	/// </summary>
-	protected DirectoryReference TempDir { get; }
+	protected DirectoryReference TempDir { get; set;  }
 
 	/// <summary>
 	/// A Perforce connection to the fixture test server
@@ -42,6 +42,11 @@ public abstract class BasePerforceFixtureTest : IDisposable
 	/// </summary>
 	protected PerforceFixture Fixture { get; } = new();
 
+	/// <summary>
+	/// Unique ID to avoid tests clashing with each other
+	/// </summary>
+	protected string TestGuid;
+	
 	private readonly string _clientName;
 
 	protected BasePerforceFixtureTest()
@@ -52,9 +57,9 @@ public abstract class BasePerforceFixtureTest : IDisposable
 		// That initialize method will get called after the constructor ensuring the connection is instantiated.
 		PerforceConnection = null!;
 
-		_clientName = "test-fixture-" + Guid.NewGuid().ToString()[..8];
-		TempDir = new DirectoryReference(Path.Join(Path.GetTempPath(), "p4-" + _clientName));
-		DirectoryReference.CreateDirectory(TempDir);
+		TestGuid = Guid.NewGuid().ToString()[..8];
+		_clientName = "test-fixture-" + TestGuid;
+		CreateTempDir();
 	}
 
 	[TestInitialize]
@@ -70,19 +75,31 @@ public abstract class BasePerforceFixtureTest : IDisposable
 		PerforceConnection = GetPerforceConnection(serverUrl!);
 	}
 
+	protected void CreateTempDir()
+	{
+		DirectoryReference dir = new (Path.Join(Path.GetTempPath(), "p4-" + _clientName));
+		DirectoryReference.CreateDirectory(dir);
+		TempDir = dir;
+	}
+
 	[TestCleanup]
 	public void RemoveTempDir()
 	{
-		if (Directory.Exists(TempDir.FullName))
+		DeletePerforceDir(TempDir.FullName);
+	}
+
+	public void DeletePerforceDir(string dirPath)
+	{
+		if (Directory.Exists(dirPath))
 		{
 			// Remove the read-only flags set by the P4 client
-			foreach (string filePath in Directory.EnumerateFiles(TempDir.FullName, "*", SearchOption.AllDirectories))
+			foreach (string filePath in Directory.EnumerateFiles(dirPath, "*", SearchOption.AllDirectories))
 			{
 				FileInfo fileInfo = new FileInfo(filePath);
 				fileInfo.IsReadOnly = false;
 			}
 
-			Directory.Delete(TempDir.FullName, true);
+			Directory.Delete(dirPath, true);
 		}
 	}
 
diff --git a/Engine/Source/Programs/Shared/EpicGames.Perforce.Fixture/PerforceFixture.cs b/Engine/Source/Programs/Shared/EpicGames.Perforce.Fixture/PerforceFixture.cs
--- a/Engine/Source/Programs/Shared/EpicGames.Perforce.Fixture/PerforceFixture.cs
+++ b/Engine/Source/Programs/Shared/EpicGames.Perforce.Fixture/PerforceFixture.cs
@@ -283,6 +283,28 @@ public class PerforceFixture
 				}, true),
 		});
 
+	public StreamFixture StreamFooRelease { get; } = new("//Foo/Release",
+		new List<ChangelistFixture>
+		{
+			new(10, "Initial import",
+				new List<DepotFileFixture>()
+				{
+					new("//Foo/Main/main.cpp", "main.cpp", "This is change main.cpp #2\n", 1),
+					new("//Foo/Main/main.h", "main.h", "This is change main.h #1\n", 1),
+					new("//Foo/Main/shared.h", "shared.h", "This is change common.h #1\n", 1),
+					// data.txt is excluded in view
+				}),
+			new(11, "Initial import",
+				new List<DepotFileFixture>()
+				{
+					new("//Foo/Main/main.cpp", "main.cpp", "This is change main.cpp #2 in //Foo/Release\n", 1),
+					new("//Foo/Main/main.h", "main.h", "This is change main.h #1\n", 1),
+					new("//Foo/Main/shared.h", "shared.h", "This is change common.h #1\n", 1),
+					new("//Foo/Main/release.cpp", "release.cpp", "This is change release.cpp #1 in //Foo/Release\n", 1),
+					// data.txt is excluded in view
+				}),
+		});
+
 	public static string CalcMd5(string content)
 	{
 		byte[] data = Encoding.ASCII.GetBytes(content);
diff --git a/Engine/Source/Programs/Shared/EpicGames.Perforce.Fixture/setup.sh b/Engine/Source/Programs/Shared/EpicGames.Perforce.Fixture/setup.sh
--- a/Engine/Source/Programs/Shared/EpicGames.Perforce.Fixture/setup.sh
+++ b/Engine/Source/Programs/Shared/EpicGames.Perforce.Fixture/setup.sh
@@ -69,7 +69,7 @@ Paths:
         share ...
 EOF
 
-export P4CLIENT=generate-fixture-data
+export P4CLIENT=generate-fixture-data-main
 export CLIENT_ROOT=/tmp/${P4CLIENT}
 p4 client -i << EOF
 Client: ${P4CLIENT}
@@ -127,5 +127,45 @@ p4 edit -c 8 main.h && echo "This is change main.h #3" > $CLIENT_ROOT/main.h
 p4 delete -c 8 main.cpp
 p4 shelve -c 8
 
+
+# Create alternative release stream
+p4 stream -i << EOF
+Stream: //Foo/Release
+Owner:  test.user
+Name:   Release
+Parent: //Foo/Main
+Type:   release
+Description:
+        Release stream for Foo depot
+
+Options:        allsubmit unlocked notoparent nofromparent mergedown
+ParentView:     inherit
+Paths:
+        share ...
+        exclude Data/...
+EOF
+
+export P4CLIENT=generate-fixture-data-release
+export CLIENT_ROOT=/tmp/${P4CLIENT}
+p4 client -i << EOF
+Client: ${P4CLIENT}
+Root: ${CLIENT_ROOT}
+Stream: //Foo/Release
+EOF
+
+# Create the workspace root dir and sync the empty stream
+mkdir -p $CLIENT_ROOT
+cd $CLIENT_ROOT
+
+# Changelist #10
+p4 populate -d "Populate from CL 5" //Foo/Main/...@5 //Foo/Release/...
+p4 sync
+
+# Changelist #11
+p4 edit main.cpp
+echo "This is change main.cpp #2 in //Foo/Release" > $CLIENT_ROOT/main.cpp
+echo "This is change release.cpp #1 in //Foo/Release" > $CLIENT_ROOT/release.cpp && p4 add $CLIENT_ROOT/release.cpp
+p4 submit -d "Improvements to //Foo/Release"
+
 # Stop server running in background
 pkill p4d
\ No newline at end of file
diff --git a/Engine/Source/Programs/Shared/EpicGames.Perforce.Fixture/start.bat b/Engine/Source/Programs/Shared/EpicGames.Perforce.Fixture/start.bat
--- a/Engine/Source/Programs/Shared/EpicGames.Perforce.Fixture/start.bat
+++ b/Engine/Source/Programs/Shared/EpicGames.Perforce.Fixture/start.bat
@@ -2,5 +2,4 @@
 IF NOT DEFINED DOCKER_EXECUTABLE (SET DOCKER_EXECUTABLE=docker)
 
 set TAG=perforce-fixture
-%DOCKER_EXECUTABLE% build --tag %TAG% .
-%DOCKER_EXECUTABLE% run -it --rm --publish 1666:1666 %TAG%
\ No newline at end of file
+%DOCKER_EXECUTABLE% build --tag %TAG% . && %DOCKER_EXECUTABLE% run -it --rm --publish 1666:1666 --ulimit nofile=65536:65536 --ulimit nproc=32768:32768 %TAG%
\ No newline at end of file
