EpicGames.Perforce: Fix compiler warnings

diff --git a/Engine/Source/Programs/Shared/EpicGames.Perforce.Fixture/BasePerforceFixtureTest.cs b/Engine/Source/Programs/Shared/EpicGames.Perforce.Fixture/BasePerforceFixtureTest.cs
--- a/Engine/Source/Programs/Shared/EpicGames.Perforce.Fixture/BasePerforceFixtureTest.cs
+++ b/Engine/Source/Programs/Shared/EpicGames.Perforce.Fixture/BasePerforceFixtureTest.cs
@@ -30,7 +30,7 @@ public abstract class BasePerforceFixtureTest : IDisposable
 	/// Temporary scratch dir, deleted after every test.
 	/// Provides a location for (managed) workspace root or cache files. 
 	/// </summary>
-	protected DirectoryReference TempDir { get; set;  }
+	protected DirectoryReference TempDir { get; set; }
 
 	/// <summary>
 	/// A Perforce connection to the fixture test server
@@ -45,7 +45,7 @@ public abstract class BasePerforceFixtureTest : IDisposable
 	/// <summary>
 	/// Unique ID to avoid tests clashing with each other
 	/// </summary>
-	protected string TestGuid;
+	protected string TestGuid { get; }
 	
 	private readonly string _clientName;
 
@@ -59,7 +59,7 @@ public abstract class BasePerforceFixtureTest : IDisposable
 
 		TestGuid = Guid.NewGuid().ToString()[..8];
 		_clientName = "test-fixture-" + TestGuid;
-		CreateTempDir();
+		TempDir = CreateTempDir();
 	}
 
 	[TestInitialize]
@@ -75,11 +75,12 @@ public abstract class BasePerforceFixtureTest : IDisposable
 		PerforceConnection = GetPerforceConnection(serverUrl!);
 	}
 
-	protected void CreateTempDir()
+	protected DirectoryReference CreateTempDir()
 	{
 		DirectoryReference dir = new (Path.Join(Path.GetTempPath(), "p4-" + _clientName));
 		DirectoryReference.CreateDirectory(dir);
 		TempDir = dir;
+		return dir;
 	}
 
 	[TestCleanup]
@@ -88,7 +89,7 @@ public abstract class BasePerforceFixtureTest : IDisposable
 		DeletePerforceDir(TempDir.FullName);
 	}
 
-	public void DeletePerforceDir(string dirPath)
+	public static void DeletePerforceDir(string dirPath)
 	{
 		if (Directory.Exists(dirPath))
 		{
