[Horde][BuildHealth][Server] Add support for file structured logging property to hashed issue handler

diff --git a/Engine/Source/Programs/Horde/Plugins/Build/HordeServer.Build.Tests/Issues/IssueServiceTests.cs b/Engine/Source/Programs/Horde/Plugins/Build/HordeServer.Build.Tests/Issues/IssueServiceTests.cs
--- a/Engine/Source/Programs/Horde/Plugins/Build/HordeServer.Build.Tests/Issues/IssueServiceTests.cs
+++ b/Engine/Source/Programs/Horde/Plugins/Build/HordeServer.Build.Tests/Issues/IssueServiceTests.cs
@@ -1815,6 +1815,44 @@ namespace HordeServer.Tests.Issues
 			Assert.AreEqual("Errors in Update Version Files and Compile UnrealHeaderTool Win64", issue.Summary);
 		}
 
+		[TestMethod]
+		public async Task HashedIssueFileMetadataTestAsync()
+		{
+			// Scenario: Job step fails at CL 325
+			// Expected: Creates hashed issue and blames submitter at CL 305 and 320 using the file metadata. CL 310, 315, and 325 are omitted from the suspect list.
+			{
+				IJob job1 = CreateJob(MainStreamId, 300, "Test Build", Graph);
+				await UpdateCompleteStepAsync(job1, 0, 0, JobStepOutcome.Success);
+
+				IJob job2 = CreateJob(MainStreamId, 325, "Test Build", Graph);
+
+				string[] lines =
+				{
+					"{ \"time\":\"2025-01-01T01:20:30.990Z\",\"level\":\"Error\",\"message\":\"Assertion failed: 1 == 2\",\"properties\":{ \"file\":\"a/c.cpp\"} }",
+					"{ \"time\":\"2025-01-01T01:20:30.990Z\",\"level\":\"Error\",\"message\":\"Assertion failed: 1 == 2\",\"properties\":{ \"file\":\"a/d.uasset\"} }"
+				};
+
+				// need to edit lines to evoid localization issue handler handling the error
+				await ParseEventsAsync(job2, 0, 0, lines);
+				await UpdateCompleteStepAsync(job2, 0, 0, JobStepOutcome.Failure);
+
+				// TODO: test multiple events in the same step with lines
+
+				IReadOnlyList<IIssue> issues = await IssueCollection.FindIssuesAsync();
+				Assert.AreEqual(1, issues.Count);
+
+				IIssue issue = issues[0];
+				Assert.AreEqual("Errors in Update Version Files", issue.Summary);
+
+				IReadOnlyList<IIssueSuspect> suspects = await IssueCollection.FindSuspectsAsync(issue);
+				suspects = suspects.OrderBy(x => x.Id).ToList();
+				Assert.AreEqual(2, suspects.Count);
+
+				Assert.AreEqual(320, suspects[0].CommitId.GetPerforceChange());
+				Assert.AreEqual(305, suspects[1].CommitId.GetPerforceChange());
+			}
+		}
+
 		[TestMethod]
 		public async Task ScopedIssueTest1Async()
 		{
diff --git a/Engine/Source/Programs/Horde/Plugins/Build/HordeServer.Build/Issues/IssueService.cs b/Engine/Source/Programs/Horde/Plugins/Build/HordeServer.Build/Issues/IssueService.cs
--- a/Engine/Source/Programs/Horde/Plugins/Build/HordeServer.Build/Issues/IssueService.cs
+++ b/Engine/Source/Programs/Horde/Plugins/Build/HordeServer.Build/Issues/IssueService.cs
@@ -1403,6 +1403,17 @@ namespace HordeServer.Issues
 					}
 				}
 
+				if (fingerprint.Metadata != null)
+				{
+					foreach (IssueMetadata metadata in fingerprint.Metadata)
+					{
+						if (metadata.Key == "File")
+						{
+							suspectFiles.Add(metadata.Value);
+						}
+					}
+				}
+
 				// Build a set of 'notes' to include as suspects
 				HashSet<string> suspectNotes = new HashSet<string>(StringComparer.OrdinalIgnoreCase);
 				foreach (IssueKey key in fingerprint.Keys)
diff --git a/Engine/Source/Programs/Shared/EpicGames.Horde/Issues/Handlers/HashedIssueHandler.cs b/Engine/Source/Programs/Shared/EpicGames.Horde/Issues/Handlers/HashedIssueHandler.cs
--- a/Engine/Source/Programs/Shared/EpicGames.Horde/Issues/Handlers/HashedIssueHandler.cs
+++ b/Engine/Source/Programs/Shared/EpicGames.Horde/Issues/Handlers/HashedIssueHandler.cs
@@ -3,6 +3,7 @@
 using System.Collections.Generic;
 using System.Linq;
 using System.Text;
+using System.Text.Json;
 using System.Text.RegularExpressions;
 using EpicGames.Core;
 using Microsoft.Extensions.Logging;
@@ -82,6 +83,7 @@ namespace EpicGames.Horde.Issues.Handlers
 					issue.Events.Add(stepEvent);
 					issue.Keys.AddHash(hash);
 					stepEvent.AuditLogger?.LogDebug("Fingerprint key: '{Key}' generated from event: '{Event}'", hash, stepEvent.Render());
+					AddFileMetadata(issue, stepEvent);
 					issue.Metadata.Add("Node", _context.NodeName);
 					issues.Add(issue);
 				}
@@ -98,6 +100,7 @@ namespace EpicGames.Horde.Issues.Handlers
 							issues.Add(genericErrorsFingerprint);
 						}
 						genericErrorsFingerprint.Events.Add(stepEvent);
+						AddFileMetadata(genericErrorsFingerprint, stepEvent);
 					}
 					else
 					{
@@ -110,6 +113,7 @@ namespace EpicGames.Horde.Issues.Handlers
 							issues.Add(genericFingerprint);
 						}
 						genericFingerprint.Events.Add(stepEvent);
+						AddFileMetadata(genericFingerprint, stepEvent);
 					}
 				}
 			}
@@ -117,6 +121,38 @@ namespace EpicGames.Horde.Issues.Handlers
 			return issues;
 		}
 
+		static void AddFileMetadata(IssueEventGroup issue, IssueEvent stepEvent)
+		{
+			try
+			{
+				foreach (JsonLogEvent line in stepEvent.Lines)
+				{
+					JsonDocument document = JsonDocument.Parse(line.Data);
+
+					JsonElement properties;
+					if (document.RootElement.TryGetProperty("properties", out properties) && properties.ValueKind == JsonValueKind.Object)
+					{
+						foreach (JsonProperty property in properties.EnumerateObject())
+						{
+							if (property.NameEquals("file") && property.Value.ValueKind == JsonValueKind.String)
+							{
+								string file = property.Value.GetString()!;
+								int endIdx = file.LastIndexOfAny(new char[] { '/', '\\' }) + 1;
+								string fileName = file.Substring(endIdx);
+
+								issue.Metadata.Add("File", fileName);
+								stepEvent.AuditLogger?.LogDebug("Metadata added for file: '{File}'", fileName);
+							}
+						}
+					}
+				}
+			}
+			catch (System.Exception ex)
+			{
+				stepEvent.AuditLogger?.LogError("Exception adding file metadata to issue: '{Exception}'", ex);
+			}
+		}
+
 		static bool TryGetHash(string message, out Md5Hash hash)
 		{
 			string sanitized = message.ToUpperInvariant();
