Added a new systemic log event for cache response timeouts from UBA
#rnx

diff --git a/Engine/Source/Programs/Shared/EpicGames.Core/KnownLogEvents.cs b/Engine/Source/Programs/Shared/EpicGames.Core/KnownLogEvents.cs
--- a/Engine/Source/Programs/Shared/EpicGames.Core/KnownLogEvents.cs
+++ b/Engine/Source/Programs/Shared/EpicGames.Core/KnownLogEvents.cs
@@ -346,6 +346,11 @@ namespace EpicGames.Core
 		/// </summary>
 		public static EventId Systemic_ManagedWorkspace { get; } = new EventId(769);
 
+		/// <summary>
+		/// Generic timeout message emmitted by UBA
+		/// </summary>
+		public static EventId Systemic_UbaCacheResponseTimeout { get; } = new EventId(770);
+
 		/// <summary>
 		/// Maximum code for systemic events. Add new events in the 700-799 range.
 		/// </summary>
