C# analysis fixes

#rnx

diff --git a/Engine/Source/Programs/Shared/EpicGames.Build/CsProjectInfo.cs b/Engine/Source/Programs/Shared/EpicGames.Build/CsProjectInfo.cs
--- a/Engine/Source/Programs/Shared/EpicGames.Build/CsProjectInfo.cs
+++ b/Engine/Source/Programs/Shared/EpicGames.Build/CsProjectInfo.cs
@@ -712,7 +712,10 @@ namespace EpicGames.Core
 						{
 							string component = remainingComponents.First();
 							if (component[0] == '*')
+							{
 								break;
+							}
+
 							subPath += Path.DirectorySeparatorChar + component;
 							remainingComponents = remainingComponents.Skip(1);
 						}
diff --git a/Engine/Source/Programs/Shared/EpicGames.Core/ManagedProcess.cs b/Engine/Source/Programs/Shared/EpicGames.Core/ManagedProcess.cs
--- a/Engine/Source/Programs/Shared/EpicGames.Core/ManagedProcess.cs
+++ b/Engine/Source/Programs/Shared/EpicGames.Core/ManagedProcess.cs
@@ -197,17 +197,23 @@ namespace EpicGames.Core
 		/// Returns the total CPU time usage for this job.
 		/// </summary>
 		public TimeSpan TotalProcessorTime =>
-			SupportsJobObjects && QueryInformationJobObject(JobHandle!, JobObjectBasicAccountingInformation, out JOBOBJECT_BASIC_ACCOUNTING_INFORMATION accountingInformation, Marshal.SizeOf(typeof(JOBOBJECT_BASIC_ACCOUNTING_INFORMATION)), out _)
+			SupportsJobObjects && JobHandle != null && QueryInformationJobObject(JobHandle, JobObjectBasicAccountingInformation, out JOBOBJECT_BASIC_ACCOUNTING_INFORMATION accountingInformation, Marshal.SizeOf(typeof(JOBOBJECT_BASIC_ACCOUNTING_INFORMATION)), out _)
 				? new TimeSpan((long)accountingInformation.TotalUserTime + (long)accountingInformation.TotalKernelTime)
 				: new TimeSpan();
 
+		/// <summary>
+		/// Returns the the peak memory for a process for this job.
+		/// </summary>
 		public long PeakProcessMemoryUsed =>
-			SupportsJobObjects && QueryInformationJobObject(JobHandle!, JobObjectExtendedLimitInformation, out JOBOBJECT_EXTENDED_LIMIT_INFORMATION accountingInformation, Marshal.SizeOf(typeof(JOBOBJECT_EXTENDED_LIMIT_INFORMATION)), out _)
+			SupportsJobObjects && JobHandle != null && QueryInformationJobObject(JobHandle, JobObjectExtendedLimitInformation, out JOBOBJECT_EXTENDED_LIMIT_INFORMATION accountingInformation, Marshal.SizeOf(typeof(JOBOBJECT_EXTENDED_LIMIT_INFORMATION)), out _)
 				? (long)accountingInformation.PeakProcessMemoryUsed
 				: 0;
 
+		/// <summary>
+		/// Returns the the peak memory for this job.
+		/// </summary>
 		public long PeakJobMemoryUsed =>
-			SupportsJobObjects && QueryInformationJobObject(JobHandle!, JobObjectExtendedLimitInformation, out JOBOBJECT_EXTENDED_LIMIT_INFORMATION accountingInformation, Marshal.SizeOf(typeof(JOBOBJECT_EXTENDED_LIMIT_INFORMATION)), out _)
+			SupportsJobObjects && JobHandle != null && QueryInformationJobObject(JobHandle, JobObjectExtendedLimitInformation, out JOBOBJECT_EXTENDED_LIMIT_INFORMATION accountingInformation, Marshal.SizeOf(typeof(JOBOBJECT_EXTENDED_LIMIT_INFORMATION)), out _)
 				? (long)accountingInformation.PeakJobMemoryUsed
 				: 0;
 
@@ -1242,6 +1248,9 @@ namespace EpicGames.Core
 		/// </summary>
 		public TimeSpan TotalProcessorTime => _accountingProcessGroup != null ? _accountingProcessGroup.TotalProcessorTime : new TimeSpan();
 
+		/// <summary>
+		/// Peak memory used
+		/// </summary>
 		public long PeakMemoryUsed => _accountingProcessGroup != null ? _accountingProcessGroup.PeakJobMemoryUsed : 0;
 
 		/// <summary>
diff --git a/Engine/Source/Programs/Shared/EpicGames.OIDC/OidcTokenManager.cs b/Engine/Source/Programs/Shared/EpicGames.OIDC/OidcTokenManager.cs
--- a/Engine/Source/Programs/Shared/EpicGames.OIDC/OidcTokenManager.cs
+++ b/Engine/Source/Programs/Shared/EpicGames.OIDC/OidcTokenManager.cs
@@ -712,18 +712,18 @@ namespace EpicGames.OIDC
 			}
 			else if (RuntimeInformation.IsOSPlatform(OSPlatform.Linux))
 			{
-				List<string> Browsers = ["xdg-open"];
-				bool IsWsl = IsRunningWsl();
-				if (IsWsl)
+				List<string> browsers = ["xdg-open"];
+				bool isWsl = IsRunningWsl();
+				if (isWsl)
 				{
-					Browsers.Add("wslview");
+					browsers.Add("wslview");
 				}
 
-				foreach (string Browser in Browsers)
+				foreach (string browser in browsers)
 				{
 					try
 					{
-						return Process.Start(Browser, url);
+						return Process.Start(browser, url);
 					}
 					catch (Exception e)
 					{
@@ -731,7 +731,7 @@ namespace EpicGames.OIDC
 					}
 				}
 
-				if (IsWsl)
+				if (isWsl)
 				{
 					_logger?.LogInformation("Did not find a browser, use 'sudo apt install xdg-utils' or 'sudo apt install wslu'.");
 				}
diff --git a/Engine/Source/Programs/Shared/EpicGames.UBA/Impl/SchedulerImpl.cs b/Engine/Source/Programs/Shared/EpicGames.UBA/Impl/SchedulerImpl.cs
--- a/Engine/Source/Programs/Shared/EpicGames.UBA/Impl/SchedulerImpl.cs
+++ b/Engine/Source/Programs/Shared/EpicGames.UBA/Impl/SchedulerImpl.cs
@@ -30,9 +30,13 @@ namespace EpicGames.UBA.Impl
 		static extern ulong ProcessHandle_GetTotalWallTime(nint process);
 		#endregion
 
-		internal ProcessFinishedInfoImpl(object userData, nint processHandle) { _userData = userData;  _processHandle = processHandle; }
-		object _userData;
-		nint _processHandle;
+		internal ProcessFinishedInfoImpl(object userData, nint processHandle)
+		{
+			_userData = userData;
+			_processHandle = processHandle;
+		}
+		readonly object _userData;
+		readonly nint _processHandle;
 
 		public ProcessExecutionType ExecutionType => (ProcessExecutionType)ProcessHandle_GetExecutionType(_processHandle);
 		public int ExitCode => (int)ProcessHandle_GetExitCode(_processHandle);
@@ -68,13 +72,17 @@ namespace EpicGames.UBA.Impl
 		{
 			_scheduler = scheduler;
 			_userData = userData;
-			_exitCallbackDelegate = RaiseExited;
-			_startInfoHandle = ProcessStartInfo_Create2(info.Application, info.Arguments, info.WorkingDirectory, info.Description, (uint)info.Priority, info.RootsHandle, (byte)(info.TrackInputs ? 1 : 0), info.LogFile ?? String.Empty, _exitCallbackDelegate);
+			ExitCallbackDelegate = RaiseExited;
+			_startInfoHandle = ProcessStartInfo_Create2(info.Application, info.Arguments, info.WorkingDirectory, info.Description, (uint)info.Priority, info.RootsHandle, (byte)(info.TrackInputs ? 1 : 0), info.LogFile ?? String.Empty, ExitCallbackDelegate);
 		}
 
 		void RaiseExited(nint userData, nint processHandle, ref byte exitedResponse)
 		{
-			while (_index == UInt32.MaxValue) { Thread.Sleep(1); } // This to prevent a (most likely impossible) race condition (Index is set after process is enqueued)
+			// This to prevent a (most likely impossible) race condition (Index is set after process is enqueued)
+			while (_index == UInt32.MaxValue)
+			{
+				Thread.Sleep(1);
+			}
 
 			//ProcessImpl process = new(handle);
 			ProcessFinishedInfoImpl info = new(_userData, processHandle);
@@ -83,9 +91,9 @@ namespace EpicGames.UBA.Impl
 		}
 
 		public delegate void ExitCallback(nint userData, nint processHandle, ref byte exitedResponse);
-		SchedulerImpl _scheduler;
-		object _userData;
-		event ExitCallback _exitCallbackDelegate;
+		readonly SchedulerImpl _scheduler;
+		readonly object _userData;
+		event ExitCallback ExitCallbackDelegate;
 		public nint _startInfoHandle;
 		public uint _index = UInt32.MaxValue;
 	}
@@ -120,7 +128,7 @@ namespace EpicGames.UBA.Impl
 
 		nint _schedulerHandle = IntPtr.Zero;
 		internal Func<IProcessFinishedInfo, ProcessFinishedResponse>? _processFinishedCallback;
-		List<SchedulerProcess> _processes = new();
+		readonly List<SchedulerProcess> _processes = [];
 
 		internal SchedulerImpl(ISessionServer server, IEnumerable<ICacheClient> cacheClients, int maxLocalProcessors, bool forceRemote)
 		{
diff --git a/Engine/Source/Programs/Shared/EpicGames.UHT/Parsers/UhtNamespaceParser.cs b/Engine/Source/Programs/Shared/EpicGames.UHT/Parsers/UhtNamespaceParser.cs
--- a/Engine/Source/Programs/Shared/EpicGames.UHT/Parsers/UhtNamespaceParser.cs
+++ b/Engine/Source/Programs/Shared/EpicGames.UHT/Parsers/UhtNamespaceParser.cs
@@ -83,7 +83,6 @@ namespace EpicGames.UHT.Parsers
 				topScope.TokenReader.Optional(';');
 			}
 
-
 			// Remove the namespaces from the active namespace
 			topScope.HeaderParser.PopNamespaces(namespaceCount);
 			return UhtParseResult.Handled;
@@ -95,12 +94,12 @@ namespace EpicGames.UHT.Parsers
 			classObj.ClassType = UhtClassType.VModule;
 
 			using UhtParsingScope topScope = new(parentScope, classObj, parentScope.Session.GetKeywordTable(UhtTableNames.VModule), UhtAccessSpecifier.Public);
-			const string scopeName = "VMODULE";
-			using UhtMessageContext tokenContext = new(scopeName);
+			const string ScopeName = "VMODULE";
+			using UhtMessageContext tokenContext = new(ScopeName);
 
 			// Parse the specifiers
 			UhtSpecifierContext specifierContext = new(topScope, topScope.TokenReader, classObj.MetaData);
-			UhtSpecifierParser specifiers = UhtSpecifierParser.GetThreadInstance(specifierContext, scopeName, parentScope.Session.GetSpecifierTable(UhtTableNames.Class));
+			UhtSpecifierParser specifiers = UhtSpecifierParser.GetThreadInstance(specifierContext, ScopeName, parentScope.Session.GetSpecifierTable(UhtTableNames.Class));
 			specifiers.ParseSpecifiers();
 			classObj.GeneratedBodyLineNumber = classObj.LineNumber;
 			classObj.GeneratedBodyAccessSpecifier = UhtAccessSpecifier.Public;
diff --git a/Engine/Source/Programs/Shared/EpicGames.UHT/Tables/UhtCodeGeneratorInjectorTable.cs b/Engine/Source/Programs/Shared/EpicGames.UHT/Tables/UhtCodeGeneratorInjectorTable.cs
--- a/Engine/Source/Programs/Shared/EpicGames.UHT/Tables/UhtCodeGeneratorInjectorTable.cs
+++ b/Engine/Source/Programs/Shared/EpicGames.UHT/Tables/UhtCodeGeneratorInjectorTable.cs
@@ -1,4 +1,4 @@
-﻿// Copyright Epic Games, Inc. All Rights Reserved.
+// Copyright Epic Games, Inc. All Rights Reserved.
 
 using System;
 using System.Collections.Generic;
@@ -124,11 +124,14 @@ namespace EpicGames.UHT.Tables
 			// Sort the delegates to ensure determinism across multiple runs
 			foreach (List<UhtCodeGeneratorInjectorDelegate>? injectors in _injectorsTable.Values)
 			{
-				if (injectors == null) { continue; }
-				injectors.Sort((Lhs, Rhs) =>
+				if (injectors == null)
 				{
-					MethodInfo methodLhs = Lhs.Method;
-					MethodInfo methodRhs = Rhs.Method;
+					continue;
+				}
+				injectors.Sort((lhs, rhs) =>
+				{
+					MethodInfo methodLhs = lhs.Method;
+					MethodInfo methodRhs = rhs.Method;
 					
 					// Sort by module
 					int moduleComparison = String.Compare(methodLhs.Module.Name, methodRhs.Module.Name, StringComparison.Ordinal);
diff --git a/Engine/Source/Programs/Shared/EpicGames.UHT/Utils/UhtSession.cs b/Engine/Source/Programs/Shared/EpicGames.UHT/Utils/UhtSession.cs
--- a/Engine/Source/Programs/Shared/EpicGames.UHT/Utils/UhtSession.cs
+++ b/Engine/Source/Programs/Shared/EpicGames.UHT/Utils/UhtSession.cs
@@ -1230,15 +1230,15 @@ namespace EpicGames.UHT.Utils
 
 			if (arguments != null)
 			{
-				foreach (UhtExporter Exporter in ExporterTable)
+				foreach (UhtExporter exporter in ExporterTable)
 				{
-					if (arguments.HasOption($"-{Exporter.Name}"))
+					if (arguments.HasOption($"-{exporter.Name}"))
 					{
-						SetExporterStatus(Exporter.Name, true);
+						SetExporterStatus(exporter.Name, true);
 					}
-					else if (arguments.HasOption($"-no{Exporter.Name}"))
+					else if (arguments.HasOption($"-no{exporter.Name}"))
 					{
-						SetExporterStatus(Exporter.Name, false);
+						SetExporterStatus(exporter.Name, false);
 					}
 				}
 			}
diff --git a/Engine/Source/Programs/Shared/UnrealBuildTool.Tests/TestUtilities/StaticTargetRulesInitializer.cs b/Engine/Source/Programs/Shared/UnrealBuildTool.Tests/TestUtilities/StaticTargetRulesInitializer.cs
--- a/Engine/Source/Programs/Shared/UnrealBuildTool.Tests/TestUtilities/StaticTargetRulesInitializer.cs
+++ b/Engine/Source/Programs/Shared/UnrealBuildTool.Tests/TestUtilities/StaticTargetRulesInitializer.cs
@@ -48,14 +48,14 @@ namespace UnrealBuildTool.Tests.TestUtilities
 			UEBuildPlatform? _;
 			if (!UEBuildPlatform.TryGetBuildPlatform(UnrealTargetPlatform.Win64, out _))
 			{
-				MicrosoftPlatformSDK SDK = new MicrosoftPlatformSDK(logger);
-				UEBuildPlatform.RegisterBuildPlatform(new WindowsPlatform(UnrealTargetPlatform.Win64, SDK, logger), logger);
+				MicrosoftPlatformSDK sdk = new MicrosoftPlatformSDK(logger);
+				UEBuildPlatform.RegisterBuildPlatform(new WindowsPlatform(UnrealTargetPlatform.Win64, sdk, logger), logger);
 			}
 
 			if (!UEBuildPlatform.TryGetBuildPlatform(UnrealTargetPlatform.Android, out _))
 			{
-				AndroidPlatformSDK ASDK = new AndroidPlatformSDK(logger);
-				UEBuildPlatform.RegisterBuildPlatform(new AndroidPlatform(ASDK, logger), logger);
+				AndroidPlatformSDK asdk = new AndroidPlatformSDK(logger);
+				UEBuildPlatform.RegisterBuildPlatform(new AndroidPlatform(asdk, logger), logger);
 			}
 
 			XmlConfig.ReadConfigFiles(null, null, logger);
diff --git a/Engine/Source/Programs/Shared/UnrealBuildTool.Tests/TestUtilities/TestBuildPlatform.cs b/Engine/Source/Programs/Shared/UnrealBuildTool.Tests/TestUtilities/TestBuildPlatform.cs
--- a/Engine/Source/Programs/Shared/UnrealBuildTool.Tests/TestUtilities/TestBuildPlatform.cs
+++ b/Engine/Source/Programs/Shared/UnrealBuildTool.Tests/TestUtilities/TestBuildPlatform.cs
@@ -7,31 +7,31 @@ namespace UnrealBuildTool.Tests.TestUtilities
 {
 	internal class TestBuildPlatfrom : UEBuildPlatform
 	{
-		public TestBuildPlatfrom(UnrealTargetPlatform InPlatform, UEBuildPlatformSDK SDK, UnrealArchitectureConfig ArchitectureConfig, ILogger InLogger) : base(InPlatform, SDK, ArchitectureConfig, InLogger)
+		public TestBuildPlatfrom(UnrealTargetPlatform inPlatform, UEBuildPlatformSDK sdk, UnrealArchitectureConfig architectureConfig, ILogger inLogger) : base(inPlatform, sdk, architectureConfig, inLogger)
 		{
 
 		}
-		public override UEToolChain CreateToolChain(ReadOnlyTargetRules Target)
+		public override UEToolChain CreateToolChain(ReadOnlyTargetRules target)
 		{
 			throw new System.NotImplementedException();
 		}
 
-		public override void Deploy(TargetReceipt Receipt)
+		public override void Deploy(TargetReceipt receipt)
 		{
 
 		}
 
-		public override bool IsBuildProduct(string FileName, string[] NamePrefixes, string[] NameSuffixes)
+		public override bool IsBuildProduct(string fileName, string[] namePrefixes, string[] nameSuffixes)
 		{
 			return false;
 		}
 
-		public override void SetUpEnvironment(ReadOnlyTargetRules Target, CppCompileEnvironment CompileEnvironment, LinkEnvironment LinkEnvironment)
+		public override void SetUpEnvironment(ReadOnlyTargetRules target, CppCompileEnvironment compileEnvironment, LinkEnvironment linkEnvironment)
 		{
 
 		}
 
-		public override bool ShouldCreateDebugInfo(ReadOnlyTargetRules Target)
+		public override bool ShouldCreateDebugInfo(ReadOnlyTargetRules target)
 		{
 			return false;
 		}
diff --git a/Engine/Source/Programs/Shared/UnrealBuildTool.Tests/TestUtilities/TestBuildPlatformFactory.cs b/Engine/Source/Programs/Shared/UnrealBuildTool.Tests/TestUtilities/TestBuildPlatformFactory.cs
--- a/Engine/Source/Programs/Shared/UnrealBuildTool.Tests/TestUtilities/TestBuildPlatformFactory.cs
+++ b/Engine/Source/Programs/Shared/UnrealBuildTool.Tests/TestUtilities/TestBuildPlatformFactory.cs
@@ -8,7 +8,7 @@ namespace UnrealBuildTool.Tests.TestUtilities
 	{
 		public override UnrealTargetPlatform TargetPlatform => UnrealTargetPlatform.Win64;
 
-		public override void RegisterBuildPlatforms(ILogger Logger)
+		public override void RegisterBuildPlatforms(ILogger logger)
 		{
 		}
 	}
diff --git a/Engine/Source/Programs/Shared/UnrealBuildTool.Tests/TestUtilities/TestBuildPlatformSDK.cs b/Engine/Source/Programs/Shared/UnrealBuildTool.Tests/TestUtilities/TestBuildPlatformSDK.cs
--- a/Engine/Source/Programs/Shared/UnrealBuildTool.Tests/TestUtilities/TestBuildPlatformSDK.cs
+++ b/Engine/Source/Programs/Shared/UnrealBuildTool.Tests/TestUtilities/TestBuildPlatformSDK.cs
@@ -7,13 +7,13 @@ namespace UnrealBuildTool.Tests.TestUtilities
 {
 	internal class TestSDK : UEBuildPlatformSDK
 	{
-		public TestSDK(ILogger InLogger) : base(InLogger)
+		public TestSDK(ILogger inLogger) : base(inLogger)
 		{
 		}
 
-		public override bool TryConvertVersionToInt(string? StringValue, out ulong OutValue, string? Hint = null)
+		public override bool TryConvertVersionToInt(string? stringValue, out ulong outValue, string? hint = null)
 		{
-			OutValue = 1;
+			outValue = 1;
 			return true;
 		}
 
diff --git a/Engine/Source/Programs/Shared/UnrealBuildTool.Tests/TestUtilities/TestToolChain.cs b/Engine/Source/Programs/Shared/UnrealBuildTool.Tests/TestUtilities/TestToolChain.cs
--- a/Engine/Source/Programs/Shared/UnrealBuildTool.Tests/TestUtilities/TestToolChain.cs
+++ b/Engine/Source/Programs/Shared/UnrealBuildTool.Tests/TestUtilities/TestToolChain.cs
@@ -13,22 +13,22 @@ namespace UnrealBuildTool.Tests.TestUtilities
 		{
 		}
 
-		public override FileItem? LinkFiles(LinkEnvironment LinkEnvironment, bool bBuildImportLibraryOnly, IActionGraphBuilder Graph)
+		public override FileItem? LinkFiles(LinkEnvironment linkEnvironment, bool bBuildImportLibraryOnly, IActionGraphBuilder graph)
 		{
 			return null;
 		}
 
-		protected override CPPOutput CompileCPPFiles(CppCompileEnvironment CompileEnvironment, IEnumerable<FileItem> InputFiles, DirectoryReference OutputDir, string ModuleName, IActionGraphBuilder Graph)
+		protected override CPPOutput CompileCPPFiles(CppCompileEnvironment compileEnvironment, IEnumerable<FileItem> inputFiles, DirectoryReference outputDir, string moduleName, IActionGraphBuilder graph)
 		{
 			return null!;
 		}
 
-		protected override CPPOutput CompileISPCFiles(CppCompileEnvironment CompileEnvironment, IEnumerable<FileItem> InputFiles, DirectoryReference OutputDir, IActionGraphBuilder Graph)
+		protected override CPPOutput CompileISPCFiles(CppCompileEnvironment compileEnvironment, IEnumerable<FileItem> inputFiles, DirectoryReference outputDir, IActionGraphBuilder graph)
 		{
 			return null!;
 		}
 
-		protected override CPPOutput GenerateISPCHeaders(CppCompileEnvironment CompileEnvironment, IEnumerable<FileItem> InputFiles, DirectoryReference OutputDir, IActionGraphBuilder Graph)
+		protected override CPPOutput GenerateISPCHeaders(CppCompileEnvironment compileEnvironment, IEnumerable<FileItem> inputFiles, DirectoryReference outputDir, IActionGraphBuilder graph)
 		{
 			return null!;
 		}
