BuildGraph: Fix missing docs for properties of BgDiagnosticDef

#rnx

diff --git a/Engine/Source/Programs/Shared/EpicGames.BuildGraph/BgDiagnosticDef.cs b/Engine/Source/Programs/Shared/EpicGames.BuildGraph/BgDiagnosticDef.cs
--- a/Engine/Source/Programs/Shared/EpicGames.BuildGraph/BgDiagnosticDef.cs
+++ b/Engine/Source/Programs/Shared/EpicGames.BuildGraph/BgDiagnosticDef.cs
@@ -31,12 +31,13 @@ namespace EpicGames.BuildGraph
 		public string Message { get; }
 
 		/// <summary>
-		///
+		/// Report the diagnostic during execution
 		/// </summary>
 		public bool ReportOnExecution { get; }
 
 		/// <summary>
-		/// 
+		/// An optional set of node names that when <see cref="ReportOnExecution"/> is true limits the reporting of the diagnostic
+		/// to only this set of nodes
 		/// </summary>
 		public string[] ReportOnNodes { get; }
 
