[Horde][BuildHealth][Server] Add using declaration to new JsonDocument reference to ensure disposal

diff --git a/Engine/Source/Programs/Shared/EpicGames.Horde/Issues/Handlers/HashedIssueHandler.cs b/Engine/Source/Programs/Shared/EpicGames.Horde/Issues/Handlers/HashedIssueHandler.cs
--- a/Engine/Source/Programs/Shared/EpicGames.Horde/Issues/Handlers/HashedIssueHandler.cs
+++ b/Engine/Source/Programs/Shared/EpicGames.Horde/Issues/Handlers/HashedIssueHandler.cs
@@ -127,8 +127,7 @@ namespace EpicGames.Horde.Issues.Handlers
 			{
 				foreach (JsonLogEvent line in stepEvent.Lines)
 				{
-					JsonDocument document = JsonDocument.Parse(line.Data);
-
+					using JsonDocument document = JsonDocument.Parse(line.Data);
 					JsonElement properties;
 					if (document.RootElement.TryGetProperty("properties", out properties) && properties.ValueKind == JsonValueKind.Object)
 					{
