Horde: Temp Storage to allow the maximum number of exports to write in a single request to be configurable from the agent settings

#pr 13349

diff --git a/Engine/Source/Programs/Horde/Drivers/JobDriver/Utility/TempStorage.cs b/Engine/Source/Programs/Horde/Drivers/JobDriver/Utility/TempStorage.cs
--- a/Engine/Source/Programs/Horde/Drivers/JobDriver/Utility/TempStorage.cs
+++ b/Engine/Source/Programs/Horde/Drivers/JobDriver/Utility/TempStorage.cs
@@ -779,6 +779,8 @@ namespace JobDriver.Utility
 					throw new TempStorageException($"Missing block \"{blockName}\" from node \"{nodeName}\". If running Horde against a 5.4 or earlier version of Unreal Engine, please make sure you have integrated //UE5/Main/Engine/Source/Programs/AutomationTool/BuildGraph/BuildGraph.cs#70.");
 				}
 
+				_ = Int32.TryParse(Environment.GetEnvironmentVariable("UE_HORDE_TEMP_STORAGE_MAX_BATCH_SIZE"), out int maxBatchSize);
+
 				for (int attempt = 0; ; attempt++)
 				{
 					StorageStats initialStats = storageNamespace.GetStats();
@@ -788,6 +790,7 @@ namespace JobDriver.Utility
 					ExtractOptions options = new ExtractOptions();
 					options.Progress = new ExtractStatsLogger(logger);
 					options.VerifyOutput = options.VerboseOutput = (attempt > 0);
+					options.MaxBatchSize = maxBatchSize;
 
 					await rootDirEntry.Handle.ExtractAsync(rootDir.ToDirectoryInfo(), options, logger, cancellationToken);
 
diff --git a/Engine/Source/Programs/Horde/HordeAgent/AgentSettings.cs b/Engine/Source/Programs/Horde/HordeAgent/AgentSettings.cs
--- a/Engine/Source/Programs/Horde/HordeAgent/AgentSettings.cs
+++ b/Engine/Source/Programs/Horde/HordeAgent/AgentSettings.cs
@@ -268,6 +268,11 @@ namespace HordeAgent
 		/// </summary>
 		public string[] HealthCheckEndpoints { get; set; } = [];
 
+		/// <summary>
+		/// Maximum number of temp storage exports to write in a single request
+		/// </summary>
+		public int? TempStorageMaxBatchSize { get; set; } = null;
+
 		/// <summary>
 		/// Gets the current server settings
 		/// </summary>
diff --git a/Engine/Source/Programs/Horde/HordeAgent/Leases/Handlers/JobHandler.cs b/Engine/Source/Programs/Horde/HordeAgent/Leases/Handlers/JobHandler.cs
--- a/Engine/Source/Programs/Horde/HordeAgent/Leases/Handlers/JobHandler.cs
+++ b/Engine/Source/Programs/Horde/HordeAgent/Leases/Handlers/JobHandler.cs
@@ -1,5 +1,6 @@
 // Copyright Epic Games, Inc. All Rights Reserved.
 
+using System.Globalization;
 using EpicGames.Core;
 using EpicGames.Horde;
 using EpicGames.Horde.Agents.Leases;
@@ -61,6 +62,10 @@ namespace HordeAgent.Leases.Handlers
 				environment[HordeHttpClient.HordeTokenEnvVarName] = executeTask.Token;
 				environment["UE_LOG_JSON_TO_STDOUT"] = "1";
 				environment["UE_HORDE_OTEL_SETTINGS"] = OpenTelemetrySettingsExtensions.Serialize(_settings.CurrentValue.OpenTelemetry, true);
+				if (_settings.CurrentValue.TempStorageMaxBatchSize.HasValue)
+				{
+					environment["UE_HORDE_TEMP_STORAGE_MAX_BATCH_SIZE"] = Convert.ToString(_settings.CurrentValue.TempStorageMaxBatchSize.Value, CultureInfo.InvariantCulture);
+				}
 
 				exitCode = await RunDotNetProcessAsync(driverAssembly, arguments, environment, AgentApp.IsSelfContained, logger, cancellationToken);
 				logger.LogInformation("Driver finished with exit code {ExitCode}", exitCode);				
diff --git a/Engine/Source/Programs/Shared/EpicGames.Horde/Storage/BatchChunkReader.cs b/Engine/Source/Programs/Shared/EpicGames.Horde/Storage/BatchChunkReader.cs
--- a/Engine/Source/Programs/Shared/EpicGames.Horde/Storage/BatchChunkReader.cs
+++ b/Engine/Source/Programs/Shared/EpicGames.Horde/Storage/BatchChunkReader.cs
@@ -21,11 +21,24 @@ namespace EpicGames.Horde.Storage
 	/// </summary>
 	class BatchChunkReader : BatchBlobReader<ChunkReadRequest>, IDisposable
 	{
+		public const int DefaultMaxBatchSize = 20;
+
 		/// <summary>
 		/// Maximum number of exports to write in a single request
 		/// </summary>
-		public int MaxBatchSize { get; } = 20;
+		public int MaxBatchSize
+		{
+			get => _maxBatchSize;
+			init
+			{
+				if (value > 0)
+				{
+					_maxBatchSize = value;
+				}
+			}
+		}
 
+		readonly int _maxBatchSize = DefaultMaxBatchSize;
 		readonly ChannelWriter<OutputChunk[]> _writer;
 
 		/// <summary>
diff --git a/Engine/Source/Programs/Shared/EpicGames.Horde/Storage/Nodes/DirectoryNode.Extract.cs b/Engine/Source/Programs/Shared/EpicGames.Horde/Storage/Nodes/DirectoryNode.Extract.cs
--- a/Engine/Source/Programs/Shared/EpicGames.Horde/Storage/Nodes/DirectoryNode.Extract.cs
+++ b/Engine/Source/Programs/Shared/EpicGames.Horde/Storage/Nodes/DirectoryNode.Extract.cs
@@ -156,6 +156,11 @@ namespace EpicGames.Horde.Storage.Nodes
 		/// Output verbose logging about operations being performed
 		/// </summary>
 		public bool VerboseOutput { get; set; }
+
+		/// <summary>
+		/// Maximum number of exports to write in a single request
+		/// </summary>
+		public int? MaxBatchSize { get; set; }
 	}
 
 	/// <summary>
@@ -250,7 +255,7 @@ namespace EpicGames.Horde.Storage.Nodes
 			int numReadTasks = options.NumReadTasks ?? ExtractOptions.DefaultNumReadTasks;
 			int numDecodeTasks = options.NumDecodeTasks ?? ExtractOptions.DefaultNumDecodeTasks;
 			int numWriteTasks = options.NumWriteTasks ?? Math.Min(1 + (int)(directoryNode.Length / (16 * 1024 * 1024)), 16);
-			logger.LogDebug("Using {NumReadTasks} read tasks, {NumDecodeTasks} decode tasks, {NumWriteTasks} write tasks", numReadTasks, numDecodeTasks, numWriteTasks);
+			int maxBatchSize = options.MaxBatchSize ?? BatchChunkReader.DefaultMaxBatchSize;
 
 			await using (AsyncPipeline pipeline = new AsyncPipeline(cancellationToken))
 			{
@@ -258,9 +263,11 @@ namespace EpicGames.Horde.Storage.Nodes
 				outputWriter.VerifyOutput = options.VerifyOutput;
 				outputWriter.VerboseOutput = options.VerboseOutput;
 
-				using BatchChunkReader chunkReader = new BatchChunkReader(outputWriter.RequestWriter);
+				using BatchChunkReader chunkReader = new BatchChunkReader(outputWriter.RequestWriter) { MaxBatchSize = maxBatchSize };
 				chunkReader.VerifyHashes = options.VerifyOutput;
 
+				logger.LogDebug("Using {NumReadTasks} read tasks, {NumDecodeTasks} decode tasks, {NumWriteTasks} write tasks, {MaxBatchSize} max batch size", numReadTasks, numDecodeTasks, numWriteTasks, chunkReader.MaxBatchSize);
+
 				Channel<ChunkReadRequest> blobRequests = Channel.CreateBounded<ChunkReadRequest>(new BoundedChannelOptions(chunkReader.MinQueueLength * 2));
 				_ = pipeline.AddTask(ctx => FindOutputChunksRootAsync(directoryInfo, directoryNode, blobRequests.Writer, options, logger, ctx));
 
