Integrating UHT changes for constinit uobjects

diff --git a/Engine/Source/Programs/Shared/EpicGames.Core/UHTTypes.cs b/Engine/Source/Programs/Shared/EpicGames.Core/UHTTypes.cs
--- a/Engine/Source/Programs/Shared/EpicGames.Core/UHTTypes.cs
+++ b/Engine/Source/Programs/Shared/EpicGames.Core/UHTTypes.cs
@@ -2,6 +2,7 @@
 
 using System;
 using System.Collections.Generic;
+using System.Diagnostics.CodeAnalysis;
 using System.Text.Json.Serialization;
 
 #pragma warning disable CS1591 // Missing documentation
@@ -99,6 +100,73 @@ namespace EpicGames.Core
 		InternalUser
 	}
 
+	/// <summary>
+	/// Types of code to generate to initialize compiled-in UObject instances. 
+	/// Flags to allow multiple formats to be generated at once for development/debugging.
+	/// </summary>
+	[Flags]
+	public enum UhtCompiledInObjectFormat
+	{
+		/// <summary>
+		/// No flags, should not be configured as a desired output type.
+		/// </summary>
+		None = 0x0,
+
+		/// <summary>
+		/// Generate params structs which will be used to initialize UObjects on the heap at runtime.
+		/// </summary>
+		Params = 0x1,
+
+		/// <summary>
+		/// Generate declarations of constinit variables of UObject-derived types allowing these
+		/// objects to be stored in the binary's data segment.
+		/// </summary>
+		ConstInit = 0x2,
+
+		/// <summary>
+		/// Default setting to use
+		/// </summary>
+		Default = Params,
+	
+	 	/// <summary>
+		/// Generate all possible types at once for comparison
+		/// </summary>
+		All = Params | ConstInit,
+	}
+
+	/// <summary>
+	/// Per-target settings for UHT to control features that need to be the same across all modules. 
+	/// </summary>
+	public record struct UHTTargetSettings
+	{
+		public UHTTargetSettings()
+		{
+		}
+
+		[SetsRequiredMembers]
+		public UHTTargetSettings(BinaryArchiveReader reader) : this()
+		{
+			CompiledInObjectFormat = Enum.Parse<UhtCompiledInObjectFormat>(reader.ReadString()!);
+		}
+
+		public void Write(BinaryArchiveWriter writer)
+		{
+			writer.WriteString(CompiledInObjectFormat.ToString());
+		}
+
+		/// <summary>
+		/// What type of code to generate for compiled-on UObject instances such as UClass instances.
+		/// </summary>
+		[JsonConverter(typeof(JsonStringEnumConverter))]
+		public UhtCompiledInObjectFormat CompiledInObjectFormat { get; set; } = UhtCompiledInObjectFormat.Default;
+
+		/// <summary>
+		/// Whether to generate a name for the maximum value of UENUMs where one was not found in the authored code.
+		/// If this is false, the UEnum code will generate it at runtime if it can do so without creating a collision. 
+		/// </summary>
+		public bool GenerateEnumMaxValues { get; set; } = false;
+	}
+
 	public class UHTManifest
 	{
 		public class Module
@@ -194,6 +262,11 @@ namespace EpicGames.Core
 		/// </summary>
 		public string ExternalDependenciesFile { get; set; } = String.Empty;
 
+		/// <summary>
+		/// Target-wide settings for UHT/code generation. 
+		/// </summary>
+		public UHTTargetSettings TargetSettings { set; get; } = new();
+
 		/// <summary>
 		/// List of modules
 		/// </summary>
diff --git a/Engine/Source/Programs/Shared/EpicGames.UHT/Exporters/CodeGen/UhtCodeGenerator.cs b/Engine/Source/Programs/Shared/EpicGames.UHT/Exporters/CodeGen/UhtCodeGenerator.cs
--- a/Engine/Source/Programs/Shared/EpicGames.UHT/Exporters/CodeGen/UhtCodeGenerator.cs
+++ b/Engine/Source/Programs/Shared/EpicGames.UHT/Exporters/CodeGen/UhtCodeGenerator.cs
@@ -12,8 +12,67 @@ using EpicGames.UHT.Utils;
 
 namespace EpicGames.UHT.Exporters.CodeGen
 {
+	/// <summary>
+	/// Type of singleton to retrieve a name or declaration for
+	/// </summary>
+	public enum UhtSingletonType
+	{
+		/// <summary>
+		/// Function that returns a fully constructed object - the distinction between this and Unregistered is only important for classes.
+		/// </summary>
+		Registered = 0,
+
+		/// <summary>
+		/// Function that returns a partially constructed object
+		/// </summary>
+		Unregistered = 1,
+
+		/// <summary>
+		/// Constinit variable (C++20 feature)
+		/// </summary>
+		ConstInit = 2,
+
+		/// <summary>
+		/// Struct containing initialization data for the object 
+		/// </summary>
+		Statics = 3,
+	}
+	
+	/// <summary>
+	/// Interface to provide singleton names during code generation. 
+	/// </summary>
+	public interface IUhtSingletonNameProvider
+	{
+		/// <summary>
+		/// Return the singleton name for the given object
+		/// </summary>
+		/// <param name="obj">Object in question - may be a struct, function, class etc or a package</param>
+		/// <param name="type">Type of declaration to get</param>
+		/// <returns>Singleton function name</returns>
+		public string GetSingletonName(UhtObject? obj, UhtSingletonType type);
+	}
+
+	/// <summary>
+	/// Interface to provide global settings during code generation. 
+	/// </summary>
+	public interface IUhtSessionSettingsProvider
+	{
+		/// <summary>
+		/// Returns whether UHT should generate a certain type of code for compiled-in object instances
+		/// </summary>
+		/// <param name="format">Format to test for</param>
+		/// <returns>Whether the given format is active - multiple formats may be active in certain development modes</returns>
+		public bool IsUsingCompiledInObjectFormat(UhtCompiledInObjectFormat format);
+
+		/// <summary>
+		/// Returns whether UHT is emitting code for multiple compiled-in object formats, and therefore such
+		/// code should be surrounded by preprocessor checks
+		/// </summary>
+		public bool IsUsingMultipleCompiledInObjectFormats { get; }
+	}
+
 	[UnrealHeaderTool]
-	class UhtCodeGenerator
+	class UhtCodeGenerator : IUhtSingletonNameProvider, IUhtSessionSettingsProvider
 	{
 		[UhtExporter(Name = "CodeGen", Description = "Standard UnrealEngine code generation", Options = UhtExporterOptions.Default,
 			CppFilters = new string[] { "*.generated.cpp", "*.generated.*.cpp", "*.gen.cpp", "*.gen.*.cpp" },
@@ -50,10 +109,23 @@ namespace EpicGames.UHT.Exporters.CodeGen
 
 		public struct ObjectInfo
 		{
+			// e.g. Z_Construct_UClass_UObject
 			public string RegisteredSingletonName { get; set; }
+			// e.g. Z_Construct_UClass_UObject_NoRegister
 			public string UnregisteredSingletonName { get; set; }
+			// e.g. COREUOBJECT_API* UClass* Z_Construct_UClass_UObject();
 			public string RegisteredExternalDecl { get; set; }
+			// e.g. COREUOBJECT_API* UClass* Z_Construct_UClass_UObject_NoRegister();
 			public string UnregisteredExternalDecl { get; set; }
+			// Fully qualified name of variable containing const-initialized object 
+			// e.g. Z_ConstInit_UClass_UObject
+			public string ConstInitSingletonName { get; set; }
+			// Full declaration to import variable containing const-initialized object
+			// e.g. COREUOBJECT_API extern UE_CONSTINIT_UOBJECT_DECL UClass Z_ConstInit_UClass_UObject
+			public string ConstInitExternalDecl { get; set; }
+			// Name of statics struct in generated cpp containing construction parameters
+			// e.g. Z_Construct_UClass_UObject_Statics
+			public string StaticsName { get; set; }
 			public UhtClass? NativeInterface { get; set; }
 			public UhtProperty? FastArrayProperty { get; set; }
 			public uint Hash { get; set; }
@@ -173,39 +245,69 @@ namespace EpicGames.UHT.Exporters.CodeGen
 		/// <summary>
 		/// Return the singleton name for an object
 		/// </summary>
-		/// <param name="obj">The object in question.</param>
-		/// <param name="registered">If true, return the registered singleton name.  Otherwise return the unregistered.</param>
+		/// <param name="obj">The type or package in question.</param>
+		/// <param name="type">Type of declaration to get</param>
 		/// <returns>Singleton name or "nullptr" if Object is null</returns>
-		public string GetSingletonName(UhtObject? obj, bool registered)
+		public string GetSingletonName(UhtObject? obj, UhtSingletonType type)
 		{
 			if (obj == null)
 			{
 				return "nullptr";
 			}
-			return registered ? ObjectInfos[obj.ObjectTypeIndex].RegisteredSingletonName : ObjectInfos[obj.ObjectTypeIndex].UnregisteredSingletonName;
+			switch (type)
+			{
+				case UhtSingletonType.Statics:
+					return ObjectInfos[obj.ObjectTypeIndex].StaticsName;
+				case UhtSingletonType.Registered:
+					return ObjectInfos[obj.ObjectTypeIndex].RegisteredSingletonName;
+				case UhtSingletonType.Unregistered:
+					return ObjectInfos[obj.ObjectTypeIndex].UnregisteredSingletonName;
+				case UhtSingletonType.ConstInit:
+					return ObjectInfos[obj.ObjectTypeIndex].ConstInitSingletonName;
+			}
+			throw new UhtIceException($"Invalid UhtSingletonType value {type}");
 		}
 
 		/// <summary>
 		/// Return the external declaration for an object
 		/// </summary>
 		/// <param name="obj">The object in question.</param>
-		/// <param name="registered">If true, return the registered external declaration.  Otherwise return the unregistered.</param>
+		/// <param name="type">Type of declaration to get</param>
 		/// <returns>External declaration</returns>
-		public string GetExternalDecl(UhtObject obj, bool registered)
+		public string GetExternalDecl(UhtObject obj, UhtSingletonType type)
 		{
-			return GetExternalDecl(obj.ObjectTypeIndex, registered);
+			return GetExternalDecl(obj.ObjectTypeIndex, type);
 		}
 
 		/// <summary>
 		/// Return the external declaration for an object
 		/// </summary>
 		/// <param name="objectIndex">The object in question.</param>
-		/// <param name="registered">If true, return the registered external declaration.  Otherwise return the unregistered.</param>
+		/// <param name="type">Type of declaration to get</param>
 		/// <returns>External declaration</returns>
-		public string GetExternalDecl(int objectIndex, bool registered)
+		public string GetExternalDecl(int objectIndex, UhtSingletonType type)
 		{
-			return registered ? ObjectInfos[objectIndex].RegisteredExternalDecl : ObjectInfos[objectIndex].UnregisteredExternalDecl;
+			switch (type)
+			{
+				case UhtSingletonType.Registered:
+					return ObjectInfos[objectIndex].RegisteredExternalDecl;
+				case UhtSingletonType.Unregistered:
+					return ObjectInfos[objectIndex].UnregisteredExternalDecl;
+				case UhtSingletonType.ConstInit:
+					return ObjectInfos[objectIndex].ConstInitExternalDecl;
+				case UhtSingletonType.Statics:
+					throw new UhtIceException($"Statics struct for {ObjectInfos[objectIndex].StaticsName} cannot be imported across translation units");
+			}
+			throw new UhtIceException($"Invalid UhtSingletonType value {type}");
 		}
+
+		public bool IsUsingCompiledInObjectFormat(UhtCompiledInObjectFormat format) => Session.IsUsingCompiledInObjectFormat(format);
+
+		/// <summary>
+		/// Returns whether UHT is emitting code for multiple compiled-in object formats, and therefore such
+		/// code should be surrounded by preprocessor checks
+		/// </summary>
+		public bool IsUsingMultipleCompiledInObjectFormats => Session.IsUsingMultipleCompiledInObjectFormats;
 		#endregion
 
 		#region Information initialization
@@ -234,7 +336,10 @@ namespace EpicGames.UHT.Exporters.CodeGen
 			builder.Append("Z_Construct_UPackage_");
 			builder.Append(packageInfo.StrippedName);
 			objectInfo.UnregisteredSingletonName = objectInfo.RegisteredSingletonName = builder.ToString();
+			objectInfo.StaticsName = $"{objectInfo.RegisteredSingletonName}_Statics";
 			objectInfo.UnregisteredExternalDecl = objectInfo.RegisteredExternalDecl = $"\tUPackage* {objectInfo.RegisteredSingletonName}();\r\n";
+			objectInfo.ConstInitSingletonName = $"Z_ConstInit_UPackage_{packageInfo.StrippedName}";
+			objectInfo.ConstInitExternalDecl = $"\textern UE_CONSTINIT_UOBJECT_DECL UPackage {objectInfo.ConstInitSingletonName};\r\n";
 		}
 
 		private void InitHeaderInfo(StringBuilder builder, UhtHeaderFile headerFile)
@@ -301,7 +406,6 @@ namespace EpicGames.UHT.Exporters.CodeGen
 
 		private void InitObjectInfo(StringBuilder builder, UhtPackage package, ref HeaderInfo headerInfo, UhtObject obj)
 		{
-			UhtModule module = package.Module;
 			ref ObjectInfo objectInfo = ref ObjectInfos[obj.ObjectTypeIndex];
 
 			builder.Clear();
@@ -310,7 +414,7 @@ namespace EpicGames.UHT.Exporters.CodeGen
 			bool isNonIntrinsicClass = false;
 			builder.Append("Z_Construct_U").Append(obj.EngineClassName).AppendOuterNames(obj);
 
-			string engineClassName = obj.EngineClassName;
+			string api = obj.Module.Api;
 			if (obj is UhtClass classObj)
 			{
 				if (!classObj.ClassFlags.HasAnyFlags(EClassFlags.Intrinsic))
@@ -366,11 +470,6 @@ namespace EpicGames.UHT.Exporters.CodeGen
 					}
 				}
 			}
-			else if (obj is UhtFunction)
-			{
-				// The method for EngineClassName returns type specific where in this case we need just the simple return type
-				engineClassName = "Function";
-			}
 			else if (obj is UhtEnum enumObj)
 			{
 				if (enumObj.IsVerseField)
@@ -384,15 +483,16 @@ namespace EpicGames.UHT.Exporters.CodeGen
 				objectInfo.RegisteredSingletonName = builder.ToString();
 				builder.Append("_NoRegister");
 				objectInfo.UnregisteredSingletonName = builder.ToString();
-
-				objectInfo.UnregisteredExternalDecl = $"\t{module.Api}U{engineClassName}* {objectInfo.UnregisteredSingletonName}();\r\n";
-				objectInfo.RegisteredExternalDecl = $"\t{module.Api}U{engineClassName}* {objectInfo.RegisteredSingletonName}();\r\n";
 			}
 			else
 			{
 				objectInfo.UnregisteredSingletonName = objectInfo.RegisteredSingletonName = builder.ToString();
-				objectInfo.UnregisteredExternalDecl = objectInfo.RegisteredExternalDecl = $"\t{module.Api}U{engineClassName}* {objectInfo.RegisteredSingletonName}();\r\n";
 			}
+			objectInfo.StaticsName = $"{objectInfo.RegisteredSingletonName}_Statics";
+			objectInfo.UnregisteredExternalDecl = $"\t{api}U{obj.EngineLinkClassName}* {objectInfo.UnregisteredSingletonName}();\r\n";
+			objectInfo.RegisteredExternalDecl = $"\t{api}U{obj.EngineLinkClassName}* {objectInfo.RegisteredSingletonName}();\r\n";
+			objectInfo.ConstInitSingletonName = objectInfo.RegisteredSingletonName.Replace("Z_Construct_", "Z_ConstInit_", StringComparison.InvariantCultureIgnoreCase);
+			objectInfo.ConstInitExternalDecl = $"\t{api}extern UE_CONSTINIT_UOBJECT_DECL TNoDestroy<U{obj.EngineClassName}> {objectInfo.ConstInitSingletonName};\r\n";
 
 			// Init the children
 			foreach (UhtType child in obj.Children)
diff --git a/Engine/Source/Programs/Shared/EpicGames.UHT/Exporters/CodeGen/UhtHeaderCodeGenerator.cs b/Engine/Source/Programs/Shared/EpicGames.UHT/Exporters/CodeGen/UhtHeaderCodeGenerator.cs
--- a/Engine/Source/Programs/Shared/EpicGames.UHT/Exporters/CodeGen/UhtHeaderCodeGenerator.cs
+++ b/Engine/Source/Programs/Shared/EpicGames.UHT/Exporters/CodeGen/UhtHeaderCodeGenerator.cs
@@ -225,7 +225,7 @@ namespace EpicGames.UHT.Exporters.CodeGen
 			builder.Append(' ');
 			if (!isDeclaration && outerClass != null)
 			{
-				builder.Append(outerClass.Namespace.FullSourceName).AppendClassSourceNameOrInterfaceName(outerClass).Append("::");
+				builder.AppendClassSourceNameOrInterfaceName(outerClass).Append("::");
 			}
 
 			if (alternateFunctionName != null)
