Integrating UHT changes for constinit uobjects

diff --git a/Engine/Source/Programs/Shared/EpicGames.UHT/Exporters/CodeGen/UhtHeaderCodeGeneratorCppFile.cs b/Engine/Source/Programs/Shared/EpicGames.UHT/Exporters/CodeGen/UhtHeaderCodeGeneratorCppFile.cs
--- a/Engine/Source/Programs/Shared/EpicGames.UHT/Exporters/CodeGen/UhtHeaderCodeGeneratorCppFile.cs
+++ b/Engine/Source/Programs/Shared/EpicGames.UHT/Exporters/CodeGen/UhtHeaderCodeGeneratorCppFile.cs
@@ -2,6 +2,7 @@
 
 using System;
 using System.Collections.Generic;
+using System.Diagnostics.Eventing.Reader;
 using System.Linq;
 using System.Text;
 using EpicGames.Core;
@@ -17,6 +18,7 @@ namespace EpicGames.UHT.Exporters.CodeGen
 	{
 		public UhtUsedDefineScopes<UhtEnum> Enumerations { get; } = new();
 		public UhtUsedDefineScopes<UhtScriptStruct> ScriptStructs { get; } = new();
+		public UhtUsedDefineScopes<UhtScriptStruct> RigVMStructs { get; } = new();
 		public UhtUsedDefineScopes<UhtClass> Classes { get; } = new();
 	}
 
@@ -51,6 +53,7 @@ namespace EpicGames.UHT.Exporters.CodeGen
 				bool addedStructuredArchiveFromArchiveHeader = false;
 				bool addedArchiveUObjectFromStructuredArchiveHeader = false;
 				bool addedCoreNetHeader = false;
+				bool addedRigVMHeaders = false;
 				HashSet<UhtHeaderFile> addedIncludes = new();
 				List<string> includesToAdd = new();
 				addedIncludes.Add(HeaderFile);
@@ -65,6 +68,24 @@ namespace EpicGames.UHT.Exporters.CodeGen
 					includesToAdd.Add("VerseInteropUtils.h");
 				}
 
+				if (HeaderFile.Children.Any(x => x is UhtScriptStruct structObj && structObj.IsVerseField))
+				{
+					includesToAdd.Add("VerseVM/VVMVerseStruct.h");
+				}
+				if (HeaderFile.Children.Any(x => x is UhtClass classObj && classObj.IsVerseField))
+				{
+					includesToAdd.Add("VerseVM/VVMVerseClass.h");
+					includesToAdd.Add("VerseVM/VVMVerseFunction.h");
+				}
+				if (HeaderFile.Children.Any(x => x is UhtEnum enumObj && enumObj.IsVerseField))
+				{
+					includesToAdd.Add("VerseVM/VVMVerseEnum.h");
+				}
+				if (HeaderFile.Children.Any(x => x is UhtStruct structObj && x.Children.Any(y => y is UhtVerseValueProperty)))
+				{
+					includesToAdd.Add("UObject/VerseValueProperty.h"); 
+				}
+
 				foreach (UhtType type in HeaderFile.Children)
 				{
 					if (type is UhtStruct structObj)
@@ -110,6 +131,16 @@ namespace EpicGames.UHT.Exporters.CodeGen
 						}
 					}
 
+					if (type is UhtScriptStruct scriptStruct)
+					{
+						if (!addedRigVMHeaders && scriptStruct.RigVMStructInfo is not null)
+						{
+							includesToAdd.Add("RigVMCore/RigVMFunction.h");
+							includesToAdd.Add("RigVMCore/RigVMRegistry.h");
+							addedRigVMHeaders = true;
+						}
+					}
+
 					if (type is UhtClass classObj)
 					{
 						if (classObj.ClassWithin != Session.UObject && !classObj.ClassWithin.HeaderFile.IsNoExportTypes)
@@ -160,18 +191,54 @@ namespace EpicGames.UHT.Exporters.CodeGen
 				builder.Append("\r\n");
 				builder.Append(DisableDeprecationWarnings).Append("\r\n");
 
+				if (!Session.IsUsingMultipleCompiledInObjectFormats)
+				{
+					if (Session.IsUsingCompiledInObjectFormat(UhtCompiledInObjectFormat.ConstInit))
+					{
+						builder.Append("static_assert(UE_WITH_CONSTINIT_UOBJECT, \"This generated code can only be compiled with UE_WITH_CONSTINIT_OBJECT\");");
+					}
+					if (Session.IsUsingCompiledInObjectFormat(UhtCompiledInObjectFormat.Params))
+					{
+						builder.Append("static_assert(!UE_WITH_CONSTINIT_UOBJECT, \"This generated code can only be compiled with !UE_WITH_CONSTINIT_OBJECT\");");
+					}
+				}
+
 				builder.Append("\r\n");
 				string cleanFileName = HeaderFile.FileNameWithoutExtension.Replace('.', '_');
 				builder.Append("void EmptyLinkFunctionForGeneratedCode").Append(cleanFileName).Append("() {}\r\n");
 
-				if (HeaderFile.References.CrossModule.References.Count > 0)
+				if (!HeaderFile.References.CrossModule.IsEmpty)
 				{
-					ReadOnlyMemory<string> sorted = HeaderFile.References.CrossModule.GetSortedReferences(
-						(int objectIndex, bool registered) => GetExternalDecl(objectIndex, registered));
 					using UhtCodeBlockComment blockComment = new(builder, "Cross Module References");
-					foreach (string crossReference in sorted.Span)
+					if (Session.IsUsingCompiledInObjectFormat(UhtCompiledInObjectFormat.ConstInit))
 					{
-						builder.Append(crossReference.AsSpan().TrimStart());
+						using UhtConditionalMacroBlock block = new(builder, "UE_WITH_CONSTINIT_UOBJECT", Session.IsUsingMultipleCompiledInObjectFormats);
+						// Cross-references based on constinit variables
+						IReadOnlyList<string> sorted = HeaderFile.References.CrossModule.GetSortedReferences(
+							(objectIndex, type) => type switch
+							{
+								UhtSingletonType.Statics => null,
+								_ => GetExternalDecl(objectIndex, UhtSingletonType.ConstInit)
+							});
+						foreach (string crossReference in sorted)
+						{
+							builder.Append(crossReference.AsSpan().TrimStart());
+						}
+					}
+					if (Session.IsUsingCompiledInObjectFormat(UhtCompiledInObjectFormat.Params))
+					{
+						using UhtConditionalMacroBlock block = new(builder, "!UE_WITH_CONSTINIT_UOBJECT", Session.IsUsingMultipleCompiledInObjectFormats);
+						// Cross-references based on factory functions
+						IReadOnlyList<string> sorted = HeaderFile.References.CrossModule.GetSortedReferences(
+							(objectIndex, type) => type switch
+							{
+								UhtSingletonType.ConstInit or UhtSingletonType.Statics => null,
+								_ => GetExternalDecl(objectIndex, type)
+							});
+						foreach (string crossReference in sorted)
+						{
+							builder.Append(crossReference.AsSpan().TrimStart());
+						}
 					}
 				}
 
@@ -207,6 +274,10 @@ namespace EpicGames.UHT.Exporters.CodeGen
 						if (scriptStruct.ScriptStructFlags.HasAnyFlags(EStructFlags.Native))
 						{
 							GetRegistrations(packageRegistrations, field).ScriptStructs.Add(scriptStruct);
+							if (scriptStruct.RigVMStructInfo != null && scriptStruct.RigVMStructInfo.Methods.Count > 0)
+							{
+								GetRegistrations(packageRegistrations, field).RigVMStructs.Add(scriptStruct);
+							}
 						}
 					}
 					else if (field is UhtFunction function)
@@ -258,92 +329,125 @@ namespace EpicGames.UHT.Exporters.CodeGen
 					uint combinedHash = UInt32.MaxValue;
 
 					using UhtCodeBlockComment blockComment = new(builder, "Registration");
-					builder.Append("struct ").Append(staticsName).Append("\r\n");
+					builder.Append($"struct {staticsName}\r\n");
 					builder.Append("{\r\n");
 
-					builder.AppendInstances(registrations.Enumerations, UhtDefineScopeNames.Standard, 
-						builder => builder.Append("\tstatic constexpr FEnumRegisterCompiledInInfo EnumInfo[] = {\r\n"),
-						(builder, enumObj) =>
-						{
-							uint hash = ObjectInfos[enumObj.ObjectTypeIndex].Hash;
-							builder
-								.Append("\t\t{ ")
-								.Append(enumObj.SourceName)
-								.Append("_StaticEnum, TEXT(\"")
-								.Append(enumObj.EngineName)
-								.Append("\"), &Z_Registration_Info_UEnum_")
-								.Append(enumObj.SourceName)
-								.Append($", CONSTRUCT_RELOAD_VERSION_INFO(FEnumReloadVersionInfo, {hash}U) }},\r\n");
-							combinedHash = HashCombine(combinedHash, hash);
-						},
-						builder => builder.Append("\t};\r\n"));
-
-					builder.AppendInstances(registrations.ScriptStructs, UhtDefineScopeNames.Standard,
-						builder => builder.Append("\tstatic constexpr FStructRegisterCompiledInInfo ScriptStructInfo[] = {\r\n"),
-						(builder, scriptStruct) =>
-						{
-							uint hash = ObjectInfos[scriptStruct.ObjectTypeIndex].Hash;
-							builder
-								.Append("\t\t{ ")
-								.Append(scriptStruct.Namespace.FullSourceName)
-								.Append(scriptStruct.SourceName)
-								.Append("::StaticStruct, Z_Construct_UScriptStruct_")
-								.Append(scriptStruct.SourceName)
-								.Append("_Statics::NewStructOps, TEXT(\"")
-								.Append(scriptStruct.EngineName)
-								.Append("\"), &Z_Registration_Info_UScriptStruct_")
-								.Append(scriptStruct.SourceName)
-								.Append(", CONSTRUCT_RELOAD_VERSION_INFO(FStructReloadVersionInfo, sizeof(")
-								.Append(scriptStruct.NamespaceExportName)
-								.Append(scriptStruct.SourceName)
-								.Append($"), {hash}U) }},\r\n");
-							combinedHash = HashCombine(combinedHash, hash);
-						},
-						builder => builder.Append("\t};\r\n"));
-
-					builder.AppendInstances(registrations.Classes, UhtDefineScopeNames.Standard,
-						builder => builder.Append("\tstatic constexpr FClassRegisterCompiledInInfo ClassInfo[] = {\r\n"),
-						(builder, classObj) =>
+					if (Session.IsUsingCompiledInObjectFormat(UhtCompiledInObjectFormat.ConstInit))
+					{
+						using UhtConditionalMacroBlock block = new(builder, "UE_WITH_CONSTINIT_UOBJECT", Session.IsUsingMultipleCompiledInObjectFormats);
+						// Constinit versions 
+						// Array of pointers to constinit objects
+						builder.AppendInstances(registrations.Enumerations, UhtDefineScopeNames.Standard,
+							builder => builder.Append("\tstatic inline constinit UEnum* EnumInfo[] = {\r\n"),
+							(builder, enumObj) => builder.Append("\t\t").AppendConstInitSingletonRef(this, enumObj).Append(",\r\n"),
+							builder => builder.Append("\t};\r\n"));
+
+						builder.AppendInstances(registrations.ScriptStructs, UhtDefineScopeNames.Standard,
+							builder => builder.Append("\tstatic inline constinit UScriptStruct* ScriptStructInfo[] = {\r\n"),
+							(builder, scriptStruct) => builder.Append("\t\t").AppendConstInitSingletonRef(this, scriptStruct).Append(",\r\n"),
+							builder => builder.Append("\t};\r\n"));
+
+						builder.AppendInstances(registrations.Classes, UhtDefineScopeNames.Standard,
+							builder => builder.Append("\tstatic inline constinit UClass* ClassInfo[] = {\r\n"),
+							(builder, classObj) => builder.Append("\t\t").AppendConstInitSingletonRef(this, classObj).Append(",\r\n"),
+							builder => builder.Append("\t};\r\n"));
+
+						// Registration for RigVM functions
+						builder.AppendInstances(registrations.RigVMStructs, UhtDefineScopeNames.Standard,
+							builder => builder.Append("\tstatic inline constinit FRigVMCompiledInStruct RigVMStructs[] = {\r\n"),
+							(builder, scriptStruct) => builder.Append("\t\t{ .Struct = ").AppendConstInitSingletonRef(this, scriptStruct).Append(", ")
+								 .Append(".Functions = MakeArrayView(").AppendSingletonName(this, scriptStruct, UhtSingletonType.Statics).Append("::RigVMFunctions), ")
+								.Append(" },\r\n"),
+							builder => builder.Append("\t};\r\n"));
+					}
+
+					if (Session.IsUsingCompiledInObjectFormat(UhtCompiledInObjectFormat.Params))
+					{ 
+						using UhtConditionalMacroBlock block = new(builder, "!UE_WITH_CONSTINIT_UOBJECT", Session.IsUsingMultipleCompiledInObjectFormats);
+						// Non-constinit version
+						builder.AppendInstances(registrations.Enumerations, UhtDefineScopeNames.Standard, 
+							builder => builder.Append("\tstatic constexpr FEnumRegisterCompiledInInfo EnumInfo[] = {\r\n"),
+							(builder, enumObj) =>
+							{ 
+								uint hash = ObjectInfos[enumObj.ObjectTypeIndex].Hash;
+								builder
+									.Append($"\t\t{{ {enumObj.SourceName}_StaticEnum, ")
+									.Append($"TEXT(\"{enumObj.EngineName}\"), ")
+									.Append($"&Z_Registration_Info_UEnum_{enumObj.SourceName}, ")
+									.Append($"CONSTRUCT_RELOAD_VERSION_INFO(FEnumReloadVersionInfo, {hash}U) }},\r\n");
+								combinedHash = HashCombine(combinedHash, hash);
+							},
+							builder => builder.Append("\t};\r\n"));
+
+						builder.AppendInstances(registrations.ScriptStructs, UhtDefineScopeNames.Standard,
+							builder => builder.Append("\tstatic constexpr FStructRegisterCompiledInInfo ScriptStructInfo[] = {\r\n"),
+							(builder, scriptStruct) =>
+							{
+								uint hash = ObjectInfos[scriptStruct.ObjectTypeIndex].Hash;
+								builder
+									.Append($"\t\t{{ {scriptStruct.FullyQualifiedSourceName}::StaticStruct, ")
+									.Append($"{GetSingletonName(scriptStruct, UhtSingletonType.Statics)}::NewStructOps, ")
+									.Append($"TEXT(\"{scriptStruct.EngineName}\"),")
+									.Append($"&Z_Registration_Info_UScriptStruct_{scriptStruct.SourceName}, ")
+									.Append($"CONSTRUCT_RELOAD_VERSION_INFO(FStructReloadVersionInfo, sizeof({scriptStruct.FullyQualifiedSourceName}), {hash}U) }},\r\n");
+								combinedHash = HashCombine(combinedHash, hash);
+							},
+							builder => builder.Append("\t};\r\n"));
+
+						builder.AppendInstances(registrations.Classes, UhtDefineScopeNames.Standard,
+							builder => builder.Append("\tstatic constexpr FClassRegisterCompiledInInfo ClassInfo[] = {\r\n"),
+							(builder, classObj) =>
+							{
+								uint hash = ObjectInfos[classObj.ObjectTypeIndex].Hash;
+								builder
+									.Append($"\t\t{{ {GetSingletonName(classObj, UhtSingletonType.Registered)}, ")
+									.Append($"{classObj.FullyQualifiedSourceName}::StaticClass, ")
+									.Append($"TEXT(\"{(classObj.IsVerseField ? classObj.EngineName : classObj.SourceName)}\"), ")
+									.Append($"&Z_Registration_Info_UClass_{classObj.SourceName}, ")
+									.Append($"CONSTRUCT_RELOAD_VERSION_INFO(FClassReloadVersionInfo, sizeof({classObj.FullyQualifiedSourceName}), {hash}U) }},\r\n");
+								combinedHash = HashCombine(combinedHash, hash);
+							},
+							builder => builder.Append("\t};\r\n"));
+					}
+
+					builder.Append($"}}; // {staticsName} \r\n");
+
+					if (Session.IsUsingCompiledInObjectFormat(UhtCompiledInObjectFormat.ConstInit))
+					{ 
+						using UhtConditionalMacroBlock block = new(builder, "UE_WITH_CONSTINIT_UOBJECT", Session.IsUsingMultipleCompiledInObjectFormats);
+						builder.Append($"static FRegisterCompiledInObjects {name}_{combinedHash}{{\r\n")
+							.AppendArrayView(registrations.Classes, UhtDefineScopeNames.Standard, staticsName, "ClassInfo", 1, ",\r\n")
+							.AppendArrayView(registrations.ScriptStructs, UhtDefineScopeNames.Standard, staticsName, "ScriptStructInfo", 1, ",\r\n")
+							.AppendArrayView(registrations.Enumerations, UhtDefineScopeNames.Standard, staticsName, "EnumInfo", 1, ",\r\n")
+							.Append("};\r\n");
+						if (!registrations.RigVMStructs.IsEmpty)
 						{
-							uint hash = ObjectInfos[classObj.ObjectTypeIndex].Hash;
-							builder
-								.Append("\t\t{ Z_Construct_UClass_")
-								.Append(classObj.SourceName)
-								.Append(", ")
-								.Append(classObj.Namespace.FullSourceName)
-								.Append(classObj.SourceName)
-								.Append("::StaticClass, TEXT(\"")
-								.Append(classObj.IsVerseField ? classObj.EngineName : classObj.SourceName)
-								.Append("\"), &Z_Registration_Info_UClass_")
-								.Append(classObj.SourceName)
-								.Append(", CONSTRUCT_RELOAD_VERSION_INFO(FClassReloadVersionInfo, sizeof(")
-								.Append(classObj.Namespace.FullSourceName)
-								.Append(classObj.SourceName)
-								.Append($"), {hash}U) }},\r\n");
-							combinedHash = HashCombine(combinedHash, hash);
-						},
-						builder => builder.Append("\t};\r\n"));
-
-					builder.Append("};\r\n");
-
-					builder
-						.Append("static FRegisterCompiledInInfo ")
-						.Append(name)
-						.Append($"_{combinedHash}(TEXT(\"")
-						.Append(package.EngineName)
-						.Append("\"),\r\n");
-
-					builder.AppendArrayPtrAndCountLine(registrations.Classes, UhtDefineScopeNames.Standard, staticsName, "ClassInfo", 1, ",\r\n");
-					builder.AppendArrayPtrAndCountLine(registrations.ScriptStructs, UhtDefineScopeNames.Standard, staticsName, "ScriptStructInfo", 1, ",\r\n");
-					builder.AppendArrayPtrAndCountLine(registrations.Enumerations, UhtDefineScopeNames.Standard, staticsName, "EnumInfo", 1, ");\r\n");
+							builder.Append($"static FRegisterRigVMStructs RigVM_{name}{{\r\n")
+								.AppendArrayView(registrations.RigVMStructs, UhtDefineScopeNames.Standard, staticsName, "RigVMStructs", 1, "\r\n")
+								.Append("};\r\n");
+						}
+					}
+					if (Session.IsUsingCompiledInObjectFormat(UhtCompiledInObjectFormat.Params))
+					{
+						using UhtConditionalMacroBlock block = new(builder, "!UE_WITH_CONSTINIT_UOBJECT", Session.IsUsingMultipleCompiledInObjectFormats);
+						builder.Append("static FRegisterCompiledInInfo ")
+							.Append(name)
+							.Append($"_{combinedHash}{{\r\n")
+							.Append("\tTEXT(\"")
+							.Append(package.EngineName)
+							.Append("\"),\r\n")
+							.AppendArrayPtrAndCountLine(registrations.Classes, UhtDefineScopeNames.Standard, staticsName, "ClassInfo", 1, ",\r\n")
+							.AppendArrayPtrAndCountLine(registrations.ScriptStructs, UhtDefineScopeNames.Standard, staticsName, "ScriptStructInfo", 1, ",\r\n")
+							.AppendArrayPtrAndCountLine(registrations.Enumerations, UhtDefineScopeNames.Standard, staticsName, "EnumInfo", 1, ",\r\n")
+							.Append("};\r\n");
+					}
 				}
 
 				if (Session.IncludeDebugOutput)
 				{
 					builder.Append("#if 0\r\n");
-					ReadOnlyMemory<string> sorted = HeaderFile.References.Declaration.GetSortedReferences(
-						(int objectIndex, bool registered) => GetExternalDecl(objectIndex, registered));
-					foreach (string declaration in sorted.Span)
+					IReadOnlyList<string> sorted = HeaderFile.References.Declaration.GetSortedReferences(GetExternalDecl);
+					foreach (string declaration in sorted)
 					{
 						builder.Append(declaration);
 					}
@@ -409,7 +513,7 @@ namespace EpicGames.UHT.Exporters.CodeGen
 		{
 			const string MetaDataParamsName = "Enum_MetaDataParams";
 			const string ObjectFlags = "RF_Public|RF_Transient|RF_MarkAsNative";
-			string singletonName = GetSingletonName(enumObj, true);
+			string singletonName = GetSingletonName(enumObj, UhtSingletonType.Registered);
 			string staticsName = singletonName + "_Statics";
 			string registrationName = $"Z_Registration_Info_UEnum_{enumObj.SourceName}";
 
@@ -443,18 +547,28 @@ namespace EpicGames.UHT.Exporters.CodeGen
 				}
 			}
 
-			builder.Append("static FEnumRegistrationInfo ").Append(registrationName).Append(";\r\n");
-			builder.Append("static UEnum* ").Append(enumObj.SourceName).Append("_StaticEnum()\r\n");
-			builder.Append("{\r\n");
-
-			builder.Append("\tif (!").Append(registrationName).Append(".OuterSingleton)\r\n");
-			builder.Append("\t{\r\n");
-			builder.Append("\t\t").Append(registrationName).Append(".OuterSingleton = GetStaticEnum(").Append(singletonName).Append(", (UObject*)")
-				.Append(GetSingletonName(enumObj.Package, true)).Append("(), TEXT(\"").Append(enumObj.SourceName).Append("\"));\r\n");
-			builder.Append("\t}\r\n");
-			builder.Append("\treturn ").Append(registrationName).Append(".OuterSingleton;\r\n");
-
-			builder.Append("}\r\n");
+			if (Session.IsUsingCompiledInObjectFormat(UhtCompiledInObjectFormat.ConstInit))
+			{
+				using UhtConditionalMacroBlock block = new(builder, "UE_WITH_CONSTINIT_UOBJECT", Session.IsUsingMultipleCompiledInObjectFormats);
+				builder.Append($"static UEnum* {enumObj.SourceName}_StaticEnum()\r\n");
+				builder.Append("{\r\n");
+				builder.Append("\treturn ").AppendConstInitSingletonRef(this, enumObj).Append(";\r\n");
+				builder.Append("}\r\n");
+			}
+			if (Session.IsUsingCompiledInObjectFormat(UhtCompiledInObjectFormat.Params))
+			{
+				using UhtConditionalMacroBlock block = new(builder, "!UE_WITH_CONSTINIT_UOBJECT", Session.IsUsingMultipleCompiledInObjectFormats);
+				builder.Append("static FEnumRegistrationInfo ").Append(registrationName).Append(";\r\n");
+				builder.Append("static UEnum* ").Append(enumObj.SourceName).Append("_StaticEnum()\r\n");
+				builder.Append("{\r\n");
+				builder.Append("\tif (!").Append(registrationName).Append(".OuterSingleton)\r\n");
+				builder.Append("\t{\r\n");
+				builder.Append("\t\t").Append(registrationName).Append(".OuterSingleton = GetStaticEnum(").Append(singletonName).Append(", (UObject*)")
+					.Append(GetSingletonName(enumObj.Package, UhtSingletonType.Registered)).Append("(), TEXT(\"").Append(enumObj.SourceName).Append("\"));\r\n");
+				builder.Append("\t}\r\n");
+				builder.Append("\treturn ").Append(registrationName).Append(".OuterSingleton;\r\n");
+				builder.Append("}\r\n");
+			}
 
 			if (Module.Module.AlwaysExportEnums)
 			{
@@ -477,81 +591,173 @@ namespace EpicGames.UHT.Exporters.CodeGen
 			int hashCodeBlockStart = builder.Length;
 
 			// Statics declaration
+			string? enumeratorNamesArray = null;
+			string? enumeratorValuesArray = null;
+			string enumFlags = enumObj.EnumFlags.HasAnyFlags(EEnumFlags.Flags) ? "EEnumFlags::Flags" : "EEnumFlags::None";
+			bool hasEnumValues = enumObj.EnumValues.Count != 0 || enumObj.GeneratedMaxName is not null;
 			{
 				builder.Append("struct ").Append(staticsName).Append("\r\n");
 				builder.Append("{\r\n");
-				builder.AppendMetaDataDecl(enumObj, null, null, MetaDataParamsName, 1);
+				builder.AppendMetaDataDecl(enumObj, CodeGenerator, MetaDataParamsName, 1);
 
 				// Enumerators
-				if (enumObj.EnumValues.Count > 0)
+				if (hasEnumValues)
 				{
-					builder.Append("\tstatic constexpr UECodeGen_Private::FEnumeratorParam Enumerators[] = {\r\n");
-					int enumIndex = 0;
-					foreach (UhtEnumValue value in enumObj.EnumValues)
+					enumeratorNamesArray = "EnumeratorNamesUTF8";
+					enumeratorValuesArray = "EnumeratorValues";
+
+					if (enumObj.GeneratedMaxValue is not null)
+					{
+						builder.Append("\tstatic inline constexpr int64 GeneratedMaxValue = ").Append(enumObj.GeneratedMaxValue).Append(";\r\n");
+					}
+					else if (enumObj.GeneratedMaxName is not null)
+					{
+						// Use helper function to generate max value if UHT couldn't parse all values
+						builder.Append("\tstatic inline constexpr int64 GeneratedMaxValue = UEnum::CalculateMaxEnumValue({\r\n")
+							.AppendEach(enumObj.EnumValues, (builder, value) => builder.Append($"\t\t(int64){enumObj.Namespace.FullSourceName}{value.Name},\r\n"))
+							.Append("\t\t(int64)0\r\n") // Emit 0 to handle empty enums and avoid skipping the last comma in the loop above 
+							.Append($"\t}}, {enumFlags});\r\n");
+					}
+
+					if (Session.IsUsingCompiledInObjectFormat(UhtCompiledInObjectFormat.ConstInit))
+					{
+						using UhtConditionalMacroBlock block = new(builder, "UE_WITH_CONSTINIT_UOBJECT", Session.IsUsingMultipleCompiledInObjectFormats);
+						builder.Append("\tstatic constexpr const UTF8CHAR* EnumeratorNamesUTF8[]= {\r\n");
+						int enumIndex = 0;
+						foreach (UhtEnumValue value in enumObj.EnumValues)
+						{
+							if (!enumObj.MetaData.TryGetValue("OverrideName", enumIndex, out string? keyName))
+							{
+								keyName = value.Name.ToString();
+							}
+							builder.Append("\t\tUTF8TEXT(").AppendUTF8LiteralString(keyName).Append("),\r\n");
+							++enumIndex;
+						}
+						if (enumObj.GeneratedMaxName is not null)
+						{
+							builder.Append("\t\tUTF8TEXT(").AppendUTF8LiteralString(enumObj.GeneratedMaxName).Append("),\r\n");
+						}
+						builder.Append("\t};\r\n");
+
+						builder.Append("\tstatic inline UE_CONSTINIT_UOBJECT_DECL int64 EnumeratorValues[]= {\r\n")
+							.AppendEach(enumObj.EnumValues, (builder, value) => builder.Append($"\t\t(int64){enumObj.Namespace.FullSourceName}{value.Name},\r\n"));
+						if (enumObj.GeneratedMaxName is not null)
+						{
+							builder.Append("\t\tGeneratedMaxValue,\r\n");
+						}
+						builder.Append("\t};\r\n");
+					}
+					if (Session.IsUsingCompiledInObjectFormat(UhtCompiledInObjectFormat.Params))
 					{
-						if (!enumObj.MetaData.TryGetValue("OverrideName", enumIndex, out string? keyName))
+						using UhtConditionalMacroBlock block = new(builder, "!UE_WITH_CONSTINIT_UOBJECT", Session.IsUsingMultipleCompiledInObjectFormats);
+						builder.Append("\tstatic constexpr UECodeGen_Private::FEnumeratorParam Enumerators[] = {\r\n");
+						int enumIndex = 0;
+						foreach (UhtEnumValue value in enumObj.EnumValues)
+						{
+							if (!enumObj.MetaData.TryGetValue("OverrideName", enumIndex, out string? keyName))
+							{
+								keyName = value.Name.ToString();
+							}
+							builder.Append("\t\t{ ").AppendUTF8LiteralString(keyName).Append(", (int64)").Append(enumObj.Namespace.FullSourceName).Append(value.Name).Append(" },\r\n");
+							++enumIndex;
+						}
+						if (enumObj.GeneratedMaxName is not null)
 						{
-							keyName = value.Name.ToString();
+							builder.Append("\t\t{ ").AppendUTF8LiteralString(enumObj.GeneratedMaxName).Append(", GeneratedMaxValue").Append(" },\r\n");
 						}
-						builder.Append("\t\t{ ").AppendUTF8LiteralString(keyName).Append(", (int64)").Append(enumObj.Namespace.FullSourceName).Append(value.Name).Append(" },\r\n");
-						++enumIndex;
+						builder.Append("\t};\r\n");
 					}
-					builder.Append("\t};\r\n");
 				}
-				builder.Append("\tstatic const UECodeGen_Private::").AppendEnumParamsType(enumObj).Append(" EnumParams;\r\n");
-				builder.Append("};\r\n");
+
+				if (Session.IsUsingCompiledInObjectFormat(UhtCompiledInObjectFormat.Params))
+				{
+					builder.Append("\tstatic const UECodeGen_Private::").AppendEnumParamsType(enumObj).Append(" EnumParams;\r\n");
+				}
+
+				builder.Append($"}}; // struct {staticsName} \r\n");
 			}
 
 			// Statics definition
 			{
-				builder.Append("const UECodeGen_Private::").AppendEnumParamsType(enumObj).Append(' ').Append(staticsName).Append("::EnumParams = {\r\n");
-				if (enumObj.IsVerseField)
-				{
-					builder.Append("\t{\r\n");
+				if (Session.IsUsingCompiledInObjectFormat(UhtCompiledInObjectFormat.ConstInit))
+				{ 
+					using UhtConditionalMacroBlock block = new(builder, "UE_WITH_CONSTINIT_UOBJECT", Session.IsUsingMultipleCompiledInObjectFormats);
+					builder.Append($"UE_CONSTINIT_UOBJECT_DECL TNoDestroy<U{enumObj.EngineClassName}> {GetSingletonName(enumObj, UhtSingletonType.ConstInit)}{{NoDestroyConstEval,\r\n");
+					builder.AppendConstInitObjectParams(enumObj, this, 1);
+					builder.Append("\tUE::CodeGen::ConstInit::FUFieldParams{},\r\n");
+					builder.Append("\tUE::CodeGen::ConstInit::FEnumParams{\r\n");
+					builder.Append("\t\t.CppTypeStaticUTF8 = UTF8TEXT(").AppendUTF8LiteralString(enumObj.CppType).Append("),\r\n");
+					builder.Append(enumeratorNamesArray is null ? null : $"\t\t.StaticNamesUTF8 = MakeArrayView({staticsName}::{enumeratorNamesArray}),\r\n"); 
+					builder.Append(enumeratorValuesArray is null ? null : $"\t\t.EnumValues = MakeArrayView({staticsName}::{enumeratorValuesArray}),\r\n");
+					builder.Append($"\t\t.CppForm = (uint8)UEnum::ECppForm::{enumObj.CppForm.ToString()},\r\n");
+					builder.Append($"\t\t.EnumFlags = {(enumObj.EnumFlags.HasAnyFlags(EEnumFlags.Flags) ? "EEnumFlags::Flags" : "EEnumFlags::None")},\r\n");
+					builder.Append($"\t\t.DisplayNameFn = {enumDisplayNameFn},\r\n");
+					builder.Append("\t\tIF_WITH_METADATA(.MetaData = ").AppendMetaDataArrayView(enumObj.MetaData, staticsName, MetaDataParamsName, 0, ",)\r\n");
+					builder.Append("\t},");
+					if (enumObj.IsVerseField)
+					{
+						builder.Append("\tUE::CodeGen::ConstInit::FVerseEnumParams{\r\n")
+							.Append("\t\t.QualifiedName = UTF8TEXT(\"").AppendVerseScopeAndName(enumObj, UhtVerseNameMode.Qualified).Append("\"),\r\n")
+							.Append("\t},\r\n");
+					}
+					builder.Append("};\r\n");
 				}
-				builder.Append("\t(UObject*(*)())").Append(GetSingletonName(enumObj.Package, true)).Append(",\r\n");
-				builder.Append('\t').Append(enumDisplayNameFn).Append(",\r\n");
-				builder.Append('\t').AppendUTF8LiteralString(enumObj.SourceName).Append(",\r\n");
-				builder.Append('\t').AppendUTF8LiteralString(enumObj.CppType).Append(",\r\n");
-				if (enumObj.EnumValues.Count > 0)
+				if (Session.IsUsingCompiledInObjectFormat(UhtCompiledInObjectFormat.Params))
 				{
-					builder.Append('\t').Append(staticsName).Append("::Enumerators,\r\n");
-					builder.Append('\t').Append(ObjectFlags).Append(",\r\n");
-					builder.Append("\tUE_ARRAY_COUNT(").Append(staticsName).Append("::Enumerators),\r\n");
+					using UhtConditionalMacroBlock block = new(builder, "!UE_WITH_CONSTINIT_UOBJECT", Session.IsUsingMultipleCompiledInObjectFormats);
+					builder.Append("const UECodeGen_Private::").AppendEnumParamsType(enumObj).Append(' ').Append(staticsName).Append("::EnumParams = {\r\n");
+					if (enumObj.IsVerseField)
+					{
+						builder.Append("\t{\r\n");
+					}
+					builder.Append("\t(UObject*(*)())").Append(GetSingletonName(enumObj.Package, UhtSingletonType.Registered)).Append(",\r\n");
+					builder.Append('\t').Append(enumDisplayNameFn).Append(",\r\n");
+					builder.Append('\t').AppendUTF8LiteralString(enumObj.SourceName).Append(",\r\n");
+					builder.Append('\t').AppendUTF8LiteralString(enumObj.CppType).Append(",\r\n");
+					if (hasEnumValues)
+					{
+						builder.Append('\t').Append(staticsName).Append("::Enumerators,\r\n");
+						builder.Append('\t').Append(ObjectFlags).Append(",\r\n");
+						builder.Append("\tUE_ARRAY_COUNT(").Append(staticsName).Append("::Enumerators),\r\n");
+					}
+					else
+					{
+						builder.Append('\t').Append("nullptr,\r\n");
+						builder.Append('\t').Append(ObjectFlags).Append(",\r\n");
+						builder.Append("\t0,\r\n");
+					}
+					builder.Append('\t').Append(enumObj.EnumFlags.HasAnyFlags(EEnumFlags.Flags) ? "EEnumFlags::Flags" : "EEnumFlags::None").Append(",\r\n");
+					builder.Append("\t(uint8)UEnum::ECppForm::").Append(enumObj.CppForm.ToString()).Append(",\r\n");
+					builder.Append('\t').AppendMetaDataParams(enumObj, staticsName, MetaDataParamsName).Append("\r\n");
+					if (enumObj.IsVerseField)
+					{
+						builder.Append("\t},\r\n");
+						builder.Append("\t\"").AppendVerseScopeAndName(enumObj, UhtVerseNameMode.Qualified).Append("\"\r\n");
+					}
+					builder.Append("};\r\n");
 				}
-				else
+			}
+
+			// Registration singleton
+			if (Session.IsUsingCompiledInObjectFormat(UhtCompiledInObjectFormat.Params))
+			{ 
+				using UhtConditionalMacroBlock block = new(builder, "!UE_WITH_CONSTINIT_UOBJECT", Session.IsUsingMultipleCompiledInObjectFormats);
+				builder.Append("UEnum* ").Append(singletonName).Append("()\r\n");
+				builder.Append("{\r\n");
+				builder.Append("\tif (!").Append(registrationName).Append(".InnerSingleton)\r\n");
+				builder.Append("\t{\r\n");
+				if (enumObj.IsVerseField)
 				{
-					builder.Append('\t').Append("nullptr,\r\n");
-					builder.Append('\t').Append(ObjectFlags).Append(",\r\n");
-					builder.Append("\t0,\r\n");
+					builder.Append("\t\tVerse::CodeGen::Private::ConstructUVerseEnum(").Append(registrationName).Append(".InnerSingleton, ").Append(staticsName).Append("::EnumParams);\r\n");
 				}
-				builder.Append('\t').Append(enumObj.EnumFlags.HasAnyFlags(EEnumFlags.Flags) ? "EEnumFlags::Flags" : "EEnumFlags::None").Append(",\r\n");
-				builder.Append("\t(uint8)UEnum::ECppForm::").Append(enumObj.CppForm.ToString()).Append(",\r\n");
-				builder.Append('\t').AppendMetaDataParams(enumObj, staticsName, MetaDataParamsName).Append("\r\n");
-				if (enumObj.IsVerseField)
+				else
 				{
-					builder.Append("\t},\r\n");
-					builder.Append("\t\"").AppendVerseScopeAndName(enumObj, UhtVerseNameMode.Qualified).Append("\"\r\n");
+					builder.Append("\t\tUECodeGen_Private::ConstructUEnum(").Append(registrationName).Append(".InnerSingleton, ").Append(staticsName).Append("::EnumParams);\r\n");
 				}
-				builder.Append("};\r\n");
-			}
-
-			// Registration singleton
-			builder.Append("UEnum* ").Append(singletonName).Append("()\r\n");
-			builder.Append("{\r\n");
-			builder.Append("\tif (!").Append(registrationName).Append(".InnerSingleton)\r\n");
-			builder.Append("\t{\r\n");
-			if (enumObj.IsVerseField)
-			{
-				builder.Append("\t\tVerse::CodeGen::Private::ConstructUVerseEnum(").Append(registrationName).Append(".InnerSingleton, ").Append(staticsName).Append("::EnumParams);\r\n");
-			}
-			else
-			{
-				builder.Append("\t\tUECodeGen_Private::ConstructUEnum(").Append(registrationName).Append(".InnerSingleton, ").Append(staticsName).Append("::EnumParams);\r\n");
+				builder.Append("\t}\r\n");
+				builder.Append("\treturn ").Append(registrationName).Append(".InnerSingleton;\r\n");
+				builder.Append("}\r\n");
 			}
-			builder.Append("\t}\r\n");
-			builder.Append("\treturn ").Append(registrationName).Append(".InnerSingleton;\r\n");
-			builder.Append("}\r\n");
 
 			{
 				using UhtRentedPoolBuffer<char> borrowBuffer = builder.RentPoolBuffer(hashCodeBlockStart, builder.Length - hashCodeBlockStart);
@@ -563,220 +769,321 @@ namespace EpicGames.UHT.Exporters.CodeGen
 		private StringBuilder AppendScriptStruct(StringBuilder builder, UhtScriptStruct scriptStruct)
 		{
 			const string MetaDataParamsName = "Struct_MetaDataParams";
-			string singletonName = GetSingletonName(scriptStruct, true);
-			string staticsName = singletonName + "_Statics";
+			string singletonName = GetSingletonName(scriptStruct, UhtSingletonType.Registered);
+			string constinitName = GetSingletonName(scriptStruct, UhtSingletonType.ConstInit);
+			string staticsName = GetSingletonName(scriptStruct, UhtSingletonType.Statics);
 			string registrationName = $"Z_Registration_Info_UScriptStruct_{scriptStruct.SourceName}";
 			List<UhtScriptStruct> noExportStructs = FindNoExportStructs(scriptStruct);
 
-			if (scriptStruct.ScriptStructFlags.HasAnyFlags(EStructFlags.Native))
+			// Everything from this point on will be part of the definition hash
+			int hashCodeBlockStart = builder.Length;
+
+			// Collect the properties
+			PropertyMemberContextImpl context = new(CodeGenerator, scriptStruct, staticsName);
+			UhtUsedDefineScopes<UhtProperty> properties = new(scriptStruct.Properties);
+
+			// Declare the statics structure
 			{
-				// Inject static assert to verify that we do not add vtable
-				if (scriptStruct.SuperScriptStruct != null)
+				builder.Append("struct ").Append(staticsName).Append("\r\n");
+				builder.Append("{\r\n");
+
+				foreach (UhtScriptStruct noExportStruct in noExportStructs)
 				{
-					builder.Append("static_assert(std::is_polymorphic<")
-						.Append(scriptStruct.FullyQualifiedSourceName)
-						.Append(">() == std::is_polymorphic<")
-						.Append(scriptStruct.SuperScriptStruct.FullyQualifiedSourceName)
-						.Append(">(), \"USTRUCT ")
-						.Append(scriptStruct.FullyQualifiedSourceName)
-						.Append(" cannot be polymorphic unless super ")
-						.Append(scriptStruct.SuperScriptStruct.FullyQualifiedSourceName)
-						.Append(" is polymorphic\");\r\n");
+					AppendMirrorsForNoexportStruct(builder, noExportStruct, 1);
+					builder.Append("\tstatic_assert(sizeof(").Append(noExportStruct.NamespaceExportName).Append(noExportStruct.SourceName).Append(") < MAX_uint16);\r\n");
+					builder.Append("\tstatic_assert(alignof(").Append(noExportStruct.NamespaceExportName).Append(noExportStruct.SourceName).Append(") < MAX_uint8);\r\n");
 				}
 
-				// Outer singleton
-				builder.Append("static FStructRegistrationInfo ").Append(registrationName).Append(";\r\n");
-				builder.Append($"class UScriptStruct* {scriptStruct.FullyQualifiedSourceName}::StaticStruct()\r\n");
-				builder.Append("{\r\n");
-				builder.Append("\tif (!").Append(registrationName).Append(".OuterSingleton)\r\n");
-				builder.Append("\t{\r\n");
-				builder.Append("\t\t")
-					.Append(registrationName)
-					.Append(".OuterSingleton = GetStaticStruct(")
-					.Append(singletonName)
-					.Append(", (UObject*)")
-					.Append(GetSingletonName(scriptStruct.Package, true))
-					.Append("(), TEXT(\"")
-					.Append(scriptStruct.EngineName)
-					.Append("\"));\r\n");
-
-				// if this struct has RigVM methods - we need to register the method to our central
-				// registry on construction of the static struct
-				if (scriptStruct.RigVMStructInfo != null)
+				// Declare a function to access the size/alignment for NoExport cases as well, consteval should prevent this function having code in the binary
+				builder.Append("\tstatic inline consteval int32 GetStructSize() { return sizeof(").Append(scriptStruct.NamespaceExportName).Append(scriptStruct.SourceName).Append("); }\r\n");
+				builder.Append("\tstatic inline consteval int16 GetStructAlignment() { return alignof(").Append(scriptStruct.NamespaceExportName).Append(scriptStruct.SourceName).Append("); }\r\n");
+
+				// Meta data
+				builder.AppendMetaDataDecl(scriptStruct, context, properties, MetaDataParamsName, 1);
+
+				// Properties
+				AppendPropertiesDecl(builder, context, properties, 1);
+
+				if (Session.IsUsingCompiledInObjectFormat(UhtCompiledInObjectFormat.ConstInit))
+				{ 
+					using UhtConditionalMacroBlock block = new(builder, "UE_WITH_CONSTINIT_UOBJECT", Session.IsUsingMultipleCompiledInObjectFormats);
+					builder.Append("\tstatic UE_CONSTINIT_UOBJECT_DECL inline const FStructBaseChain* StructBases[] = {\r\n");
+					if (scriptStruct.SuperScriptStruct is not null) 
+					{
+						List<UhtScriptStruct> supers = new();
+						for (UhtScriptStruct? super = scriptStruct.SuperScriptStruct; super is not null; super = super.SuperScriptStruct)
+						{
+							supers.Add(super);
+						}
+						foreach (UhtScriptStruct super in supers.AsEnumerable().Reverse())
+						{
+							builder.Append("\t\tUE::Private::AsStructBaseChain(").AppendConstInitSingletonRef(this, super).Append("),\r\n");
+						}
+					}
+					builder.Append("\t\tUE::Private::AsStructBaseChain(").AppendConstInitSingletonRef(this, scriptStruct).Append("),\r\n");
+					builder.Append("\t};\r\n");
+				}
+
+				if (scriptStruct.RigVMStructInfo != null && scriptStruct.RigVMStructInfo.Methods.Count > 0)
 				{
+					// Arrays of parameters per method - note that predicates and other functions read different lists
 					foreach (UhtRigVMMethodInfo methodInfo in scriptStruct.RigVMStructInfo.Methods)
 					{
 						if (methodInfo.IsPredicate)
 						{
-							builder.Append("\t\tTArray<FRigVMFunctionArgument> ").AppendArgumentsName(scriptStruct, methodInfo).Append(";\r\n");
+							if (methodInfo.Parameters.Count == 0)
+							{
+								continue;
+							}
+
+							builder.Append($"\tstatic inline constexpr FRigVMCompiledInFunctionArgument RigVMFunctionParameters_{methodInfo.Name}[] = {{\r\n");
 							foreach (UhtRigVMParameter parameter in methodInfo.Parameters)
 							{
-								builder
-									.Append("\t\t")
-									.AppendArgumentsName(scriptStruct, methodInfo)
-									.Append(".Emplace(TEXT(\"")
-									.Append(parameter.NameOriginal())
-									.Append("\"), TEXT(\"")
-									.Append(parameter.TypeOriginal())
-									.Append("\"));\r\n");
+								builder.Append($"\t\t{{ .Name = TEXT(\"{parameter.NameOriginal()}\"), ")
+									.Append($".Type = TEXT(\"{parameter.TypeOriginal()}\"), ")
+									.Append($".Direction = ERigVMFunctionArgumentDirection::Input }},\r\n");
 							}
-							builder
-								.Append("\t\t")
-								.AppendArgumentsName(scriptStruct, methodInfo)
-								.Append(".Emplace(TEXT(\"Return\"), TEXT(\"")
-								.Append(methodInfo.ReturnType)
-								.Append("\"), ERigVMFunctionArgumentDirection::Output);\r\n");
-							
-							builder
-								.Append("\t\tFRigVMRegistry::Get().RegisterPredicate(")
-								.Append(registrationName)
-								.Append(".OuterSingleton, ")
-								.Append("TEXT(\"")
-								.Append(methodInfo.Name)
-								.Append("\"), ")
-								.AppendArgumentsName(scriptStruct, methodInfo)
-								.Append(");\r\n");
+							builder.Append($"\t\t{{ .Name = TEXT(\"Return\"),")
+								.Append($".Type = TEXT(\"{methodInfo.ReturnType}\"),")
+								.Append($".Direction = ERigVMFunctionArgumentDirection::Output, }}\r\n");
+							builder.Append("\t};\r\n");
 						}
 						else
 						{
-							builder.Append("\t\tTArray<FRigVMFunctionArgument> ").AppendArgumentsName(scriptStruct, methodInfo).Append(";\r\n");
+							if (scriptStruct.RigVMStructInfo.Members.Count == 0)
+							{
+								continue;
+							}
+
+							builder.Append($"\tstatic inline constexpr FRigVMCompiledInFunctionArgument RigVMFunctionParameters_{methodInfo.Name}[] = {{\r\n");
 							foreach (UhtRigVMParameter parameter in scriptStruct.RigVMStructInfo.Members)
 							{
-								builder
-									.Append("\t\t")
-									.AppendArgumentsName(scriptStruct, methodInfo)
-									.Append(".Emplace(TEXT(\"")
-									.Append(parameter.NameOriginal())
-									.Append("\"), TEXT(\"")
-									.Append(parameter.TypeOriginal())
-									.Append("\"));\r\n");
+								builder.Append($"\t\t{{ .Name = TEXT(\"{parameter.NameOriginal()}\"), ")
+									.Append($".Type = TEXT(\"{parameter.TypeOriginal()}\"), ")
+									.Append($".Direction = ERigVMFunctionArgumentDirection::Input }},\r\n");
 							}
-							
-							builder
-								.Append("\t\tFRigVMRegistry::Get().Register(TEXT(\"")
-								.Append(scriptStruct.SourceName)
-								.Append("::")
-								.Append(methodInfo.Name)
-								.Append("\"), &")
-								.Append(scriptStruct.SourceName)
-								.Append("::RigVM")
-								.Append(methodInfo.Name)
-								.Append(", ")
-								.Append(registrationName)
-								.Append(".OuterSingleton, ")
-								.AppendArgumentsName(scriptStruct, methodInfo)
-								.Append(");\r\n");
+							builder.Append("\t};\r\n");
+						}
+					}
+
+					builder.Append("\tstatic inline constexpr FRigVMCompiledInFunction RigVMFunctions[] = {\r\n");
+					foreach (UhtRigVMMethodInfo methodInfo in scriptStruct.RigVMStructInfo.Methods)
+					{
+						if (methodInfo.IsPredicate)
+						{
+							builder.Append($"\t\t{{ .MethodName = TEXT(\"{methodInfo.Name}\"), .Function = nullptr, \r\n");
+						}
+						else
+						{
+							builder.Append($"\t\t{{ .MethodName = TEXT(\"{scriptStruct::SourceName}::{methodInfo.Name}\"), \r\n")
+								.Append($"\t\t\t.Function = &{scriptStruct.SourceName}::RigVM{methodInfo.Name}, \r\n");
 						}
+
+						if ((methodInfo.IsPredicate ? methodInfo.Parameters.Count : scriptStruct.RigVMStructInfo.Members.Count) > 0)
+						{
+							builder.Append($"\t\t\t.Parameters = MakeArrayView(RigVMFunctionParameters_{methodInfo.Name}),\r\n");
+						}
+						builder.Append("\t\t},\r\n");
 					}
+					builder.Append("\t};\r\n");
 				}
 
-				builder.Append("\t}\r\n");
-				builder.Append("\treturn ").Append(registrationName).Append(".OuterSingleton;\r\n");
-				builder.Append("}\r\n");
+				// New struct ops
+				if (scriptStruct.ScriptStructFlags.HasAnyFlags(EStructFlags.Native))
+				{
+					if (Session.IsUsingCompiledInObjectFormat(UhtCompiledInObjectFormat.ConstInit))
+					{ 
+						using UhtConditionalMacroBlock block = new(builder, "UE_WITH_CONSTINIT_UOBJECT", Session.IsUsingMultipleCompiledInObjectFormats);
+						builder.Append($"\tstatic inline UE_CONSTINIT_UOBJECT_DECL UScriptStruct::TCppStructOps<{scriptStruct.Namespace.FullSourceName}{scriptStruct.SourceName}> StructOps").Append("{};\r\n");
+					}
 
-				// Inject implementation needed to support auto bindings of fast arrays
-				if (ObjectInfos[scriptStruct.ObjectTypeIndex].FastArrayProperty != null)
+					if (Session.IsUsingCompiledInObjectFormat(UhtCompiledInObjectFormat.Params))
+					{
+						using UhtConditionalMacroBlock block = new(builder, "!UE_WITH_CONSTINIT_UOBJECT", Session.IsUsingMultipleCompiledInObjectFormats);
+						builder.Append("\tstatic void* NewStructOps()\r\n");
+						builder.Append("\t{\r\n");
+						builder.Append($"\t\treturn (UScriptStruct::ICppStructOps*)new UScriptStruct::TCppStructOps<{scriptStruct.Namespace.FullSourceName}{scriptStruct.SourceName}>();\r\n");
+						builder.Append("\t}\r\n");
+					}
+				}
+
+				if (Session.IsUsingCompiledInObjectFormat(UhtCompiledInObjectFormat.Params))
 				{
-					// The preprocessor conditional is written here instead of in FastArraySerializerImplementation.h
-					// since it may evaluate differently in different modules, triggering warnings in IncludeTool.
-					builder.Append("#if defined(UE_NET_HAS_IRIS_FASTARRAY_BINDING) && UE_NET_HAS_IRIS_FASTARRAY_BINDING\r\n");
-					builder.Append("UE_NET_IMPLEMENT_FASTARRAY(").Append(scriptStruct.SourceName).Append(");\r\n");
-					builder.Append("#else\r\n");
-					builder.Append("UE_NET_IMPLEMENT_FASTARRAY_STUB(").Append(scriptStruct.SourceName).Append(");\r\n");
-					builder.Append("#endif\r\n");
+					using UhtConditionalMacroBlock block = new(builder, "!UE_WITH_CONSTINIT_UOBJECT", Session.IsUsingMultipleCompiledInObjectFormats);
+					builder.Append("\tstatic const UECodeGen_Private::").AppendStructParamsType(scriptStruct).Append(" StructParams;\r\n");
 				}
+				builder.Append($"}}; // struct {staticsName}\r\n");
 			}
 
-			// Everything from this point on will be part of the definition hash
-			int hashCodeBlockStart = builder.Length;
-
-			// Collect the properties
-			PropertyMemberContextImpl context = new(CodeGenerator, scriptStruct, staticsName);
-			UhtUsedDefineScopes<UhtProperty> properties = new(scriptStruct.Properties);
-
-			// Declare the statics structure
+			if (scriptStruct.ScriptStructFlags.HasAnyFlags(EStructFlags.Native))
 			{
-				builder.Append("struct ").Append(staticsName).Append("\r\n");
-				builder.Append("{\r\n");
-
-				foreach (UhtScriptStruct noExportStruct in noExportStructs)
+				// Inject static assert to verify that we do not add vtable
+				if (scriptStruct.SuperScriptStruct != null)
 				{
-					AppendMirrorsForNoexportStruct(builder, noExportStruct, 1);
-					builder.Append("\tstatic_assert(sizeof(").Append(noExportStruct.NamespaceExportName).Append(noExportStruct.SourceName).Append(") < MAX_uint16);\r\n");
-					builder.Append("\tstatic_assert(alignof(").Append(noExportStruct.NamespaceExportName).Append(noExportStruct.SourceName).Append(") < MAX_uint8);\r\n");
+					builder.Append("static_assert(std::is_polymorphic<")
+						.Append(scriptStruct.FullyQualifiedSourceName)
+						.Append(">() == std::is_polymorphic<")
+						.Append(scriptStruct.SuperScriptStruct.FullyQualifiedSourceName)
+						.Append(">(), \"USTRUCT ")
+						.Append(scriptStruct.FullyQualifiedSourceName)
+						.Append(" cannot be polymorphic unless super ")
+						.Append(scriptStruct.SuperScriptStruct.FullyQualifiedSourceName)
+						.Append(" is polymorphic\");\r\n");
 				}
 
-				// Meta data
-				builder.AppendMetaDataDecl(scriptStruct, context, properties, MetaDataParamsName, 1);
-
-				// Properties
-				AppendPropertiesDecl(builder, context, properties, 1);
-
-				// New struct ops
-				if (scriptStruct.ScriptStructFlags.HasAnyFlags(EStructFlags.Native))
+				// Outer singleton
+				if (Session.IsUsingCompiledInObjectFormat(UhtCompiledInObjectFormat.ConstInit))
+				{ 
+					using UhtConditionalMacroBlock block = new(builder, "UE_WITH_CONSTINIT_UOBJECT", Session.IsUsingMultipleCompiledInObjectFormats);
+					builder.Append($"class UScriptStruct* {scriptStruct.FullyQualifiedSourceName}::StaticStruct()\r\n");
+					builder.Append("{\r\n");
+					builder.Append("\treturn &").Append(constinitName).Append(";\r\n");
+					builder.Append("\t}\r\n");
+				}
+				if (Session.IsUsingCompiledInObjectFormat(UhtCompiledInObjectFormat.Params))
 				{
-					builder.Append("\tstatic void* NewStructOps()\r\n");
+					using UhtConditionalMacroBlock block = new(builder, "!UE_WITH_CONSTINIT_UOBJECT", Session.IsUsingMultipleCompiledInObjectFormats);
+					builder.Append("static FStructRegistrationInfo ").Append(registrationName).Append(";\r\n");
+					builder.Append($"class UScriptStruct* {scriptStruct.FullyQualifiedSourceName}::StaticStruct()\r\n");
+					builder.Append("{\r\n");
+					builder.Append("\tif (!").Append(registrationName).Append(".OuterSingleton)\r\n");
 					builder.Append("\t{\r\n");
-					builder.Append("\t\treturn (UScriptStruct::ICppStructOps*)new UScriptStruct::TCppStructOps<").Append(scriptStruct.NamespaceExportName).Append(scriptStruct.SourceName).Append(">();\r\n");
+					builder.Append("\t\t")
+						.Append(registrationName)
+						.Append(".OuterSingleton = GetStaticStruct(")
+						.Append(singletonName)
+						.Append(", (UObject*)")
+						.Append(GetSingletonName(scriptStruct.Package, UhtSingletonType.Registered))
+						.Append("(), TEXT(\"")
+						.Append(scriptStruct.EngineName)
+						.Append("\"));\r\n");
+
+					// if this struct has RigVM methods - we need to register the method to our central
+					// registry on construction of the static struct
+					if (scriptStruct.RigVMStructInfo != null && scriptStruct.RigVMStructInfo.Methods.Count > 0)
+					{
+						builder.Append($"\t\tFRigVMRegistry::Get().RegisterCompiledInStruct({registrationName}.OuterSingleton, MakeArrayView({staticsName}::RigVMFunctions));\r\n");
+					}
+
+					builder.Append("\t}\r\n");
+					builder.Append("\treturn ").Append(registrationName).Append(".OuterSingleton;\r\n");
 					builder.Append("\t}\r\n");
 				}
-				builder.Append("\tstatic const UECodeGen_Private::").AppendStructParamsType(scriptStruct).Append(" StructParams;\r\n");
-				builder.Append("};\r\n");
+				
+
+				// Inject implementation needed to support auto bindings of fast arrays
+				if (ObjectInfos[scriptStruct.ObjectTypeIndex].FastArrayProperty != null)
+				{
+					// The preprocessor conditional is written here instead of in FastArraySerializerImplementation.h
+					// since it may evaluate differently in different modules, triggering warnings in IncludeTool.
+					builder.Append("#if defined(UE_NET_HAS_IRIS_FASTARRAY_BINDING) && UE_NET_HAS_IRIS_FASTARRAY_BINDING\r\n");
+					builder.Append("UE_NET_IMPLEMENT_FASTARRAY(").Append(scriptStruct.SourceName).Append(");\r\n");
+					builder.Append("#else\r\n");
+					builder.Append("UE_NET_IMPLEMENT_FASTARRAY_STUB(").Append(scriptStruct.SourceName).Append(");\r\n");
+					builder.Append("#endif\r\n");
+				}
 			}
 
 			// Populate the elements of the static structure
 			{
 				AppendPropertiesDefs(builder, context, properties, 0);
 
-				builder.Append("const UECodeGen_Private::").AppendStructParamsType(scriptStruct).Append(' ').Append(staticsName).Append("::StructParams = {\r\n");
-				if (scriptStruct.IsVerseField)
-				{
-					builder.Append("\t{\r\n");
-				}
-				builder.Append("\t(UObject* (*)())").Append(GetSingletonName(scriptStruct.Package, true)).Append(",\r\n");
-				builder.Append('\t').Append(GetSingletonName(scriptStruct.SuperScriptStruct, true)).Append(",\r\n");
-				if (scriptStruct.ScriptStructFlags.HasAnyFlags(EStructFlags.Native))
-				{
-					builder.Append("\t&NewStructOps,\r\n");
-				}
-				else
-				{
-					builder.Append("\tnullptr,\r\n");
+				// NOTE: This is temporary while the new VM doesn't match the value in the old VM
+				string? verseDefaultName = scriptStruct.IsVerseField ? scriptStruct.GetVerseScopeAndName(UhtVerseNameMode.Default) : null;
+				string? verseQualifiedName = scriptStruct.IsVerseField ? scriptStruct.GetVerseScopeAndName(UhtVerseNameMode.Qualified) : null;
+
+				if (Session.IsUsingCompiledInObjectFormat(UhtCompiledInObjectFormat.ConstInit))
+				{ 
+					using UhtConditionalMacroBlock block = new(builder, "UE_WITH_CONSTINIT_UOBJECT", Session.IsUsingMultipleCompiledInObjectFormats);
+					builder.Append($"UE_CONSTINIT_UOBJECT_DECL TNoDestroy<U{scriptStruct.EngineClassName}> {constinitName}{{NoDestroyConstEval,\r\n");
+					builder.AppendConstInitObjectParams(scriptStruct, this, 1);
+					builder.Append("\tUE::CodeGen::ConstInit::FUFieldParams{},\r\n");
+					builder.Append("\tUE::CodeGen::ConstInit::FStructParams{\r\n")
+						.Append($"\t\t.Super = ").AppendConstInitSingletonRef(this, scriptStruct.SuperScriptStruct).Append(",\r\n")
+						.Append($"\t\t.BaseChain = MakeArrayView({staticsName}::StructBases),\r\n")
+						.Append($"\t\t.ChildProperties =  {staticsName}::GetFirstProperty(),\r\n")
+						.Append($"\t\t.PropertiesSize = {staticsName}::GetStructSize(),\r\n")
+						.Append($"\t\t.MinAlignment = {staticsName}::GetStructAlignment(),\r\n")
+						.Append("\t\tIF_WITH_METADATA(.MetaData = ").AppendMetaDataArrayView(scriptStruct.MetaData, staticsName, MetaDataParamsName, 0, ",)\r\n")
+						.Append("\t},\r\n");
+					builder.Append("\tUE::CodeGen::ConstInit::FScriptStructParams{\r\n")
+						.Append($"\t\t.StructFlags = EStructFlags(0x{(uint)(scriptStruct.ScriptStructFlags & ~EStructFlags.ComputedFlags):X8}),\r\n")
+						.Append(scriptStruct.ScriptStructFlags.HasAnyFlags(EStructFlags.Native)
+							? $"\t\t.StructOps = &{staticsName}::StructOps,\r\n" 
+							: null)
+						.Append("\t},\r\n");
+					if (scriptStruct.IsVerseField)
+					{
+						uint guidA = Crc.Compute(new StringView(scriptStruct.EngineName));
+						uint guidB = Crc.Compute(new StringView(scriptStruct.Package.EngineName));
+						builder.Append("\tUE::CodeGen::ConstInit::FVerseStructParams{\r\n")
+							.Append($"\t\t.GuidA = {guidA},\r\n")
+							.Append($"\t\t.GuidB = {guidB},\r\n")
+							.Append("\t\t.QualifiedName = ");
+						if (verseDefaultName!.Equals(verseQualifiedName, StringComparison.Ordinal))
+						{
+							builder.Append($"UTF8TEXT(\"{verseQualifiedName}\")\r\n");
+						}
+						else
+						{
+							builder.Append($"IF_WITH_VERSE_VM(UTF8TEXT(\"{verseDefaultName}\"), UTF8TEXT(\"{verseQualifiedName}\"))\r\n");
+						}
+						builder.Append("},\r\n");
+					}
+					builder.Append("};\r\n");
 				}
-				builder.Append('\t').AppendUTF8LiteralString(scriptStruct.EngineName).Append(",\r\n");
-				builder.AppendArrayPtrLine(properties, UhtDefineScopeNames.Standard, staticsName, "PropPointers", 1, ",\r\n");
-				builder.AppendArrayCountLine(properties, UhtDefineScopeNames.Standard, staticsName, "PropPointers", 1, ",\r\n");
-				builder.Append("\tsizeof(").Append(scriptStruct.NamespaceExportName).Append(scriptStruct.SourceName).Append("),\r\n");
-				builder.Append("\talignof(").Append(scriptStruct.NamespaceExportName).Append(scriptStruct.SourceName).Append("),\r\n");
-				builder.Append("\tRF_Public|RF_Transient|RF_MarkAsNative,\r\n");
-				builder.Append($"\tEStructFlags(0x{(uint)(scriptStruct.ScriptStructFlags & ~EStructFlags.ComputedFlags):X8}),\r\n");
-				builder.Append('\t').AppendMetaDataParams(scriptStruct, staticsName, MetaDataParamsName).Append("\r\n");
-				if (scriptStruct.IsVerseField)
-				{
-					builder.Append("\t},\r\n");
-
-					// NOTE: This is temporary while the new VM doesn't match the value in the old VM
-					string defaultName = scriptStruct.GetVerseScopeAndName(UhtVerseNameMode.Default);
-					string qualifiedName = scriptStruct.GetVerseScopeAndName(UhtVerseNameMode.Qualified);
-					if (defaultName.Equals(qualifiedName, StringComparison.Ordinal))
+				if (Session.IsUsingCompiledInObjectFormat(UhtCompiledInObjectFormat.Params))
+				{ 
+					using UhtConditionalMacroBlock block = new(builder, "!UE_WITH_CONSTINIT_UOBJECT", Session.IsUsingMultipleCompiledInObjectFormats);
+					builder.Append("const UECodeGen_Private::").AppendStructParamsType(scriptStruct).Append(' ').Append(staticsName).Append("::StructParams = {\r\n");
+					if (scriptStruct.IsVerseField)
 					{
-						builder.Append($"\t\"{qualifiedName}\"\r\n");
+						builder.Append("\t{\r\n");
+					}
+					builder.Append("\t(UObject* (*)())").Append(GetSingletonName(scriptStruct.Package, UhtSingletonType.Registered)).Append(",\r\n");
+					builder.Append('\t').Append(GetSingletonName(scriptStruct.SuperScriptStruct, UhtSingletonType.Registered)).Append(",\r\n");
+					if (scriptStruct.ScriptStructFlags.HasAnyFlags(EStructFlags.Native))
+					{
+						builder.Append("\t&NewStructOps,\r\n");
 					}
 					else
 					{
-						builder.Append("\t#if WITH_VERSE_VM\r\n");
-						builder.Append($"\t\"{defaultName}\"\r\n");
-						builder.Append("\t#else\r\n");
-						builder.Append($"\t\"{qualifiedName}\"\r\n");
-						builder.Append("\t#endif\r\n");
+						builder.Append("\tnullptr,\r\n");
+					}
+					builder.Append('\t').AppendUTF8LiteralString(scriptStruct.EngineName).Append(",\r\n");
+					builder.AppendArrayPtrLine(properties, UhtDefineScopeNames.Standard, staticsName, "PropPointers", 1, ",\r\n");
+					builder.AppendArrayCountLine(properties, UhtDefineScopeNames.Standard, staticsName, "PropPointers", 1, ",\r\n");
+					builder.Append("\tsizeof(").Append(scriptStruct.NamespaceExportName).Append(scriptStruct.SourceName).Append("),\r\n");
+					builder.Append("\talignof(").Append(scriptStruct.NamespaceExportName).Append(scriptStruct.SourceName).Append("),\r\n");
+					builder.Append("\tRF_Public|RF_Transient|RF_MarkAsNative,\r\n");
+					builder.Append($"\tEStructFlags(0x{(uint)(scriptStruct.ScriptStructFlags & ~EStructFlags.ComputedFlags):X8}),\r\n");
+					builder.Append('\t').AppendMetaDataParams(scriptStruct, staticsName, MetaDataParamsName).Append("\r\n");
+					if (scriptStruct.IsVerseField)
+					{
+						builder.Append("\t},\r\n");
+
+						if (verseDefaultName!.Equals(verseQualifiedName, StringComparison.Ordinal))
+						{
+							builder.Append($"\t\"{verseQualifiedName}\"\r\n");
+						}
+						else
+						{
+							builder.Append("\t#if WITH_VERSE_VM\r\n");
+							builder.Append($"\t\"{verseDefaultName}\"\r\n");
+							builder.Append("\t#else\r\n");
+							builder.Append($"\t\"{verseQualifiedName}\"\r\n");
+							builder.Append("\t#endif\r\n");
+						}
 					}
+					builder.Append("};\r\n");
 				}
-				builder.Append("};\r\n");
 			}
 
 			// Generate the registration function
+			if (Session.IsUsingCompiledInObjectFormat(UhtCompiledInObjectFormat.Params))
 			{
-				builder.Append("UScriptStruct* ").Append(singletonName).Append("()\r\n");
+				using UhtConditionalMacroBlock block = new(builder, "!UE_WITH_CONSTINIT_UOBJECT", Session.IsUsingMultipleCompiledInObjectFormats);
+				builder.Append($"U{scriptStruct.EngineLinkClassName}* {singletonName}()\r\n");
 				builder.Append("{\r\n");
 				string innerSingletonName;
 				if (scriptStruct.ScriptStructFlags.HasAnyFlags(EStructFlags.Native))
@@ -799,7 +1106,7 @@ namespace EpicGames.UHT.Exporters.CodeGen
 					builder.Append("\t\tUECodeGen_Private::ConstructUScriptStruct(").Append(innerSingletonName).Append(", ").Append(staticsName).Append("::StructParams);\r\n");
 				}
 				builder.Append("\t}\r\n");
-				builder.Append("\treturn ").Append(innerSingletonName).Append(";\r\n");
+				builder.Append($"\treturn CastChecked<U{scriptStruct.EngineClassName}>({innerSingletonName});\r\n");
 				builder.Append("}\r\n");
 			}
 
@@ -932,10 +1239,12 @@ namespace EpicGames.UHT.Exporters.CodeGen
 		private StringBuilder AppendFunction(StringBuilder builder, UhtFunction function, bool isNoExport, bool isVerseNativeCallable)
 		{
 			const string MetaDataParamsName = "Function_MetaDataParams";
-			string singletonName = GetSingletonName(function, true);
-			string staticsName = singletonName + "_Statics";
+			string singletonName = GetSingletonName(function, UhtSingletonType.Registered);
+			string staticsName = GetSingletonName(function, UhtSingletonType.Statics); 
+			string constinitName = GetSingletonName(function, UhtSingletonType.ConstInit);
 			bool paramsInStatic = isNoExport || !IsCallbackFunction(function);
 			bool isNet = function.FunctionFlags.HasAnyFlags(EFunctionFlags.NetRequest | EFunctionFlags.NetResponse);
+			UhtClass? outerClass = function.Outer as UhtClass;
 
 			string strippedFunctionName = function.StrippedFunctionName;
 			string eventParameters = GetEventStructParametersName(function.Outer, strippedFunctionName);
@@ -965,8 +1274,26 @@ namespace EpicGames.UHT.Exporters.CodeGen
 				builder.AppendMetaDataDecl(function, context, properties, MetaDataParamsName, 1);
 
 				AppendPropertiesDecl(builder, context, properties, 1);
-
-				builder.Append("\tstatic const UECodeGen_Private::").AppendFunctionParamsType(function).Append(" FuncParams;\r\n");
+				
+				if (Session.IsUsingCompiledInObjectFormat(UhtCompiledInObjectFormat.ConstInit))
+				{ 
+					using UhtConditionalMacroBlock block = new(builder, "UE_WITH_CONSTINIT_UOBJECT", Session.IsUsingMultipleCompiledInObjectFormats);
+					if (function.FunctionFlags.HasExactFlags(EFunctionFlags.Native | EFunctionFlags.NetRequest, EFunctionFlags.Native)
+						&& function.FunctionExportFlags.HasAllFlags(UhtFunctionExportFlags.CustomThunk))
+					{
+						// Add an accessor for the probably-private static memeber function ptr in the static struct which can be declared as a friend
+						builder.Append($"\tUE_CONSTEVAL static FNativeFuncPtr GetCustomThunk_{function.SourceName}()\r\n")
+							.Append("\t{\r\n")
+							.Append($"\t\treturn &{(function.Outer as UhtClass)!.NativeFunctionCallName}::exec{function.SourceName};\r\n")
+							.Append("\t}\r\n");
+					}
+				}
+				if (Session.IsUsingCompiledInObjectFormat(UhtCompiledInObjectFormat.Params))
+				{ 
+					using UhtConditionalMacroBlock block = new(builder, "!UE_WITH_CONSTINIT_UOBJECT", Session.IsUsingMultipleCompiledInObjectFormats);
+					builder.Append("\tstatic const UECodeGen_Private::").AppendFunctionParamsType(function).Append(" FuncParams;\r\n");
+				}
+				
 
 				builder.Append("};\r\n");
 			}
@@ -974,86 +1301,131 @@ namespace EpicGames.UHT.Exporters.CodeGen
 			// Statics definition
 			{
 				AppendPropertiesDefs(builder, context, properties, 0);
-
+				
 				string functionEngineName = isVerseNativeCallable ? function.GetEncodedVerseScopeAndName(UhtVerseNameMode.Default) : function.EngineName;
-				builder
-					.Append("const UECodeGen_Private::").AppendFunctionParamsType(function).Append(' ')
-					.Append(staticsName).Append("::FuncParams = { { ")
-					.Append("(UObject*(*)())").Append(GetFunctionOuterFunc(function)).Append(", ")
-					.Append(GetSingletonName(function.SuperFunction, true)).Append(", ")
-					.AppendUTF8LiteralString(functionEngineName).Append(", ");
 
-				builder.AppendArrayPtrLine(properties, UhtDefineScopeNames.Standard, staticsName, "PropPointers", 0, ", ");
-				builder.AppendArrayCountLine(properties, UhtDefineScopeNames.Standard, staticsName, "PropPointers", 0, ", ");
-
-				string sizeOfStatic = "";
+				string? sizeOfStatic = null;
+				string? alignOfStatic = null;
 				if (function.Children.Count > 0)
 				{
 					UhtFunction tempFunction = function;
 					while (tempFunction.SuperFunction != null)
 					{
 						tempFunction = tempFunction.SuperFunction;
-					}
-
-					if (paramsInStatic)
-					{
-						sizeOfStatic = "sizeof(" + staticsName + "::" + GetEventStructParametersName(tempFunction.Outer, tempFunction.StrippedFunctionName) + ")";
-					}
-					else
-					{
-						sizeOfStatic = "sizeof(" + GetEventStructParametersName(tempFunction.Outer, tempFunction.StrippedFunctionName) + ")";
-					}
-
-					builder.Append(sizeOfStatic);
-				}
-				else
-				{
-					builder.Append('0');
-				}
-				builder.Append(", ");
+					}
 
-				builder
-					.Append("RF_Public|RF_Transient|RF_MarkAsNative, ")
-					.Append($"(EFunctionFlags)0x{(uint)function.FunctionFlags:X8}, ")
-					.Append(isNet ? function.RPCId : 0).Append(", ")
-					.Append(isNet ? function.RPCResponseId : 0).Append(", ")
-					.AppendMetaDataParams(function, staticsName, MetaDataParamsName)
-					.Append("}, ");
-
-				switch (function.FunctionType)
-				{
-					case UhtFunctionType.Function:
-						if (function.IsVerseField)
-						{
-							builder.Append("\r\n");
-							builder.Append("#if WITH_VERSE_VM\r\n");
-							builder.AppendUTF8LiteralString(function.GetVerseScopeAndName(UhtVerseNameMode.Default)).Append(",\r\n");
-							builder.Append("#else\r\n");
-							builder.AppendUTF8LiteralString(function.GetEncodedVerseScopeAndName(UhtVerseNameMode.Default)).Append(",\r\n");
-							builder.Append("#endif\r\n");
-						}
-						break;
+					string eventStructName = $"{(paramsInStatic ? staticsName : "")}{(paramsInStatic ? "::" : "")}{GetEventStructParametersName(tempFunction.Outer, tempFunction.StrippedFunctionName)}";
+					sizeOfStatic = $"sizeof({eventStructName})";
+					alignOfStatic = $"alignof({eventStructName})";
+				}
+
+				if (Session.IsUsingCompiledInObjectFormat(UhtCompiledInObjectFormat.ConstInit))
+				{ 
+					using UhtConditionalMacroBlock block = new(builder, "UE_WITH_CONSTINIT_UOBJECT", Session.IsUsingMultipleCompiledInObjectFormats);
+					builder.Append($"UE_CONSTINIT_UOBJECT_DECL TNoDestroy<U{function.EngineClassName}> {constinitName}{{NoDestroyConstEval,\r\n")
+						.AppendConstInitObjectParams(function, this, 1);
+					builder.Append("\tUE::CodeGen::ConstInit::FUFieldParams{ \r\n")
+						.AppendScopeLink(function.NextFunction, UhtDefineScopeNames.WithEditor, context, 2, ".Next = ", ",\r\n")
+						.Append("\t},\r\n");
+					builder.Append("\tUE::CodeGen::ConstInit::FStructParams{\r\n")
+						.Append($"\t\t.Super = ").AppendConstInitSingletonRef(this, function.SuperFunction).Append(",\r\n")
+						.Append($"\t\t.ChildProperties = {staticsName}::GetFirstProperty(),\r\n")
+						.Append(sizeOfStatic is null ? null : $"\t\t.PropertiesSize = {sizeOfStatic},\r\n")
+						.Append(alignOfStatic is null ? null : $"\t\t.MinAlignment = {alignOfStatic},\r\n")
+						.Append("\t\tIF_WITH_METADATA(.MetaData = ").AppendMetaDataArrayView(function.MetaData, staticsName, MetaDataParamsName, 0, ",)\r\n")
+						.Append("\t},\r\n");
+					builder.Append("\tUE::CodeGen::ConstInit::FFunctionParams{\r\n")
+						.Append((function.FunctionFlags & EFunctionFlags.Native, function.FunctionFlags & EFunctionFlags.NetRequest, function.FunctionExportFlags & UhtFunctionExportFlags.CustomThunk) switch
+						{
+							// Native non-service with custom thunk
+							(EFunctionFlags.Native, EFunctionFlags.None, UhtFunctionExportFlags.CustomThunk)
+								=> $"\t\t.NativeFunction = {staticsName}::GetCustomThunk_{function.SourceName}(),\r\n",
+							// Native without custom thunk
+							(EFunctionFlags.Native, EFunctionFlags.None, _)
+								=> $"\t\t.NativeFunction = &{outerClass!.Namespace.FullSourceName}{outerClass!.NativeFunctionCallName}::exec{function.SourceName},\r\n",
+							// Non-native function 
+							(EFunctionFlags.None, _, _) => $"\t\t.NativeFunction = &UObject::ProcessInternal,\r\n",
+							// No thunk for whatever reason - e.g. net service function
+							_ => $"",
+						})
+						.Append($"\t\t.FunctionFlags = (EFunctionFlags)0x{(uint)function.FunctionFlags:X8}, \r\n")
+						.Append(isNet ? $"\t\t.RPCId = {function.RPCId}, \r\n" : null)
+						.Append(isNet ? $"\t\t.RPCResponseId = {function.RPCResponseId}, \r\n" : null)
+						.Append("\t},\r\n");
+					if (function.FunctionType == UhtFunctionType.SparseDelegate)
+					{
+						builder.Append("\tUE::CodeGen::ConstInit::FSparseDelegateParams{\r\n")
+						.Append("\t\tUTF8TEXT(").AppendUTF8LiteralString(function.SparseOwningClassName).Append("), \r\n")
+						.Append("\t\tUTF8TEXT(").AppendUTF8LiteralString(function.SparseDelegateName).Append("), \r\n")
+						.Append("\t},\r\n");
+					}
+					if (function.IsVerseField)
+					{
+						builder.Append("\tUE::CodeGen::ConstInit::FVerseFunctionParams{\r\n")
+							.Append("\t\t.CompiledInAlternateNameUTF8 = IF_WITH_VERSE_VM(")
+							.Append("UTF8TEXT(").AppendUTF8LiteralString(function.GetVerseScopeAndName(UhtVerseNameMode.Default)).Append("), ")
+							.Append("UTF8TEXT(").AppendUTF8LiteralString(function.GetEncodedVerseScopeAndName(UhtVerseNameMode.Default)).Append("), ")
+							.Append("), \r\n")
+							.Append("\t},\r\n");
+					}
+					builder.Append(" };\r\n");
+				}
+
+				if (Session.IsUsingCompiledInObjectFormat(UhtCompiledInObjectFormat.Params))
+				{ 
+					using UhtConditionalMacroBlock block = new(builder, "!UE_WITH_CONSTINIT_UOBJECT", Session.IsUsingMultipleCompiledInObjectFormats);
+					builder.Append("const UECodeGen_Private::").AppendFunctionParamsType(function).Append(' ')
+						.Append(staticsName).Append("::FuncParams = { { ")
+						.Append("(UObject*(*)())").Append(GetFunctionOuterFunc(function)).Append(", ")
+						.Append(GetSingletonName(function.SuperFunction, UhtSingletonType.Registered)).Append(", ")
+						.AppendUTF8LiteralString(functionEngineName).Append(", ")
+						.AppendArrayPtrLine(properties, UhtDefineScopeNames.Standard, staticsName, "PropPointers", 1, ", \r\n")
+						.AppendArrayCountLine(properties, UhtDefineScopeNames.Standard, staticsName, "PropPointers", 1, ", \r\n")
+						.Append(sizeOfStatic ?? "0").Append(",\r\n")
+						.Append("RF_Public|RF_Transient|RF_MarkAsNative, ")
+						.Append($"(EFunctionFlags)0x{(uint)function.FunctionFlags:X8}, ")
+						.Append(isNet ? function.RPCId : 0).Append(", ")
+						.Append(isNet ? function.RPCResponseId : 0).Append(", ")
+						.AppendMetaDataParams(function, staticsName, MetaDataParamsName)
+						.Append("}, ");
+
+					switch (function.FunctionType)
+					{
+						case UhtFunctionType.Function:
+							if (function.IsVerseField)
+							{
+								builder.Append("\r\n");
+								builder.Append("#if WITH_VERSE_VM\r\n");
+								builder.AppendUTF8LiteralString(function.GetVerseScopeAndName(UhtVerseNameMode.Default)).Append(",\r\n");
+								builder.Append("#else\r\n");
+								builder.AppendUTF8LiteralString(function.GetEncodedVerseScopeAndName(UhtVerseNameMode.Default)).Append(",\r\n");
+								builder.Append("#endif\r\n");
+							}
+							break;
 
-					case UhtFunctionType.Delegate:
-						break;
+						case UhtFunctionType.Delegate:
+							break;
 
-					case UhtFunctionType.SparseDelegate:
-						builder
-							.AppendUTF8LiteralString(function.SparseOwningClassName).Append(", ")
-							.AppendUTF8LiteralString(function.SparseDelegateName).Append(", ");
-						break;
-				}
+						case UhtFunctionType.SparseDelegate:
+							builder
+								.AppendUTF8LiteralString(function.SparseOwningClassName).Append(", ")
+								.AppendUTF8LiteralString(function.SparseDelegateName).Append(", ");
+							break;
+					}
 
-				builder.Append(" };\r\n");
+					builder.Append(" };\r\n");
+				}
 
-				if (!String.IsNullOrEmpty(sizeOfStatic))
+				if (sizeOfStatic is not null)
 				{
 					builder.Append("static_assert(").Append(sizeOfStatic).Append(" < MAX_uint16);\r\n");
 				}
 			}
 
 			// Registration function
+			if (Session.IsUsingCompiledInObjectFormat(UhtCompiledInObjectFormat.Params))
 			{
+				using UhtConditionalMacroBlock block = new(builder, "!UE_WITH_CONSTINIT_UOBJECT", Session.IsUsingMultipleCompiledInObjectFormats);
 				builder.Append("UFunction* ").Append(singletonName).Append("()\r\n");
 				builder.Append("{\r\n");
 				builder.Append("\tstatic UFunction* ReturnFunction = nullptr;\r\n");
@@ -1080,11 +1452,11 @@ namespace EpicGames.UHT.Exporters.CodeGen
 			}
 			else if (function.Outer is UhtPackage package)
 			{
-				return GetSingletonName(package, true);
+				return GetSingletonName(package, UhtSingletonType.Registered);
 			}
 			else
 			{
-				return GetSingletonName((UhtObject)function.Outer, true);
+				return GetSingletonName((UhtObject)function.Outer, UhtSingletonType.Registered);
 			}
 		}
 
@@ -1119,36 +1491,74 @@ namespace EpicGames.UHT.Exporters.CodeGen
 			return builder;
 		}
 
-		private static StringBuilder AppendPropertiesDecl(StringBuilder builder, IUhtPropertyMemberContext context, UhtUsedDefineScopes<UhtProperty> properties, int tabs)
+		private StringBuilder AppendPropertiesDecl(StringBuilder builder, IUhtPropertyMemberContext context, UhtUsedDefineScopes<UhtProperty> properties, int tabs)
 		{
-			if (properties.IsEmpty)
+			using UhtCodeBlockComment block = new(builder, context.OuterStruct, "constinit property declarations");
+			if (Session.IsUsingCompiledInObjectFormat(UhtCompiledInObjectFormat.ConstInit))
 			{
-				return builder;
+				using UhtConditionalMacroBlock macroBlock = new(builder, "UE_WITH_CONSTINIT_UOBJECT", Session.IsUsingMultipleCompiledInObjectFormats);
+				if (properties.IsEmpty)
+				{
+					builder.AppendTabs(tabs).Append("static consteval FProperty* GetFirstProperty() { return nullptr; }\r\n");
+				}
+				else
+				{
+					builder.AppendInstances(properties, UhtDefineScopeNames.Standard,
+							(builder, property) => builder.AppendConstInitMemberDecl(property, context, property.EngineName, "", tabs))
+						// declare consteval function to provide first property for constructor uncondtionally - body will be empty if there are no properties
+						.AppendTabs(tabs).Append("static consteval FProperty* GetFirstProperty();\r\n")
+						// for properties which have more than one 'next' property depending on preprocessor definitions, define a function to return the correct property
+						.AppendInstances(properties, UhtDefineScopeNames.Standard, (builder, property) =>
+						{
+							property.AppendGetNextProperty(builder, context, tabs);
+						});
+				}
+			}
+			if (Session.IsUsingCompiledInObjectFormat(UhtCompiledInObjectFormat.Params))
+			{
+				using UhtConditionalMacroBlock macroBlock = new(builder, "!UE_WITH_CONSTINIT_UOBJECT", Session.IsUsingMultipleCompiledInObjectFormats);
+				builder.AppendInstances(properties, UhtDefineScopeNames.Standard,
+					builder => { },
+					(builder, property) => builder.AppendMemberDecl(property, context, property.EngineName, "", tabs),
+					builder => builder.AppendTabs(tabs).Append("static const UECodeGen_Private::FPropertyParamsBase* const PropPointers[];\r\n")
+					);
 			}
-			builder.AppendInstances(properties, UhtDefineScopeNames.Standard,
-				builder => { },
-				(builder, property) => builder.AppendMemberDecl(property, context, property.EngineName, "", tabs),
-				builder => builder.AppendTabs(tabs).Append("static const UECodeGen_Private::FPropertyParamsBase* const PropPointers[];\r\n"));
+		
 			return builder;
 		}
 
-		private static StringBuilder AppendPropertiesDefs(StringBuilder builder, IUhtPropertyMemberContext context, UhtUsedDefineScopes<UhtProperty> properties, int tabs)
+		private StringBuilder AppendPropertiesDefs(StringBuilder builder, IUhtPropertyMemberContext context, UhtUsedDefineScopes<UhtProperty> properties, int tabs)
 		{
 			if (properties.IsEmpty)
 			{
 				return builder;
 			}
-			builder.AppendInstances(properties, UhtDefineScopeNames.Standard,
-				(builder, property) => builder.AppendMemberDef(property, context, property.EngineName, "", null, tabs));
 
-			builder.AppendInstances(properties, UhtDefineScopeNames.Standard,
-				builder => builder.AppendTabs(tabs).Append("const UECodeGen_Private::FPropertyParamsBase* const ").Append(context.StaticsName).Append("::PropPointers[] = {\r\n"),
-				(builder, property) => builder.AppendMemberPtr(property, context, property.EngineName, "", tabs + 1),
-				builder =>
-				{
-					builder.AppendTabs(tabs).Append("};\r\n");
-					builder.AppendTabs(tabs).Append("static_assert(UE_ARRAY_COUNT(").Append(context.StaticsName).Append("::PropPointers) < 2048);\r\n");
-				});
+			using UhtCodeBlockComment block = new(builder, context.OuterStruct, "Property Definitions");
+			if (Session.IsUsingCompiledInObjectFormat(UhtCompiledInObjectFormat.ConstInit))
+			{
+				using UhtConditionalMacroBlock macroBlock = new(builder, "UE_WITH_CONSTINIT_UOBJECT", Session.IsUsingMultipleCompiledInObjectFormats);
+				builder.AppendInstances(properties, UhtDefineScopeNames.Standard,
+						(builder, property) => builder.AppendConstInitMemberDef(property, context, property.EngineName, "", null, null, tabs))
+					.AppendTabs(tabs).Append($"consteval FProperty* {context.StaticsName}::GetFirstProperty() {{\r\n")
+					.AppendScopeLink(context.OuterStruct.FirstProperty, UhtDefineScopeNames.Standard, context, tabs+1, "return ", ";\r\n")
+					.Append("}\r\n");
+			}
+			if (Session.IsUsingCompiledInObjectFormat(UhtCompiledInObjectFormat.Params))
+			{ 
+				using UhtConditionalMacroBlock macroBlock = new(builder, "!UE_WITH_CONSTINIT_UOBJECT", Session.IsUsingMultipleCompiledInObjectFormats);
+				builder.AppendInstances(properties, UhtDefineScopeNames.Standard,
+					(builder, property) => builder.AppendMemberDef(property, context, property.EngineName, "", null, tabs));
+
+				builder.AppendInstances(properties, UhtDefineScopeNames.Standard,
+					builder => builder.AppendTabs(tabs).Append("const UECodeGen_Private::FPropertyParamsBase* const ").Append(context.StaticsName).Append("::PropPointers[] = {\r\n"),
+					(builder, property) => builder.AppendMemberPtr(property, context, property.EngineName, "", tabs + 1),
+					builder =>
+					{
+						builder.AppendTabs(tabs).Append("};\r\n");
+						builder.AppendTabs(tabs).Append("static_assert(UE_ARRAY_COUNT(").Append(context.StaticsName).Append("::PropPointers) < 2048);\r\n");
+					});
+			}
 			return builder;
 		}
 
@@ -1192,7 +1602,7 @@ namespace EpicGames.UHT.Exporters.CodeGen
 				}
 				if (isRpc)
 				{
-					builder.Append("DEFINE_FUNCTION(").Append(classObj.Namespace.FullSourceName).AppendClassSourceNameOrInterfaceName(classObj).Append("::").Append(function.UnMarshalAndCallName).Append(")\r\n");
+					builder.Append($"DEFINE_FUNCTION({classObj.Namespace.FullSourceName}{classObj.NativeFunctionCallName}::{function.UnMarshalAndCallName})\r\n");
 					builder.Append("{\r\n");
 					AppendFunctionThunk(builder, classObj, function);
 					builder.Append("}\r\n");
@@ -1864,11 +2274,6 @@ namespace EpicGames.UHT.Exporters.CodeGen
 			// Add sparse accessors
 			AppendSparseAccessors(builder, classObj);
 
-			if (!classObj.ClassExportFlags.HasAnyFlags(UhtClassExportFlags.NoExport))
-			{
-				AppendNatives(builder, classObj);
-			}
-
 			AppendNativeGeneratedInitCode(builder, classObj, functions);
 
 			if (!classObj.ClassExportFlags.HasAnyFlags(UhtClassExportFlags.HasConstructor))
@@ -2125,11 +2530,14 @@ namespace EpicGames.UHT.Exporters.CodeGen
 			return builder;
 		}
 
+#pragma warning disable CA1505 //  'AppendNativeGeneratedInitCode' has a maintainability index of '5'. Rewrite or refactor the code to increase its maintainability index (MI) above '9'. 
+
 		private StringBuilder AppendNativeGeneratedInitCode(StringBuilder builder, UhtClass classObj, UhtUsedDefineScopes<UhtFunction> functions)
 		{
 			const string MetaDataParamsName = "Class_MetaDataParams";
-			string singletonName = GetSingletonName(classObj, true);
+			string singletonName = GetSingletonName(classObj, UhtSingletonType.Registered);
 			string staticsName = singletonName + "_Statics";
+			string constinitName = GetSingletonName(classObj, UhtSingletonType.ConstInit);
 			string registrationName = $"Z_Registration_Info_UClass_{classObj.SourceName}";
 			string[]? sparseDataTypes = classObj.MetaData.GetStringArray(UhtNames.SparseClassDataTypes);
 
@@ -2138,38 +2546,50 @@ namespace EpicGames.UHT.Exporters.CodeGen
 			bool hasInterfaces = classObj.Bases.Any(x => x is UhtClass baseClass && baseClass.ClassFlags.HasAnyFlags(EClassFlags.Interface)) || classObj.FlattenedVerseInterfaces.Count > 0;
 
 			// This block of code is the contents of IMPLEMENT_CLASS_NO_AUTO_REGISTRATION 
-			builder.Append("FClassRegistrationInfo ").Append(registrationName).Append(";\r\n");
-			builder.Append("UClass* ").Append(classObj.Namespace.FullSourceName).Append(classObj.SourceName).Append("::GetPrivateStaticClass()\r\n");
-			builder.Append("{\r\n");
-			builder.Append("\tusing TClass = ").Append(classObj.Namespace.FullSourceName).Append(classObj.SourceName).Append(";\r\n");
-			builder.Append("\tif (!").Append(registrationName).Append(".InnerSingleton)\r\n");
-			builder.Append("\t{\r\n");
-			builder.Append("\t\t").Append(classObj.IsVerseField ? "Verse::CodeGen::Private::ConstructUVerseClassNoInit" : "GetPrivateStaticClassBody").Append("(\r\n");
-			builder.Append("\t\t\tTClass::StaticPackage(),\r\n");
-			builder.Append("\t\t\tTEXT(\"").Append(classObj.EngineName).Append("\"),\r\n");
-			builder.Append("\t\t\t").Append(registrationName).Append(".InnerSingleton,\r\n");
-			builder.Append("\t\t\tStaticRegisterNatives").Append(classObj.SourceName).Append(",\r\n");
-			builder.Append("\t\t\tsizeof(TClass),\r\n");
-			builder.Append("\t\t\talignof(TClass),\r\n");
-			builder.Append("\t\t\tTClass::StaticClassFlags,\r\n");
-			builder.Append("\t\t\tTClass::StaticClassCastFlags(),\r\n");
-			builder.Append("\t\t\tTClass::StaticConfigName(),\r\n");
-			builder.Append("\t\t\t(UClass::ClassConstructorType)InternalConstructor<TClass>,\r\n");
-			builder.Append("\t\t\t(UClass::ClassVTableHelperCtorCallerType)InternalVTableHelperCtorCaller<TClass>,\r\n");
-			builder.Append("\t\t\tUOBJECT_CPPCLASS_STATICFUNCTIONS_FORCLASS(TClass),\r\n");
-			builder.Append("\t\t\t&TClass::Super::StaticClass,\r\n");
-			builder.Append("\t\t\t&TClass::WithinClass::StaticClass\r\n");
-			builder.Append("\t\t);\r\n");
-			builder.Append("\t}\r\n");
-			builder.Append("\treturn ").Append(registrationName).Append(".InnerSingleton;\r\n");
-			builder.Append("}\r\n");
+			if (Session.IsUsingCompiledInObjectFormat(UhtCompiledInObjectFormat.ConstInit))
+			{
+				using UhtConditionalMacroBlock block = new(builder, "UE_WITH_CONSTINIT_UOBJECT", Session.IsUsingMultipleCompiledInObjectFormats);
+				builder.Append($"UClass* {classObj.Namespace.FullSourceName}{classObj.SourceName}::GetPrivateStaticClass()\r\n");
+				builder.Append("{\r\n");
+				builder.Append($"\treturn &{GetSingletonName(classObj, UhtSingletonType.ConstInit)};\r\n");
+				builder.Append("}\r\n");
+			}
+			if (Session.IsUsingCompiledInObjectFormat(UhtCompiledInObjectFormat.Params))
+			{
+				using UhtConditionalMacroBlock block = new(builder, "!UE_WITH_CONSTINIT_UOBJECT", Session.IsUsingMultipleCompiledInObjectFormats);
+				builder.Append("FClassRegistrationInfo ").Append(registrationName).Append(";\r\n");
+				builder.Append("UClass* ").Append(classObj.Namespace.FullSourceName).Append(classObj.SourceName).Append("::GetPrivateStaticClass()\r\n");
+				builder.Append("{\r\n");
+				builder.Append("\tusing TClass = ").Append(classObj.Namespace.FullSourceName).Append(classObj.SourceName).Append(";\r\n");
+				builder.Append("\tif (!").Append(registrationName).Append(".InnerSingleton)\r\n");
+				builder.Append("\t{\r\n");
+				builder.Append("\t\t").Append(classObj.IsVerseField ? "Verse::CodeGen::Private::ConstructUVerseClassNoInit" : "GetPrivateStaticClassBody").Append("(\r\n");
+				builder.Append("\t\t\tTClass::StaticPackage(),\r\n");
+				builder.Append("\t\t\tTEXT(\"").Append(classObj.EngineName).Append("\"),\r\n");
+				builder.Append("\t\t\t").Append(registrationName).Append(".InnerSingleton,\r\n");
+				builder.Append("\t\t\tStaticRegisterNatives").Append(classObj.SourceName).Append(",\r\n");
+				builder.Append("\t\t\tsizeof(TClass),\r\n");
+				builder.Append("\t\t\talignof(TClass),\r\n");
+				builder.Append("\t\t\tTClass::StaticClassFlags,\r\n");
+				builder.Append("\t\t\tTClass::StaticClassCastFlags(),\r\n");
+				builder.Append("\t\t\tTClass::StaticConfigName(),\r\n");
+				builder.Append("\t\t\t(UClass::ClassConstructorType)InternalConstructor<TClass>,\r\n");
+				builder.Append("\t\t\t(UClass::ClassVTableHelperCtorCallerType)InternalVTableHelperCtorCaller<TClass>,\r\n");
+				builder.Append("\t\t\tUOBJECT_CPPCLASS_STATICFUNCTIONS_FORCLASS(TClass),\r\n");
+				builder.Append("\t\t\t&TClass::Super::StaticClass,\r\n");
+				builder.Append("\t\t\t&TClass::WithinClass::StaticClass\r\n");
+				builder.Append("\t\t);\r\n");
+				builder.Append("\t}\r\n");
+				builder.Append("\treturn ").Append(registrationName).Append(".InnerSingleton;\r\n");
+				builder.Append("}\r\n");
+			}
 
 			// Everything from this point on will be part of the definition hash
 			int hashCodeBlockStart = builder.Length;
 
 			// simple ::StaticClass wrapper to avoid header, link and DLL hell
 			{
-				builder.Append("UClass* ").Append(GetSingletonName(classObj, false)).Append("()\r\n");
+				builder.Append("UClass* ").Append(GetSingletonName(classObj, UhtSingletonType.Unregistered)).Append("()\r\n");
 				builder.Append("{\r\n");
 				builder.Append("\treturn ").Append(classObj.Namespace.FullSourceName).Append(classObj.SourceName).Append("::GetPrivateStaticClass();\r\n");
 				builder.Append("}\r\n");
@@ -2177,6 +2597,9 @@ namespace EpicGames.UHT.Exporters.CodeGen
 
 			// Collect the properties
 			UhtUsedDefineScopes<UhtProperty> properties = new(classObj.Properties);
+			UhtUsedDefineScopes<UhtFunction> nativeFunctions = new();
+			bool hasNativeFunctions = false;
+			bool hasVerseCallableFunctions = false;
 
 			// Declare the statics object
 			{
@@ -2187,44 +2610,167 @@ namespace EpicGames.UHT.Exporters.CodeGen
 
 				AppendPropertiesDecl(builder, context, properties, 1);
 
-				builder.Append("\tstatic UObject* (*const DependentSingletons[])();\r\n");
-
-				//builder.AppendIfInstances(functions, UhtDefineScopeNames.WithEditor, builder => builder.Append("\t\tstatic const FClassFunctionLinkInfo FuncInfo[];\r\n"));
-				// Functions
-				builder.AppendInstances(functions, UhtDefineScopeNames.WithEditor,
-					builder =>
+				if (!classObj.ClassExportFlags.HasAnyFlags(UhtClassExportFlags.NoExport))
+				{
+					// Filter to only native non-net functions for exporting
+					nativeFunctions.AddRange(functions.Instances.Where(x => x.FunctionFlags.HasExactFlags(EFunctionFlags.Native | EFunctionFlags.NetRequest, EFunctionFlags.Native)));
+					hasNativeFunctions = !nativeFunctions.IsEmpty;
+					hasVerseCallableFunctions = nativeFunctions.Instances.Any(x => x.IsVerseField && x.FunctionExportFlags.HasAnyFlags(UhtFunctionExportFlags.VerseCallable));
+					if (hasNativeFunctions)
 					{
-						builder.Append("\tstatic constexpr FClassFunctionLinkInfo ").Append("FuncInfo[] = {\r\n");
-					},
-					(builder, function) =>
-					{
-						builder
-							.Append("\t\t{ &")
-							.Append(GetSingletonName(function, true))
-							.Append(", ")
-							.AppendUTF8LiteralString(function.EngineName)
-							.Append(" },")
-							.AppendObjectHash(classObj, context, function)
-							.Append("\r\n");
-					},
-					builder =>
+						using UhtMacroBlockEmitter blockEmitter = new(builder, UhtDefineScopeNames.WithEditor, nativeFunctions.SoleScope);
+						builder.Append("\tstatic constexpr UE::CodeGen::FClassNativeFunction Funcs[] = {\r\n");
+
+						foreach (UhtFunction function in nativeFunctions.Instances)
+						{
+							blockEmitter.Set(function.DefineScope);
+							builder
+								.Append("\t\t{ .NameUTF8 = UTF8TEXT(")
+								.AppendUTF8LiteralString(function.EngineName)
+								.Append("), .Pointer = &")
+								.AppendClassSourceNameOrInterfaceName(classObj)
+								.Append($"::exec{function.EngineName} }},\r\n");
+						}
+
+						if (hasVerseCallableFunctions)
+						{
+							builder.Append("#if WITH_VERSE_BPVM\r\n");
+							foreach (UhtFunction function in nativeFunctions.Instances)
+							{
+								if (function.IsVerseField && function.FunctionExportFlags.HasAnyFlags(UhtFunctionExportFlags.VerseCallable))
+								{
+									blockEmitter.Set(function.DefineScope);
+									builder
+										.Append("\t\t{ UTF8TEXT(")
+										.AppendUTF8LiteralString(function.GetEncodedVerseScopeAndName(UhtVerseNameMode.Default))
+										.Append("), &")
+										.AppendClassSourceNameOrInterfaceName(classObj)
+										.Append($"::_exec_{function.SourceName}___ }},\r\n");
+								}
+							}
+							builder.Append("#endif\r\n");
+						}
+
+						// This will close the block if we have one that isn't editor only
+						blockEmitter.Set(nativeFunctions.SoleScope);
+
+						builder.Append("\t};\r\n");
+					}
+
+					if (hasVerseCallableFunctions)
 					{
+						builder.Append("#if WITH_VERSE_VM\r\n");
+						using UhtMacroBlockEmitter blockEmitter = new(builder, UhtDefineScopeNames.WithEditor, nativeFunctions.SoleScope);
+						builder.Append("\tstatic constexpr FVerseCallableThunk VerseFuncs[] = {\r\n");
+
+						foreach (UhtFunction function in nativeFunctions.Instances)
+						{
+							if (function.IsVerseField && function.FunctionExportFlags.HasAnyFlags(UhtFunctionExportFlags.VerseCallable))
+							{
+								blockEmitter.Set(function.DefineScope);
+								builder
+									.Append("\t\t{ ")
+									.AppendUTF8LiteralString(function.GetVerseScopeAndName(UhtVerseNameMode.Default))
+									.Append(", &")
+									.AppendClassSourceNameOrInterfaceName(classObj)
+									.Append($"::_exec_{function.SourceName}___ }},\r\n");
+							}
+						}
+
+						// This will close the block if we have one that isn't editor only
+						blockEmitter.Set(nativeFunctions.SoleScope);
+
 						builder.Append("\t};\r\n");
-						builder.Append("\tstatic_assert(UE_ARRAY_COUNT(").Append("FuncInfo) < 2048);\r\n");
-					});
+						builder.Append("#endif\r\n");
+					}
+				}
+
+				if (Session.IsUsingCompiledInObjectFormat(UhtCompiledInObjectFormat.Params))
+				{
+					using UhtConditionalMacroBlock block = new(builder, "!UE_WITH_CONSTINIT_UOBJECT", Session.IsUsingMultipleCompiledInObjectFormats);
+					builder.Append("\tstatic UObject* (*const DependentSingletons[])();\r\n");
+				}
+
+				// Functions
+				if (Session.IsUsingCompiledInObjectFormat(UhtCompiledInObjectFormat.ConstInit))
+				{
+					using UhtConditionalMacroBlock block = new(builder, "UE_WITH_CONSTINIT_UOBJECT", Session.IsUsingMultipleCompiledInObjectFormats);
+					builder.Append("\tstatic consteval UFunction* GetFirstFunction()\r\n");
+					builder.Append("\t{\r\n");
+					builder.AppendScopeLink(classObj.FirstFunction, UhtDefineScopeNames.WithEditor, context, 2, "return ", ";\r\n");
+					builder.Append("\t}\r\n");
+				}
+				if (Session.IsUsingCompiledInObjectFormat(UhtCompiledInObjectFormat.Params))
+				{
+					using UhtConditionalMacroBlock block = new(builder, "!UE_WITH_CONSTINIT_UOBJECT", Session.IsUsingMultipleCompiledInObjectFormats);
+					builder.AppendInstances(functions, UhtDefineScopeNames.WithEditor,
+						builder =>
+						{
+							builder.Append("\tstatic constexpr FClassFunctionLinkInfo ").Append("FuncInfo[] = {\r\n");
+						},
+						(builder, function) =>
+						{
+							builder
+								.Append("\t\t{ &")
+								.Append(GetSingletonName(function, UhtSingletonType.Registered))
+								.Append(", ")
+								.AppendUTF8LiteralString(function.EngineName)
+								.Append(" },")
+								.AppendObjectHash(classObj, context, function)
+								.Append("\r\n");
+						},
+						builder =>
+						{
+							builder.Append("\t};\r\n");
+							builder.Append("\tstatic_assert(UE_ARRAY_COUNT(").Append("FuncInfo) < 2048);\r\n");
+						});
+				}
 
 				if (hasInterfaces)
 				{
-					builder.Append("\tstatic const UECodeGen_Private::FImplementedInterfaceParams InterfaceParams[];\r\n");
+					if (Session.IsUsingCompiledInObjectFormat(UhtCompiledInObjectFormat.ConstInit))
+					{
+						using UhtConditionalMacroBlock block = new(builder, "UE_WITH_CONSTINIT_UOBJECT", Session.IsUsingMultipleCompiledInObjectFormats);
+						builder.Append("\tstatic const UE::CodeGen::ConstInit::FClassImplementedInterface Interfaces[];\r\n");
+					}
+					if (Session.IsUsingCompiledInObjectFormat(UhtCompiledInObjectFormat.Params))
+					{
+						using UhtConditionalMacroBlock block = new(builder, "UE_WITH_CONSTINIT_UOBJECT", Session.IsUsingMultipleCompiledInObjectFormats);
+						builder.Append("\tstatic const UECodeGen_Private::FImplementedInterfaceParams InterfaceParams[];\r\n");
+					}
 				}
 
 				builder.Append("\tstatic constexpr FCppClassTypeInfoStatic StaticCppClassTypeInfo = {\r\n");
-				builder.Append("\t\tTCppClassTypeTraits<").Append(classObj.Namespace.FullSourceName).AppendClassSourceNameOrInterfaceName(classObj).Append(">::IsAbstract,\r\n");
+				builder.Append($"\t\tTCppClassTypeTraits<{classObj.Namespace.FullSourceName}{classObj.NativeFunctionCallName}>::IsAbstract,\r\n");
 				builder.Append("\t};\r\n");
 
-				builder.Append("\tstatic const UECodeGen_Private::").AppendClassParamsType(classObj).Append(" ClassParams;\r\n");
+				if (Session.IsUsingCompiledInObjectFormat(UhtCompiledInObjectFormat.ConstInit))
+				{
+					using UhtConditionalMacroBlock block = new(builder, "UE_WITH_CONSTINIT_UOBJECT", Session.IsUsingMultipleCompiledInObjectFormats);
+					builder.Append("\tstatic UE_CONSTINIT_UOBJECT_DECL inline const FStructBaseChain* StructBases[] = {\r\n");
+					if (classObj.SuperClass is not null)
+					{
+						List<UhtClass> supers = new();
+						for (UhtClass? super = classObj.SuperClass; super is not null; super = super.SuperClass)
+						{
+							supers.Add(super);
+						}
+						foreach (UhtClass super in supers.AsEnumerable().Reverse())
+						{
+							builder.Append("\t\tUE::Private::AsStructBaseChain(").AppendConstInitSingletonRef(this, super).Append("),\r\n");
+						}
+					}
+					builder.Append("\t\tUE::Private::AsStructBaseChain(").AppendConstInitSingletonRef(this, classObj).Append("),\r\n");
+					builder.Append("\t};\r\n");
+				}
+
+				if (Session.IsUsingCompiledInObjectFormat(UhtCompiledInObjectFormat.Params))
+				{
+					using UhtConditionalMacroBlock block = new(builder, "UE_WITH_CONSTINIT_UOBJECT", Session.IsUsingMultipleCompiledInObjectFormats);
+					builder.Append("\tstatic const UECodeGen_Private::").AppendClassParamsType(classObj).Append(" ClassParams;\r\n");
+				}
 
-				builder.Append("};\r\n");
+				builder.Append($"}}; // struct {staticsName}\r\n");
 			}
 
 			// Define the statics object
@@ -2232,83 +2778,212 @@ namespace EpicGames.UHT.Exporters.CodeGen
 				AppendPropertiesDefs(builder, context, properties, 0);
 
 				// Dependent singletons
-				builder.Append("UObject* (*const ").Append(staticsName).Append("::DependentSingletons[])() = {\r\n");
-				if (classObj.SuperClass != null && classObj.SuperClass != classObj)
+				if (Session.IsUsingCompiledInObjectFormat(UhtCompiledInObjectFormat.Params))
 				{
-					builder.Append("\t(UObject* (*)())").Append(GetSingletonName(classObj.SuperClass, true)).Append(",\r\n");
+					using UhtConditionalMacroBlock block = new(builder, "!UE_WITH_CONSTINIT_UOBJECT", Session.IsUsingMultipleCompiledInObjectFormats);
+					builder.Append("UObject* (*const ").Append(staticsName).Append("::DependentSingletons[])() = {\r\n");
+					if (classObj.SuperClass != null && classObj.SuperClass != classObj)
+					{
+						builder.Append("\t(UObject* (*)())").Append(GetSingletonName(classObj.SuperClass, UhtSingletonType.Registered)).Append(",\r\n");
+					}
+					builder.Append("\t(UObject* (*)())").Append(GetSingletonName(classObj.Package, UhtSingletonType.Registered)).Append(",\r\n");
+					builder.Append("};\r\n");
+					builder.Append("static_assert(UE_ARRAY_COUNT(").Append(staticsName).Append("::DependentSingletons) < 16);\r\n");
 				}
-				builder.Append("\t(UObject* (*)())").Append(GetSingletonName(classObj.Package, true)).Append(",\r\n");
-				builder.Append("};\r\n");
-				builder.Append("static_assert(UE_ARRAY_COUNT(").Append(staticsName).Append("::DependentSingletons) < 16);\r\n");
 
 				// Implemented interfaces
 				if (hasInterfaces)
 				{
-					builder.Append("const UECodeGen_Private::FImplementedInterfaceParams ").Append(staticsName).Append("::InterfaceParams[] = {\r\n");
-
-					foreach (UhtStruct structObj in classObj.Bases)
+					if (Session.IsUsingCompiledInObjectFormat(UhtCompiledInObjectFormat.ConstInit))
 					{
-						if (structObj is UhtClass interfaceObj)
+						using UhtConditionalMacroBlock block = new(builder, "UE_WITH_CONSTINIT_UOBJECT", Session.IsUsingMultipleCompiledInObjectFormats);
+						builder.Append($"const UE::CodeGen::ConstInit::FClassImplementedInterface {staticsName}::Interfaces[] = {{\r\n");
+						Action<UhtClass> appendConstInitInterfaceParam = (interfaceObj) =>
+						{
+							if (interfaceObj.ClassFlags.HasAnyFlags(EClassFlags.Interface))
+							{
+								string isVerseDirectInterface = interfaceObj.IsVerseField && classObj.VerseInterfaces.Contains(interfaceObj) ? "true" : "false";
+								builder.Append("\t{ ")
+									.Append(".Class = ")
+									.AppendSingletonRef(context, interfaceObj.AlternateObject, UhtSingletonType.ConstInit).Append(", ")
+									.Append(interfaceObj.IsVerseField
+										? ".PointerOffset= 0, .bImplementedByK2 = true, "
+										: $".PointerOffset = (int32)VTABLE_OFFSET({classObj.FullyQualifiedSourceName}, {interfaceObj.FullyQualifiedSourceName}), .bImplementedByK2 = false, ")
+									.Append($".bVerseDirectInterface = {isVerseDirectInterface}, ")
+									.Append("},")
+									.AppendObjectHash(classObj, context, interfaceObj.AlternateObject)
+									.Append("\r\n");
+							}
+						};
+						foreach (UhtStruct structObj in classObj.Bases)
+						{
+							if (structObj is UhtClass interfaceObj)
+							{
+								appendConstInitInterfaceParam(interfaceObj);
+							}
+						}
+						foreach (UhtClass interfaceObj in classObj.FlattenedVerseInterfaces)
 						{
-							AppendImplementedInterfaceParam(builder, context, classObj, interfaceObj);
+							appendConstInitInterfaceParam(interfaceObj);
 						}
+						builder.Append("};\r\n");
 					}
-
-					foreach (UhtClass interfaceObj in classObj.FlattenedVerseInterfaces)
+					if (Session.IsUsingCompiledInObjectFormat(UhtCompiledInObjectFormat.Params))
 					{
-						AppendImplementedInterfaceParam(builder, context, classObj, interfaceObj);
+						using UhtConditionalMacroBlock block = new(builder, "!UE_WITH_CONSTINIT_UOBJECT", Session.IsUsingMultipleCompiledInObjectFormats);
+						builder.Append($"const UECodeGen_Private::FImplementedInterfaceParams {staticsName}::InterfaceParams[] = {{\r\n");
+						foreach (UhtStruct structObj in classObj.Bases)
+						{
+							if (structObj is UhtClass interfaceObj)
+							{
+								AppendImplementedInterfaceParam(builder, context, classObj, interfaceObj, UhtSingletonType.Unregistered);
+							}
+						}
+						foreach (UhtClass interfaceObj in classObj.FlattenedVerseInterfaces)
+						{
+							AppendImplementedInterfaceParam(builder, context, classObj, interfaceObj, UhtSingletonType.Unregistered);
+						}
+						builder.Append("};\r\n");
 					}
-
-					builder.Append("};\r\n");
 				}
 
-				// Class parameters
 				EClassFlags classFlags = classObj.ClassFlags & EClassFlags.SaveInCompiledInClasses;
-				builder.Append("const UECodeGen_Private::").AppendClassParamsType(classObj).Append(' ').Append(staticsName).Append("::ClassParams = {\r\n");
-				if (classObj.IsVerseField)
+				// Propagate class flags without allowing removal for UClass::ClassFlags
+				for (UhtClass? superClass = classObj.SuperClass; superClass is not null; superClass = superClass.SuperClass)
 				{
-					builder.Append("\t{\r\n");
+					classFlags |= superClass.ClassFlags & EClassFlags.SaveInCompiledInClasses & EClassFlags.Inherit;
 				}
-				builder.Append("\t&").Append(classObj.Namespace.FullSourceName).Append(classObj.SourceName).Append("::StaticClass,\r\n");
-				if (classObj.Config.Length > 0)
+				if (Session.IsUsingCompiledInObjectFormat(UhtCompiledInObjectFormat.ConstInit))
 				{
-					builder.Append('\t').AppendUTF8LiteralString(classObj.Config).Append(",\r\n");
+					// constinit class definition
+					using UhtConditionalMacroBlock block = new(builder, "UE_WITH_CONSTINIT_UOBJECT", Session.IsUsingMultipleCompiledInObjectFormats);
+					builder.Append($"UE_CONSTINIT_UOBJECT_DECL TNoDestroy<U{classObj.EngineClassName}> {constinitName}{{NoDestroyConstEval,\r\n")
+						.AppendConstInitObjectParams(classObj, this, 1, "|RF_Standalone")
+						.Append("\tUE::CodeGen::ConstInit::FUFieldParams{},\r\n")
+						.Append("\tUE::CodeGen::ConstInit::FStructParams{\r\n")
+						.Append("\t\t.Super = ").AppendConstInitSingletonRef(this, classObj.Super).Append(",\r\n")
+						.Append($"\t\t.FirstChild = {staticsName}::GetFirstFunction(),\r\n")
+						.Append($"\t\t.BaseChain = MakeArrayView({staticsName}::StructBases),\r\n")
+						.Append($"\t\t.ChildProperties = {staticsName}::GetFirstProperty(),\r\n")
+						.Append($"\t\t.PropertiesSize = sizeof({classObj.Namespace.FullSourceName}{classObj.SourceName}),\r\n")
+						.Append($"\t\t.MinAlignment = alignof({classObj.Namespace.FullSourceName}{classObj.SourceName}),\r\n")
+						.Append("\t\tIF_WITH_METADATA(.MetaData = ").AppendMetaDataArrayView(classObj.MetaData, staticsName, MetaDataParamsName, 0, ",)\r\n")
+						.Append("},\r\n")
+						.Append("\tUE::CodeGen::ConstInit::FClassParams{\r\n")
+						.Append($"\t\t.ConfigName = {classObj.Namespace.FullSourceName}{classObj.SourceName}::StaticConfigName(),\r\n")
+						.Append($"\t\t.ClassConstructor = InternalConstructor<{classObj.Namespace.FullSourceName}{classObj.SourceName}>,\r\n")
+						.Append($"\t\t.ClassVTableHelperCtorCaller = InternalVTableHelperCtorCaller<{classObj.Namespace.FullSourceName}{classObj.SourceName}>,\r\n")
+						.Append($"\t\t.CppClassStaticFunctions = FUObjectCppClassStaticFunctions(({classObj.Namespace.FullSourceName}{classObj.SourceName}*)nullptr),\r\n")
+						.Append($"\t\t.ClassFlags = (EClassFlags)0x{(uint)classFlags:X8}u,\r\n")
+						.Append($"\t\t.ClassCastFlags = {classObj.Namespace.FullSourceName}{classObj.SourceName}::StaticAllClassCastFlags(),\r\n") // Propagate cast flags at compile time
+						.Append($"\t\t.CppClassTypeInfo = &{staticsName}::StaticCppClassTypeInfo,\r\n");
+					if (hasNativeFunctions)
+					{
+						builder.Append("\t\t.NativeFunctions = ").AppendArrayView(nativeFunctions, UhtDefineScopeNames.WithEditor, staticsName, "Funcs", 2, ",\r\n");
+					}
+					if (classObj.ClassWithin is not null && classObj.ClassWithin != Session.UObject)
+					{
+						builder.Append($"\t\t.ClassWithin = &{GetSingletonName(classObj.ClassWithin, UhtSingletonType.ConstInit)},\r\n");
+					}
+					if (hasInterfaces)
+					{
+						builder.Append($"\t\t.Interfaces = MakeConstArrayView({staticsName}::Interfaces),\r\n");
+					}
+					if (classObj.SparseClassDataStruct is not null)
+					{
+						builder.Append($"\t\t.SparseClassDataStruct = &{GetSingletonName(classObj.SparseClassDataStruct, UhtSingletonType.ConstInit)}\r\n");
+					}
+					builder.Append("\t},\r\n");
+					if (classObj.IsVerseField)
+					{
+						builder.Append("\tUE::CodeGen::ConstInit::FVerseClassParams{\r\n")
+							.Append("\t\t.PackageRelativeVersePath = ").AppendVerseScopeAndName(classObj, UhtVerseNameMode.PackageRelative).Append(",\r\n")
+							.Append("\t\t.MangledPackageVersePath = ").Append(VerseNameMangling.MangleCasedName(classObj.Module.Module.VersePath).Result).Append(",\r\n");
+						if (hasVerseCallableFunctions)
+						{
+							builder.Append("\t\tIF_WITH_VERSE_VM(.VerseCallableThunks = ").AppendArrayView(staticsName, "VerseFuncs", 0, ",)\r\n");
+						}
+						builder.Append("},\r\n");
+					}
+					builder.Append("};\r\n");
 				}
-				else
+
+				if (Session.IsUsingCompiledInObjectFormat(UhtCompiledInObjectFormat.Params))
 				{
-					builder.Append("\tnullptr,\r\n");
-				}
-				builder.Append("\t&StaticCppClassTypeInfo,\r\n");
-				builder.Append("\tDependentSingletons,\r\n");
-				builder.AppendArrayPtrLine(functions, UhtDefineScopeNames.WithEditor, null, "FuncInfo", 1, ",\r\n");
-				builder.AppendArrayPtrLine(properties, UhtDefineScopeNames.Standard, staticsName, "PropPointers", 1, ",\r\n");
-				builder.Append('\t').Append(hasInterfaces ? "InterfaceParams" : "nullptr").Append(",\r\n");
-				builder.Append("\tUE_ARRAY_COUNT(DependentSingletons),\r\n");
-				builder.AppendArrayCountLine(functions, UhtDefineScopeNames.WithEditor, null, "FuncInfo", 1, ",\r\n");
-				builder.AppendArrayCountLine(properties, UhtDefineScopeNames.Standard, staticsName, "PropPointers", 1, ",\r\n");
-				builder.Append('\t').Append(hasInterfaces ? "UE_ARRAY_COUNT(InterfaceParams)" : "0").Append(",\r\n");
-				builder.Append($"\t0x{(uint)classFlags:X8}u,\r\n");
-				builder.Append('\t').AppendMetaDataParams(classObj, staticsName, MetaDataParamsName).Append("\r\n");
-				if (classObj.IsVerseField)
-				{
-					EVerseClassFlags verseClassFlags = classObj.ClassType == UhtClassType.VModule ? EVerseClassFlags.Module : EVerseClassFlags.None;
-					builder.Append("\t},\r\n");
-					if (classObj.ClassType == UhtClassType.Class || classObj.ClassType == UhtClassType.VModule)
+					// Class parameters
+					using UhtConditionalMacroBlock block = new(builder, "!UE_WITH_CONSTINIT_UOBJECT", Session.IsUsingMultipleCompiledInObjectFormats);
+					builder.Append("const UECodeGen_Private::").AppendClassParamsType(classObj).Append($" {staticsName}::ClassParams = {{\r\n");
+					if (classObj.IsVerseField)
+					{
+						builder.Append("\t{\r\n");
+					}
+					builder.Append("\t&").Append(classObj.Namespace.FullSourceName).Append(classObj.SourceName).Append("::StaticClass,\r\n");
+					if (classObj.Config.Length > 0)
 					{
-						builder.Append("\t\"").AppendVerseScopeAndName(classObj, UhtVerseNameMode.PackageRelative).Append("\",\r\n");
+						builder.Append('\t').AppendUTF8LiteralString(classObj.Config).Append(",\r\n");
 					}
 					else
 					{
-						builder.Append("\t\"\",\r\n");
+						builder.Append("\tnullptr,\r\n");
+					}
+					builder.Append("\t&StaticCppClassTypeInfo,\r\n");
+					builder.Append("\tDependentSingletons,\r\n");
+					builder.AppendArrayPtrLine(functions, UhtDefineScopeNames.WithEditor, null, "FuncInfo", 1, ",\r\n");
+					builder.AppendArrayPtrLine(properties, UhtDefineScopeNames.Standard, staticsName, "PropPointers", 1, ",\r\n");
+					builder.Append('\t').Append(hasInterfaces ? "InterfaceParams" : "nullptr").Append(",\r\n");
+					builder.Append("\tUE_ARRAY_COUNT(DependentSingletons),\r\n");
+					builder.AppendArrayCountLine(functions, UhtDefineScopeNames.WithEditor, null, "FuncInfo", 1, ",\r\n");
+					builder.AppendArrayCountLine(properties, UhtDefineScopeNames.Standard, staticsName, "PropPointers", 1, ",\r\n");
+					builder.Append('\t').Append(hasInterfaces ? "UE_ARRAY_COUNT(InterfaceParams)" : "0").Append(",\r\n");
+					builder.Append($"\t0x{(uint)classFlags:X8}u,\r\n");
+					builder.Append('\t').AppendMetaDataParams(classObj, staticsName, MetaDataParamsName).Append("\r\n");
+					if (classObj.IsVerseField)
+					{
+						EVerseClassFlags verseClassFlags = classObj.ClassType == UhtClassType.VModule ? EVerseClassFlags.Module : EVerseClassFlags.None;
+						builder.Append("\t},\r\n");
+						if (classObj.ClassType == UhtClassType.Class || classObj.ClassType == UhtClassType.VModule)
+						{
+							builder.Append("\t\"").AppendVerseScopeAndName(classObj, UhtVerseNameMode.PackageRelative).Append("\",\r\n");
+						}
+						else
+						{
+							builder.Append("\t\"\",\r\n");
+						}
+						builder.Append("\t\"").Append(VerseNameMangling.MangleCasedName(classObj.Module.Module.VersePath).Result).Append("\",\r\n");
+						builder.Append($"\t0x{(uint)verseClassFlags:X8}u,\r\n");
+					}
+					builder.Append("};\r\n");
+				}
+			}
+
+			// Native function registration
+			if (Session.IsUsingCompiledInObjectFormat(UhtCompiledInObjectFormat.Params))
+			{
+				using UhtConditionalMacroBlock block = new(builder, "!UE_WITH_CONSTINIT_UOBJECT", Session.IsUsingMultipleCompiledInObjectFormats);
+				builder.Append($"void {classObj.Namespace.FullSourceName}{classObj.SourceName}::StaticRegisterNatives{classObj.SourceName}()\r\n");
+				builder.Append("{\r\n");
+
+				if (hasNativeFunctions)
+				{
+					builder.Append($"\tUClass* Class = {classObj.SourceName}::StaticClass();\r\n");
+					builder.Append($"\tFNativeFunctionRegistrar::RegisterFunctions(Class, ").AppendArrayView(nativeFunctions, UhtDefineScopeNames.WithEditor, staticsName, "Funcs", 0, ");\r\n");
+
+					if (hasVerseCallableFunctions)
+					{
+						builder.Append("#if WITH_VERSE_VM\r\n");
+						builder.Append($"\tVerse::CodeGen::Private::RegisterVerseCallableThunks(Class, {staticsName}::VerseFuncs, UE_ARRAY_COUNT({staticsName}::VerseFuncs));\r\n");
+						builder.Append("#endif\r\n");
 					}
-					builder.Append("\t\"").Append(VerseNameMangling.MangleCasedName(classObj.Module.Module.VersePath).Result).Append("\",\r\n");
-					builder.Append($"\t0x{(uint)verseClassFlags:X8}u,\r\n");
 				}
-				builder.Append("};\r\n");
+
+				builder.Append("}\r\n"); // Close StaticRegisterNatives function
+
 			}
 
 			// Class registration
+			if (Session.IsUsingCompiledInObjectFormat(UhtCompiledInObjectFormat.Params))
 			{
+				using UhtConditionalMacroBlock block = new(builder, "!UE_WITH_CONSTINIT_UOBJECT", Session.IsUsingMultipleCompiledInObjectFormats);
 				builder.Append("UClass* ").Append(singletonName).Append("()\r\n");
 				builder.Append("{\r\n");
 				builder.Append("\tif (!").Append(registrationName).Append(".OuterSingleton)\r\n");
@@ -2423,111 +3098,15 @@ namespace EpicGames.UHT.Exporters.CodeGen
 			}
 			return builder;
 		}
+#pragma warning restore CA1505 //  'AppendNativeGeneratedInitCode' has a maintainability index of '5'. Rewrite or refactor the code to increase its maintainability index (MI) above '9'. 
 
-		private static StringBuilder AppendNatives(StringBuilder builder, UhtClass classObj)
-		{
-			builder.Append("void ").Append(classObj.Namespace.FullSourceName).Append(classObj.SourceName).Append("::StaticRegisterNatives").Append(classObj.SourceName).Append("()\r\n");
-			builder.Append("{\r\n");
-
-			UhtUsedDefineScopes<UhtFunction> functions = new(classObj.Functions.Where(
-				x => x.FunctionFlags.HasExactFlags(EFunctionFlags.Native | EFunctionFlags.NetRequest, EFunctionFlags.Native)));
-			functions.Instances.Sort((x, y) => StringComparerUE.OrdinalIgnoreCase.Compare(x.EngineName, y.EngineName));
-
-			if (!functions.IsEmpty)
-			{
-				bool hasVerseCallableFunctions = false;
-				{
-					using UhtMacroBlockEmitter blockEmitter = new(builder, UhtDefineScopeNames.WithEditor, functions.SoleScope);
-					builder.Append("\tUClass* Class = ").Append(classObj.SourceName).Append("::StaticClass();\r\n");
-					builder.Append("\tstatic const FNameNativePtrPair Funcs[] = {\r\n");
-
-					foreach (UhtFunction function in functions.Instances)
-					{
-						blockEmitter.Set(function.DefineScope);
-						builder
-							.Append("\t\t{ ")
-							.AppendUTF8LiteralString(function.EngineName)
-							.Append(", &")
-							.AppendClassSourceNameOrInterfaceName(classObj)
-							.Append("::exec")
-							.Append(function.EngineName)
-							.Append(" },\r\n");
-						if (function.IsVerseField && function.FunctionExportFlags.HasAnyFlags(UhtFunctionExportFlags.VerseCallable))
-						{
-							hasVerseCallableFunctions = true;
-						}
-					}
-
-					if (hasVerseCallableFunctions)
-					{
-						builder.Append("#if WITH_VERSE_BPVM\r\n");
-						foreach (UhtFunction function in functions.Instances)
-						{
-							if (function.IsVerseField && function.FunctionExportFlags.HasAnyFlags(UhtFunctionExportFlags.VerseCallable))
-							{
-								blockEmitter.Set(function.DefineScope);
-								builder
-									.Append("\t\t{ ")
-									.AppendUTF8LiteralString(function.GetEncodedVerseScopeAndName(UhtVerseNameMode.Default))
-									.Append(", &")
-									.AppendClassSourceNameOrInterfaceName(classObj)
-									.Append("::_exec_")
-									.Append(function.SourceName)
-									.Append("___ },\r\n");
-							}
-						}
-						builder.Append("#endif\r\n");
-					}
-
-					// This will close the block if we have one that isn't editor only
-					blockEmitter.Set(functions.SoleScope);
-
-					builder.Append("\t};\r\n");
-					builder.Append("\tFNativeFunctionRegistrar::RegisterFunctions(Class, Funcs, UE_ARRAY_COUNT(Funcs));\r\n");
-				}
-
-				if (hasVerseCallableFunctions)
-				{
-					builder.Append("#if WITH_VERSE_VM\r\n");
-					using UhtMacroBlockEmitter blockEmitter = new(builder, UhtDefineScopeNames.WithEditor, functions.SoleScope);
-					builder.Append("\tstatic const FVerseCallableThunk VerseFuncs[] = {\r\n");
-
-					foreach (UhtFunction function in functions.Instances)
-					{
-						if (function.IsVerseField && function.FunctionExportFlags.HasAnyFlags(UhtFunctionExportFlags.VerseCallable))
-						{
-							blockEmitter.Set(function.DefineScope);
-							builder
-								.Append("\t\t{ ")
-								.AppendUTF8LiteralString(function.GetVerseScopeAndName(UhtVerseNameMode.Default))
-								.Append(", &")
-								.AppendClassSourceNameOrInterfaceName(classObj)
-								.Append("::_exec_")
-								.Append(function.SourceName)
-								.Append("___ },\r\n");
-						}
-					}
-
-					// This will close the block if we have one that isn't editor only
-					blockEmitter.Set(functions.SoleScope);
-
-					builder.Append("\t};\r\n");
-					builder.Append("\tVerse::CodeGen::Private::RegisterVerseCallableThunks(Class, VerseFuncs, UE_ARRAY_COUNT(VerseFuncs));\r\n");
-					builder.Append("#endif\r\n");
-				}
-			}
-
-			builder.Append("}\r\n");
-			return builder;
-		}
-
-		private StringBuilder AppendImplementedInterfaceParam(StringBuilder builder, PropertyMemberContextImpl context, UhtClass classObj, UhtClass interfaceObj)
+		private static StringBuilder AppendImplementedInterfaceParam(StringBuilder builder, PropertyMemberContextImpl context, UhtClass classObj, UhtClass interfaceObj, UhtSingletonType singletonType)
 		{
 			if (interfaceObj.ClassFlags.HasAnyFlags(EClassFlags.Interface))
 			{
 				builder
 					.Append("\t{ ")
-					.Append(GetSingletonName(interfaceObj.AlternateObject, false));
+					.AppendSingletonRef(context, interfaceObj.AlternateObject, singletonType);
 				if (interfaceObj.IsVerseField)
 				{
 					builder
@@ -2703,15 +3282,19 @@ namespace EpicGames.UHT.Exporters.CodeGen
 			public string NamePrefix => "NewProp_";
 			public string MetaDataSuffix => "_MetaData";
 
-			public string GetSingletonName(UhtObject? obj, bool registered)
+			public string GetSingletonName(UhtObject? obj, UhtSingletonType type)
 			{
-				return _codeGenerator.GetSingletonName(obj, registered);
+				return _codeGenerator.GetSingletonName(obj, type);
 			}
 
 			public uint GetTypeHash(UhtObject obj)
 			{
 				return _codeGenerator.ObjectInfos[obj.ObjectTypeIndex].Hash;
 			}
+
+			public bool IsUsingCompiledInObjectFormat(UhtCompiledInObjectFormat format) => _codeGenerator.Session.IsUsingCompiledInObjectFormat(format);
+
+			public bool IsUsingMultipleCompiledInObjectFormats => _codeGenerator.Session.IsUsingMultipleCompiledInObjectFormats;
 		}
 	}
 
@@ -2720,6 +3303,31 @@ namespace EpicGames.UHT.Exporters.CodeGen
 	/// </summary>
 	public static class UhtHeaderCodeGeneratorCppFileStringBuilderExtensions
 	{
+		/// <summary>
+		/// Appends a structure initialization expression for a constant-initialized UObject, followed
+		/// by a comma and newline for use in a larger aggregate initialization expression
+		/// </summary>
+		/// <param name="builder"></param>
+		/// <param name="uhtObject"></param>
+		/// <param name="names"></param>
+		/// <param name="tabs"></param>
+		/// <param name="additionalObjectFlags"></param>
+		/// <returns></returns>
+		public static StringBuilder AppendConstInitObjectParams(this StringBuilder builder, UhtObject uhtObject, IUhtSingletonNameProvider names, int tabs, string? additionalObjectFlags = null)
+		{
+			const string ObjectFlags = "RF_Public|RF_Transient|RF_MarkAsNative";
+			builder.AppendTabs(tabs).Append("UE::CodeGen::ConstInit::FObjectParams{\r\n");
+			builder.AppendTabs(tabs + 1).Append($".Flags = {ObjectFlags}{additionalObjectFlags},\r\n");
+			// Can't take a pointer to Class to initialize itself because of TNoDestroy, null class signals to take the address of Class in the constructor 
+			if (uhtObject != uhtObject.EngineClass)
+			{
+				builder.AppendTabs(tabs + 1).Append($".Class =").AppendConstInitSingletonRef(names, uhtObject.EngineClass).Append(",\r\n");
+			}
+			builder.AppendTabs(tabs + 1).Append($".NameUTF8 = UTF8TEXT(").AppendUTF8LiteralString(uhtObject.EngineName).Append("),\r\n");
+			builder.AppendTabs(tabs + 1).Append($".Outer = ").AppendConstInitSingletonRef(names, uhtObject.Outer as UhtObject).Append(",\r\n");
+			builder.AppendTabs(tabs).Append("},\r\n");
+			return builder;
+		}
 
 		/// <summary>
 		/// Append the parameter names for a function
