Integrating UHT changes for constinit uobjects

diff --git a/Engine/Source/Programs/Shared/EpicGames.UHT/Types/UhtClass.cs b/Engine/Source/Programs/Shared/EpicGames.UHT/Types/UhtClass.cs
--- a/Engine/Source/Programs/Shared/EpicGames.UHT/Types/UhtClass.cs
+++ b/Engine/Source/Programs/Shared/EpicGames.UHT/Types/UhtClass.cs
@@ -7,6 +7,7 @@ using System.Linq;
 using System.Text;
 using System.Text.Json.Serialization;
 using EpicGames.Core;
+using EpicGames.UHT.Exporters.CodeGen;
 using EpicGames.UHT.Parsers;
 using EpicGames.UHT.Tables;
 using EpicGames.UHT.Tokenizer;
@@ -320,7 +321,9 @@ namespace EpicGames.UHT.Types
 		public UhtClass ClassWithin { get; set; }
 
 		/// <summary>
-		/// Class flags
+		/// Class flags used to populate UMyClass::StaticClassFlags. 
+		/// Note that these may be different to those in the instance field UClass::ClassFlags which
+		/// does not respect removing class flags in a derived class.
 		/// </summary>
 		[JsonConverter(typeof(JsonStringEnumConverter))]
 		public EClassFlags ClassFlags { get; set; } = EClassFlags.None;
@@ -448,6 +451,11 @@ namespace EpicGames.UHT.Types
 		[JsonIgnore]
 		public List<string> ClassGroupNames { get; } = new List<string>();
 
+		/// <summary>
+		/// After functions are fully resolved, points to the first non-delegate function in this class
+		/// </summary>
+		public UhtDefineScopeLink<UhtFunction>? FirstFunction { get; set; }
+
 		/// <summary>
 		/// Add the given class flags
 		/// </summary>
@@ -485,13 +493,31 @@ namespace EpicGames.UHT.Types
 					case UhtClassType.NativeInterface:
 						return UhtEngineType.NativeInterface;
 					default:
-						throw new UhtIceException("Unknown class type");
+						throw new UhtIceException($"UhtEngineType: Unknown class type {ClassType}");
 				}
 			}
 		}
 
 		/// <inheritdoc/>
-		public override string EngineClassName => "Class";
+		public override string EngineClassName => IsVerseField ? "VerseClass" : "Class";
+
+		/// <inheritdoc/>
+		public override string EngineLinkClassName => "Class";
+
+		/// <inheritdoc/>
+		public override UhtClass? EngineClass => IsVerseField ? Session.UVerseClass : Session.UClass;
+
+		private string? _nativeFunctionCallName;
+
+		/// <summary>
+		/// The name used to scope native function calls into this type e.g. UClassName or IInterfaceName
+		/// </summary>
+		public string NativeFunctionCallName => _nativeFunctionCallName ??= ClassType switch
+		{
+			UhtClassType.Class or UhtClassType.NativeInterface or UhtClassType.VModule => SourceName,
+			UhtClassType.Interface => $"I{SourceName[1..]}",
+			var unknown => throw new UhtIceException($"NativeFunctionCallName: Unknown class type {unknown}"),
+		};
 
 		///<inheritdoc/>
 		[JsonIgnore]
@@ -514,7 +540,7 @@ namespace EpicGames.UHT.Types
 					case UhtClassType.VModule:
 						return Session.GetSpecifierValidatorTable(UhtTableNames.VModule);
 					default:
-						throw new UhtIceException("Unknown class type");
+						throw new UhtIceException($"SpecifierValidatorTable: Unknown class type {ClassType}");
 				}
 			}
 		}
@@ -1158,6 +1184,7 @@ namespace EpicGames.UHT.Types
 			switch (phase)
 			{
 				case UhtResolvePhase.Final:
+					LinkFunctions();
 					if (ScanForInstancedReferenced(false))
 					{
 						ClassFlags |= EClassFlags.HasInstancedReference;
@@ -1166,6 +1193,34 @@ namespace EpicGames.UHT.Types
 			}
 		}
 
+		/// <summary>
+		/// Create a linked list of functions across preprocessor definition changes
+		/// </summary>
+		protected void LinkFunctions()
+		{
+			List<UhtFunction> functions = new();
+			foreach (UhtType type in Children)
+			{
+				// Delegates are not linked into class's field list
+				if (type is UhtFunction function && !function.FunctionFlags.HasAnyFlags(EFunctionFlags.Delegate))
+				{
+					functions.Add(function);
+				}
+			}
+			functions.Sort((x, y) => StringComparerUE.OrdinalIgnoreCase.Compare(x.EngineName, y.EngineName));
+
+			using UhtDefineScopeListBuilder<UhtFunction> functionList = new(DefineScope);
+			foreach (UhtFunction function in functions)
+			{
+				if (function.NextFunction is not null) 
+				{ 
+					throw new Exception("Unexpectedly initialized next link on property");
+				}
+				function.NextFunction = functionList.Add(function);
+			}
+			FirstFunction = functionList.Head;
+		}
+
 		/// <inheritdoc/>
 		protected override bool ScanForInstancedReferencedInternal(bool deepScan)
 		{
@@ -2141,17 +2196,19 @@ namespace EpicGames.UHT.Types
 				return;
 			}
 
+			base.CollectReferences(collector);
+
 			// Add myself as declarations and cross module references
-			collector.AddDeclaration(this, false);
-			collector.AddDeclaration(this, true);
-			collector.AddCrossModuleReference(this, false);
-			collector.AddCrossModuleReference(this, true);
+			collector.AddDeclaration(this, UhtSingletonType.Unregistered);
+			collector.AddDeclaration(this, UhtSingletonType.Registered);
+			collector.AddCrossModuleReference(this, UhtSingletonType.Unregistered);
+			collector.AddCrossModuleReference(this, UhtSingletonType.Registered);
 
 			// Add the super class
 			if (SuperClass != null)
 			{
-				collector.AddDeclaration(SuperClass, true);
-				collector.AddCrossModuleReference(SuperClass, true);
+				collector.AddDeclaration(SuperClass, UhtSingletonType.Registered);
+				collector.AddCrossModuleReference(SuperClass, UhtSingletonType.Registered);
 			}
 
 			// Add any other base classes
@@ -2159,19 +2216,28 @@ namespace EpicGames.UHT.Types
 			{
 				if (baseObj is UhtClass baseClass)
 				{
-					collector.AddCrossModuleReference(baseClass, false);
+					collector.AddCrossModuleReference(baseClass, UhtSingletonType.Unregistered);
 				}
 			}
 
 			// Add any verse interfaces
 			foreach (UhtStruct interfaceObj in FlattenedVerseInterfaces)
 			{
-				collector.AddCrossModuleReference(interfaceObj, false);
+				collector.AddCrossModuleReference(interfaceObj, UhtSingletonType.Unregistered);
 			}
 
 			// Add the package
-			collector.AddDeclaration(Package, true);
-			collector.AddCrossModuleReference(Package, true);
+			collector.AddDeclaration(Package, UhtSingletonType.Registered);
+			collector.AddCrossModuleReference(Package, UhtSingletonType.Registered);
+
+			// Pointers that are needed for constinit only
+			collector.AddCrossModuleReference(Session.UClass, UhtSingletonType.ConstInit);
+			collector.AddCrossModuleReference(ClassWithin, UhtSingletonType.ConstInit);
+			// Add all parents for constinit initializion of bases array 
+			for (UhtClass? super = SuperClass; super is not null; super = super.SuperClass)
+			{
+				collector.AddCrossModuleReference(super, UhtSingletonType.ConstInit);
+			}
 
 			// Collect children references
 			foreach (UhtType child in Children)
diff --git a/Engine/Source/Programs/Shared/EpicGames.UHT/Types/UhtEnum.cs b/Engine/Source/Programs/Shared/EpicGames.UHT/Types/UhtEnum.cs
--- a/Engine/Source/Programs/Shared/EpicGames.UHT/Types/UhtEnum.cs
+++ b/Engine/Source/Programs/Shared/EpicGames.UHT/Types/UhtEnum.cs
@@ -5,6 +5,7 @@ using System.Collections.Generic;
 using System.Text;
 using System.Text.Json.Serialization;
 using EpicGames.Core;
+using EpicGames.UHT.Exporters.CodeGen;
 using EpicGames.UHT.Tables;
 using EpicGames.UHT.Utils;
 
@@ -210,12 +211,29 @@ namespace EpicGames.UHT.Types
 		/// </summary>
 		public List<UhtEnumValue> EnumValues { get; }
 
+		/// <summary>
+		/// If header didn't declare a max value, this is a generated name to be appended to the list of names in reflection data.
+		/// This is a full name as returned by GetFullEnumName i.e. it may be namespaced.
+		/// </summary>
+		public string? GeneratedMaxName { get; private set; }
+
+		/// <summary>
+		/// If header didn't declare a max value, and UHT was able to parse all values, this the value for GeneratedMaxName
+		/// </summary>
+		public long? GeneratedMaxValue { get; private set; }
+
 		/// <inheritdoc/>
 		[JsonIgnore]
 		public override UhtEngineType EngineType => UhtEngineType.Enum;
 
 		/// <inheritdoc/>
-		public override string EngineClassName => "Enum";
+		public override string EngineClassName => IsVerseField ? "VerseEnum" : "Enum";
+
+		/// <inheritdoc/>
+		public override string EngineLinkClassName => "Enum";
+
+		/// <inheritdoc/>
+		public override UhtClass EngineClass => IsVerseField ? Session.UVerseEnum : Session.UEnum;
 
 		///<inheritdoc/>
 		[JsonIgnore]
@@ -249,6 +267,52 @@ namespace EpicGames.UHT.Types
 			}
 			return builder;
 		}
+		
+		/// <inheritdoc/>
+		protected override bool ResolveSelf(UhtResolvePhase phase)
+		{
+			if (phase == UhtResolvePhase.Bases && (Session.TargetSettings?.GenerateEnumMaxValues ?? false))
+			{
+				// TODO: Some engine types have authored max values that don't match the established pattern. 
+				// We should add metadata to explicitly tag the max value, or module settings to relax the pattern.
+				string prefix = GenerateEnumPrefix();
+				string maxName = CleanEnumValueName(prefix + "_MAX");
+				bool hasMax = GetIndexByName(GetFullEnumName("MAX")) != -1
+					|| GetIndexByName(maxName) != -1;
+				if (!hasMax)
+				{
+					GeneratedMaxName = maxName;
+					bool hasUnparsedValue = false;
+					long maxValue = 0;
+					if (EnumFlags.HasAllFlags(EEnumFlags.Flags))
+					{
+						foreach (UhtEnumValue enumValue in EnumValues)
+						{
+							hasUnparsedValue = hasUnparsedValue || enumValue.Value == -1;
+							if (System.Numerics.BitOperations.IsPow2(enumValue.Value))
+							{
+								maxValue |= enumValue.Value;
+							}
+						}
+					}
+					else
+					{
+						maxValue = EnumValues.Count == 0 ? 0 : EnumValues[0].Value;
+						hasUnparsedValue = maxValue == -1;
+						for (int i = 1; i < EnumValues.Count; ++i)
+						{
+							maxValue = Math.Max(maxValue, EnumValues[i].Value);
+							hasUnparsedValue = hasUnparsedValue || EnumValues[i].Value == -1;
+						}
+					}
+					if (!hasUnparsedValue)
+					{
+						GeneratedMaxValue = maxValue + 1;
+					}
+				}
+			}
+			return base.ResolveSelf(phase);
+		}
 
 		/// <summary>
 		/// Test to see if the value is a known enum value
@@ -326,6 +390,92 @@ namespace EpicGames.UHT.Types
 					throw new UhtIceException("Unexpected UhtEnumCppForm value");
 			}
 		}
+		
+		/// <summary>
+		/// Return the enum name without any namespacing
+		/// </summary>
+		/// <param name="longEnumName"></param>
+		/// <returns></returns>
+		public string GetShortEnumName(string longEnumName)
+		{
+			switch (CppForm)
+			{
+				case UhtEnumCppForm.Namespaced:
+				case UhtEnumCppForm.EnumClass:
+					int lastColons = longEnumName.LastIndexOf("::", StringComparison.Ordinal);
+					return lastColons == -1 ? longEnumName: longEnumName[(lastColons + 2)..];
+
+				case UhtEnumCppForm.Regular:
+					return longEnumName;
+
+				default:
+					throw new UhtIceException("Unexpected UhtEnumCppForm value");
+			}
+		}
+
+		/// <summary>
+		/// Return a prefix common to all members of this enumeration, or the enum name, without a trailing underscore.
+		/// This should return the same strings as UEnum::GenerateEnumPrefix.
+		/// Note that UEnum::GenerateEnumPrefix generates some surprising results, such as giving a namespaced enum such as 
+		/// ESomeEnum_WithUnderScore whose short names share no prefix the prefix ESomeEnum_, giving the generated max name
+		/// of ESomeEnum_WithUnderScore::ESomeEnum_MAX
+		/// </summary>
+		/// <returns>Depending on the type of enum and where underscores appear, this may return a short or long enum name</returns>
+		private string GenerateEnumPrefix()
+		{
+			if (EnumValues.Count == 0)
+			{
+				return SourceName;
+			}
+			string prefix = EnumValues[0].Name;
+
+			// For each item in the enumeration, trim the prefix as much as necessary to keep it a prefix.
+			// This ensures that once all items have been processed, a common prefix will have been constructed.
+			// This will be the longest common prefix since as little as possible is trimmed at each step.
+			for (int nameIdx = 1; nameIdx < EnumValues.Count; nameIdx++)
+			{
+				if (!prefix.Contains('_', StringComparison.InvariantCulture))
+				{
+					// See below - if we can't generate a prefix containing an underscore we'll discard it at the end anyway
+					return SourceName;
+				}
+				string enumItemName = EnumValues[nameIdx].Name;
+
+				// Find the length of the longest common prefix of prefix and enumItemName.
+				int prefixIdx = 0;
+				while (prefixIdx < prefix.Length && prefixIdx < enumItemName.Length && prefix[prefixIdx] == enumItemName[prefixIdx])
+				{
+					prefixIdx++;
+				}
+
+				// Trim the prefix to the length of the common prefix.
+				prefix = prefix.Substring(0, prefixIdx);
+			}
+
+			// Find the index of the rightmost underscore in the prefix.
+			int underscoreIdx = prefix.LastIndexOf('_');
+
+			// If an underscore was found, trim the prefix so only the part before the rightmost underscore is included.
+			if (underscoreIdx > 0)
+			{
+				prefix = prefix[..underscoreIdx];
+			}
+			else
+			{
+				// no underscores in the common prefix - this probably indicates that the names
+				// for this enum are not using Epic's notation, so just empty the prefix so that
+				// the max item will use the full name of the enum
+				prefix = "";
+			}
+
+			// If no common prefix was found, or if the enum does not contain any entries,
+			// use the name of the enumeration instead.
+			if (prefix.Length == 0)
+			{
+				prefix = SourceName;
+			}
+			return prefix;
+		}
 
 		/// <summary>
 		/// Add a new enum value.
@@ -394,9 +544,11 @@ namespace EpicGames.UHT.Types
 		public override void CollectReferences(IUhtReferenceCollector collector)
 		{
 			collector.AddExportType(this);
-			collector.AddDeclaration(this, true);
-			collector.AddCrossModuleReference(this, true);
-			collector.AddCrossModuleReference(Package, true);
+			collector.AddDeclaration(this, UhtSingletonType.Registered);
+			collector.AddCrossModuleReference(this, UhtSingletonType.Registered);
+			collector.AddCrossModuleReference(Package, UhtSingletonType.Registered);
+
+			collector.AddCrossModuleReference(Session.UEnum, UhtSingletonType.ConstInit);
 		}
 	}
 }
diff --git a/Engine/Source/Programs/Shared/EpicGames.UHT/Types/UhtFunction.cs b/Engine/Source/Programs/Shared/EpicGames.UHT/Types/UhtFunction.cs
--- a/Engine/Source/Programs/Shared/EpicGames.UHT/Types/UhtFunction.cs
+++ b/Engine/Source/Programs/Shared/EpicGames.UHT/Types/UhtFunction.cs
@@ -6,6 +6,7 @@ using System.Linq;
 using System.Text;
 using System.Text.Json.Serialization;
 using EpicGames.Core;
+using EpicGames.UHT.Exporters.CodeGen;
 using EpicGames.UHT.Parsers;
 using EpicGames.UHT.Tables;
 using EpicGames.UHT.Tokenizer;
@@ -264,6 +265,12 @@ namespace EpicGames.UHT.Types
 		[JsonConverter(typeof(JsonStringEnumConverter))]
 		public UhtFunctionType FunctionType { get; set; } = UhtFunctionType.Function;
 
+		/// <summary>
+		/// Super struct type for inheritance, i.e. functions do not inherit properties from their Super
+		/// </summary>
+		[JsonIgnore]
+		public override UhtStruct? InheritanceSuperStruct => null;
+
 		/// <summary>
 		/// The line number for the macro.
 		/// </summary>
@@ -285,9 +292,30 @@ namespace EpicGames.UHT.Types
 		[JsonConverter(typeof(UhtNullableTypeSourceNameJsonConverter<UhtFunction>))]
 		public UhtFunction? SuperFunction => (UhtFunction?)Super;
 
+		/// <summary>
+		/// After functions are fully resolved, points to the next function in the owning struct in each define scope.
+		/// </summary>
+		public UhtDefineScopeLink<UhtFunction>? NextFunction { get; set; }
+
 		/// <inheritdoc/>
 		[JsonIgnore]
-		public override UhtEngineType EngineType => FunctionFlags.HasAnyFlags(EFunctionFlags.Delegate) ? UhtEngineType.Delegate : UhtEngineType.Function;
+		public override UhtEngineType EngineType
+		{
+			get
+			{
+				switch (FunctionType)
+				{
+					case UhtFunctionType.Function:
+						return UhtEngineType.Function;
+					case UhtFunctionType.Delegate:
+						return UhtEngineType.Delegate;
+					case UhtFunctionType.SparseDelegate:
+						return UhtEngineType.SparseDelegate;
+					default:
+						throw new UhtIceException("Invalid function type");
+				}
+			}
+		}
 
 		/// <inheritdoc/>
 		public override string EngineClassName
@@ -297,7 +325,7 @@ namespace EpicGames.UHT.Types
 				switch (FunctionType)
 				{
 					case UhtFunctionType.Function:
-						return "Function";
+						return IsVerseField ? "VerseFunction" : "Function";
 					case UhtFunctionType.Delegate:
 						return "DelegateFunction";
 					case UhtFunctionType.SparseDelegate:
@@ -308,6 +336,28 @@ namespace EpicGames.UHT.Types
 			}
 		}
 
+		/// <inheritdoc/>
+		public override string EngineLinkClassName => "Function";
+
+		/// <inheritdoc/>
+		public override UhtClass? EngineClass
+		{
+			get
+			{
+				switch (FunctionType)
+				{
+					case UhtFunctionType.Function:
+						return IsVerseField ? Session.UVerseFunction : Session.UFunction;
+					case UhtFunctionType.Delegate:
+						return Session.UDelegateFunction;
+					case UhtFunctionType.SparseDelegate:
+						return Session.USparseDelegateFunction;
+					default:
+						throw new UhtIceException("Invalid function type");
+				}
+			}
+		}
+
 		/// <summary>
 		/// Stripped function name without the generated delegate suffix
 		/// </summary>
@@ -1250,14 +1300,31 @@ namespace EpicGames.UHT.Types
 			{
 				collector.AddExportType(this);
 				collector.AddSingleton(this);
-				collector.AddDeclaration(this, true);
-				collector.AddCrossModuleReference(this, true);
-				collector.AddCrossModuleReference(Package, true);
+				collector.AddDeclaration(this, UhtSingletonType.Registered);
+				collector.AddCrossModuleReference(this, UhtSingletonType.Registered);
+				collector.AddCrossModuleReference(Package, UhtSingletonType.Registered);
 			}
 			foreach (UhtType child in Children)
 			{
 				child.CollectReferences(collector);
 			}
+
+			switch (FunctionType)
+			{
+				case UhtFunctionType.Function:
+					collector.AddCrossModuleReference(Session.UFunction, UhtSingletonType.ConstInit);
+					collector.AddCrossModuleReference(this, UhtSingletonType.ConstInit);
+					collector.AddCrossModuleReference(Package, UhtSingletonType.ConstInit);
+					break;
+				case UhtFunctionType.Delegate:
+					collector.AddCrossModuleReference(Session.UDelegateFunction, UhtSingletonType.ConstInit);
+					break;
+				case UhtFunctionType.SparseDelegate:
+					collector.AddCrossModuleReference(Session.USparseDelegateFunction, UhtSingletonType.ConstInit);
+					break;
+				default:
+					break;
+			}
 		}
 
 		/// <summary>
diff --git a/Engine/Source/Programs/Shared/EpicGames.UHT/Types/UhtPackage.cs b/Engine/Source/Programs/Shared/EpicGames.UHT/Types/UhtPackage.cs
--- a/Engine/Source/Programs/Shared/EpicGames.UHT/Types/UhtPackage.cs
+++ b/Engine/Source/Programs/Shared/EpicGames.UHT/Types/UhtPackage.cs
@@ -33,6 +33,9 @@ namespace EpicGames.UHT.Types
 		/// <inheritdoc/>
 		public override string EngineClassName => "Package";
 
+		/// <inheritdoc/>
+		public override UhtClass? EngineClass => Session.UPackage;
+
 		/// <inheritdoc/>
 		[JsonIgnore]
 		public override UhtPackage Package => this;
diff --git a/Engine/Source/Programs/Shared/EpicGames.UHT/Types/UhtProperty.cs b/Engine/Source/Programs/Shared/EpicGames.UHT/Types/UhtProperty.cs
--- a/Engine/Source/Programs/Shared/EpicGames.UHT/Types/UhtProperty.cs
+++ b/Engine/Source/Programs/Shared/EpicGames.UHT/Types/UhtProperty.cs
@@ -5,10 +5,10 @@ using System.Collections.Generic;
 using System.Diagnostics;
 using System.Diagnostics.CodeAnalysis;
 using System.Globalization;
-using System.Linq;
 using System.Text;
 using System.Text.Json.Serialization;
 using EpicGames.Core;
+using EpicGames.UHT.Exporters.CodeGen;
 using EpicGames.UHT.Parsers;
 using EpicGames.UHT.Tables;
 using EpicGames.UHT.Tokenizer;
@@ -647,7 +647,7 @@ namespace EpicGames.UHT.Types
 	/// <summary>
 	/// Property member context provides extra context when formatting the member declaration and definitions.
 	/// </summary>
-	public interface IUhtPropertyMemberContext
+	public interface IUhtPropertyMemberContext : IUhtSingletonNameProvider, IUhtSessionSettingsProvider
 	{
 		/// <summary>
 		/// The outer/owning structure
@@ -686,13 +686,6 @@ namespace EpicGames.UHT.Types
 		/// <returns>Hash code</returns>
 		public uint GetTypeHash(UhtObject obj);
 
-		/// <summary>
-		/// Return the singleton name for the given object
-		/// </summary>
-		/// <param name="obj">Object in question</param>
-		/// <param name="registered">If true, the singleton that returns the registered object is returned.</param>
-		/// <returns>Singleton function name</returns>
-		public string GetSingletonName(UhtObject? obj, bool registered);
 	}
 
 	/// <summary>
@@ -1148,6 +1141,10 @@ namespace EpicGames.UHT.Types
 		[JsonIgnore]
 		public override UhtEngineType EngineType => UhtEngineType.Property;
 
+		///<inheritdoc/>
+		[JsonIgnore]
+		public override UhtClass? EngineClass => null; // Properties do not have a UObject derived class
+
 		///<inheritdoc/>
 		[JsonIgnore]
 		public override bool Deprecated => PropertyFlags.HasAnyFlags(EPropertyFlags.Deprecated);
@@ -1190,6 +1187,16 @@ namespace EpicGames.UHT.Types
 		[JsonIgnore]
 		public bool IsEditorOnlyProperty => PropertyFlags.HasAnyFlags(EPropertyFlags.EditorOnly);
 
+		/// <summary>
+		/// After properties are fully resolved, points to the next property in the owning struct in each define scope.
+		/// </summary>
+		public UhtDefineScopeLink<UhtProperty>? NextProperty { get; set; }
+
+		/// <summary>
+		/// Returns whether this property has a setter or getter function and therefor needs to be wrapped in a template property type
+		/// </summary>
+		public bool HasGetterOrSetter => PropertyExportFlags.HasAnyFlags(UhtPropertyExportFlags.SetterFound | UhtPropertyExportFlags.GetterFound);
+
 		/// <summary>
 		/// Return the root property for the given property.  If the property is outside of a container, 
 		/// return the property itself.  Otherwise, return the property associated with the container.
@@ -1349,7 +1356,7 @@ namespace EpicGames.UHT.Types
 		/// <returns>Output builder</returns>
 		public virtual StringBuilder AppendMetaDataDecl(StringBuilder builder, IUhtPropertyMemberContext context, string name, string nameSuffix, int tabs)
 		{
-			return builder.AppendMetaDataDecl(this, context.NamePrefix, name, nameSuffix, context.MetaDataSuffix, tabs);
+			return builder.AppendMetaDataDefinition(this, context, context.NamePrefix, name, nameSuffix, context.MetaDataSuffix, tabs);
 		}
 
 		/// <summary>
@@ -1363,6 +1370,24 @@ namespace EpicGames.UHT.Types
 		/// <returns>Output builder</returns>
 		public abstract StringBuilder AppendMemberDecl(StringBuilder builder, IUhtPropertyMemberContext context, string name, string nameSuffix, int tabs);
 
+		/// <summary>
+		/// Append a declaration of the given property as a constinit variable.
+		/// For non-compound properties this is simply a declaration of a single variable of the appropriate type.
+		/// For comound properties this must be overridden to declare the inner/key/value property.
+		/// </summary>
+		/// <param name="builder">Output builder</param>
+		/// <param name="context">Current context</param>
+		/// <param name="name">Name of the property.  This is needed in some cases where the name in the declarations doesn't match the property name.</param>
+		/// <param name="nameSuffix">Suffix to the property name</param>
+		/// <param name="tabs">Number of tabs prefix the line with</param>
+		/// <returns>Output builder</returns>
+		public virtual StringBuilder AppendConstInitMemberDecl(StringBuilder builder, IUhtPropertyMemberContext context, string name, string nameSuffix, int tabs)
+		{
+			string? typePrefix = HasGetterOrSetter ? "TPropertyWithSetterAndGetter<" : null;
+			string? typeSuffix = HasGetterOrSetter ? ">" : null;
+			return builder.AppendTabs(tabs).Append($"static UE_CONSTINIT_UOBJECT_DECL TNoDestroy<{typePrefix}F{EngineClassName}{typeSuffix}> ").AppendNameDecl(context, name, nameSuffix).Append(";\r\n");
+		}
+
 		/// <summary>
 		/// Append the required code to declare the property as a member
 		/// </summary>
@@ -1391,6 +1416,19 @@ namespace EpicGames.UHT.Types
 		/// <returns>Output builder</returns>
 		public abstract StringBuilder AppendMemberDef(StringBuilder builder, IUhtPropertyMemberContext context, string name, string nameSuffix, string? offset, int tabs);
 
+		/// <summary>
+		/// Append a full definition of constinit variable(s) for this property
+		/// </summary>
+		/// <param name="builder">Output builder</param>
+		/// <param name="context">Context of the call</param>
+		/// <param name="name">Name of the property.  This is needed in some cases where the name in the declarations doesn't match the property name.</param>
+		/// <param name="nameSuffix">Suffix to the property name</param>
+		/// <param name="outerFunc">Optional function to append an expression for the address of the outer property of this property</param>
+		/// <param name="offset">Offset of the property</param>
+		/// <param name="tabs">Number of tabs prefix the line with</param>
+		/// <returns>Output builder</returns>
+		public abstract StringBuilder AppendConstInitMemberDef(StringBuilder builder, IUhtPropertyMemberContext context, string name, string nameSuffix, Action<StringBuilder>? outerFunc, string? offset, int tabs);
+
 		/// <summary>
 		/// Append the required start of code to define the property as a member
 		/// </summary>
@@ -1451,6 +1489,122 @@ namespace EpicGames.UHT.Types
 			return builder;
 		}
 
+		/// <summary>
+		/// Append the required start of code to define a constinit property variable for the member
+		/// </summary>
+		/// <param name="builder">Output builder</param>
+		/// <param name="context">Context of the call</param>
+		/// <param name="name">Name of the property.  This is needed in some cases where the name in the declarations doesn't match the property name.</param>
+		/// <param name="nameSuffix">Suffix to the property name</param>
+		/// <param name="outerFunc">Optional function to append an expression for the outer property of this property</param>
+		/// <param name="offset">Offset of the property</param>
+		/// <param name="tabs">Number of tabs prefix the line with</param>
+		/// <returns>Output builder</returns>
+		public StringBuilder AppendConstInitMemberDefStart(StringBuilder builder, IUhtPropertyMemberContext context, string name, string nameSuffix, Action<StringBuilder>? outerFunc, string? offset, int tabs)
+		{
+			string? typePrefix = HasGetterOrSetter ? "TPropertyWithSetterAndGetter<" : null;
+			string? typeSuffix = HasGetterOrSetter ? ">" : null;
+			builder.AppendTabs(tabs)
+			 	.Append($"UE_CONSTINIT_UOBJECT_DECL TNoDestroy<{typePrefix}F{EngineClassName}{typeSuffix}> ")
+				.AppendNameDef(context, name, nameSuffix).Append("{NoDestroyConstEval, ");
+
+			if (HasGetterOrSetter)
+			{
+				if (PropertyExportFlags.HasAnyFlags(UhtPropertyExportFlags.SetterFound))
+				{
+					builder.Append('&').Append(Outer!.SourceName).Append("::").AppendPropertySetterWrapperName(this).Append(", ");
+				}
+				else
+				{
+					builder.Append("nullptr, ");
+				}
+
+				if (PropertyExportFlags.HasAnyFlags(UhtPropertyExportFlags.GetterFound))
+				{
+					builder.Append('&').Append(Outer!.SourceName).Append("::").AppendPropertyGetterWrapperName(this).Append(", ");
+				}
+				else
+				{
+					builder.Append("nullptr, ");
+				}
+			}
+
+			builder.Append($"UE::CodeGen::ConstInit::FPropertyParams{{ .FieldClass = &F{EngineClassName}::ConstInitClass, .Owner = ");
+			if (outerFunc is not null)
+			{
+				outerFunc(builder);
+				builder.Append(", ");
+			}
+			else if (Outer is UhtObject obj)
+			{
+				builder.Append($"&{context.GetSingletonName(obj, UhtSingletonType.ConstInit)}, ");
+			}
+			else
+			{
+				throw new UhtException("Property had an outer which was not a UhtObject and no outerFunc was passed");
+			}
+			// Only properties inside UhtStruct have a next property - properties inside other properties do not 
+			if (NextProperty is not null)
+			{
+				builder.Append(".NextProperty = ");
+				if (NextProperty!.Next.Count > 1)
+				{
+					// Call the GetNextProperty function we defined
+					builder.Append($"GetNextProperty_{EngineName}(), ");
+				}
+				else 
+				{
+					// 0 or 1, so we should be appending null or a property defined in the same scopes as this. Otherwise the builder should have appended an explicit null
+					UhtProperty? next = NextProperty!.Next[0];
+					builder.AppendConstInitMemberPtr(next, context, next?.EngineName, "", 0, ", ");
+				}
+			}
+			builder.Append(".NameUTF8 = UTF8TEXT(").AppendUTF8LiteralString(EngineName).Append("), ");
+			if (RepNotifyName is not null)
+			{
+				builder.Append(".RepNotifyFuncUTF8 = UTF8TEXT(").AppendUTF8LiteralString(RepNotifyName).Append("), ");
+			}
+			builder.Append(".PropertyFlags = ").AppendFlags(PropertyFlags).Append(", ");
+			builder.Append(".ObjectFlags = ").Append(ObjectFlags).Append(", ");
+
+			builder.Append(".ArrayDim = ").AppendArrayDim(this, context).Append(", ");
+
+			if (!String.IsNullOrEmpty(offset))
+			{
+				builder.Append($".Offset = {offset}, ");
+			}
+			else
+			{
+				builder.Append($".Offset = STRUCT_OFFSET({context.OuterStructNamespaceName}{context.OuterStructSourceName}, {SourceName}), ");
+			}
+			builder.Append(".ElementSize = 0, "); // At present this is overridden by properties with reference to cpp type info. In future this can be used to provide size for e.g. optional properties by referencing inner property type.
+			if (MetaData.IsEmpty())
+			{
+				builder.Append("IF_WITH_METADATA(.MetaData = {},)");
+			}
+			else
+			{
+				builder.Append("IF_WITH_METADATA(.MetaData = MakeConstArrayView(").AppendNameDef(context, name, nameSuffix).Append($"{context.MetaDataSuffix}),)");
+			}
+			builder.Append("}, ");
+			return builder;
+		}
+
+		/// <summary>
+		/// Append the required end of code to define the property as a member
+		/// </summary>
+		/// <param name="builder">Output builder</param>
+		/// <param name="context">Context of the call</param>
+		/// <returns>Output builder</returns>
+		protected StringBuilder AppendConstInitMemberDefEnd(StringBuilder builder, IUhtPropertyMemberContext context)
+		{
+			builder
+				.Append(" };")
+				.AppendObjectHashes(this, context)
+				.Append("\r\n");
+			return builder;
+		}
+
 		/// <summary>
 		/// Append the required end of code to define the property as a member
 		/// </summary>
@@ -1475,14 +1629,18 @@ namespace EpicGames.UHT.Types
 		/// <param name="builder">Output builder</param>
 		/// <param name="context">Context of the call</param>
 		/// <param name="obj">Referenced object</param>
-		/// <param name="registered">True if the registered singleton name is to be used</param>
+		/// <param name="type">Type of declaration to get</param>
 		/// <param name="appendNull">True if a "nullptr" is to be appended if the object is null</param>
 		/// <returns>Output builder</returns>
-		protected static StringBuilder AppendMemberDefRef(StringBuilder builder, IUhtPropertyMemberContext context, UhtObject? obj, bool registered, bool appendNull = false)
+		protected static StringBuilder AppendMemberDefRef(StringBuilder builder, IUhtPropertyMemberContext context, UhtObject? obj, UhtSingletonType type, bool appendNull = false)
 		{
 			if (obj != null)
 			{
-				builder.AppendSingletonName(context, obj, registered).Append(", ");
+				if (type == UhtSingletonType.ConstInit)
+				{
+					builder.Append('&');
+				}
+				builder.AppendSingletonName(context, obj, type).Append(", ");
 			}
 			else if (appendNull)
 			{
@@ -1506,6 +1664,40 @@ namespace EpicGames.UHT.Types
 			return builder;
 		}
 
+		/// <summary>
+		/// Append an expression to get a pointer to the property object
+		/// </summary>
+		/// <param name="builder">Output builder</param>
+		/// <param name="context">Context of the call</param>
+		/// <param name="name">Name of the property.  This is needed in some cases where the name in the declarations doesn't match the property name.</param>
+		/// <param name="nameSuffix">Suffix to the property name</param>
+		/// <param name="tabs">Number of tabs prefix the line with</param>
+		/// <returns>Output builder</returns>
+		public virtual StringBuilder AppendConstInitMemberPtr(StringBuilder builder, IUhtPropertyMemberContext context, string name, string nameSuffix, int tabs)
+		{
+			builder.AppendTabs(tabs).Append('&').AppendNameDef(context, name, nameSuffix);
+			return builder;
+		}
+
+		/// <summary>
+		/// If this property needs a different 'next property' pointer under different preprocessor definitions, append an inline function that returns 
+		/// the correct property in each preprocessor state.
+		/// </summary>
+		/// <param name="builder"></param>
+		/// <param name="context"></param>
+		/// <param name="tabs"></param>
+		public StringBuilder AppendGetNextProperty(StringBuilder builder, IUhtPropertyMemberContext context, int tabs)
+		{
+			if (NextProperty?.Next?.Count > 1)
+			{
+				// Note that the caller is responsible for ensuring this is scoped to when this property is defined
+				builder.AppendTabs(tabs).Append("static consteval FProperty* GetNextProperty_").Append(EngineName).Append("() {\r\n");
+				builder.AppendScopeLink(NextProperty, UhtDefineScopeNames.Standard, context, tabs + 1, "return ", ";\r\n");
+				builder.AppendTabs(tabs).Append("}\r\n");
+			}
+			return builder;
+		}
+
 		/// <summary>
 		/// Append a P_GET macro
 		/// </summary>
@@ -2182,6 +2374,10 @@ namespace EpicGames.UHT.Types
 		public override void CollectReferences(IUhtReferenceCollector collector)
 		{
 			CollectReferencesInternal(collector, PropertyFlags.HasAnyFlags(EPropertyFlags.ParmFlags), false);
+			if (Outer is UhtObject outerObject)
+			{
+				collector.AddCrossModuleReference(outerObject, UhtSingletonType.ConstInit);
+			}
 		}
 
 		/// <summary>
@@ -2315,6 +2511,21 @@ namespace EpicGames.UHT.Types
 			return property.AppendMemberDecl(builder, context, name, nameSuffix, tabs);
 		}
 
+		/// <summary>
+		/// Append a declaration of a constinit variable representing the given property.
+		/// </summary>
+		/// <param name="builder">Destination builder</param>
+		/// <param name="property">Property in question</param>
+		/// <param name="context">Context of the property</param>
+		/// <param name="name">Property name</param>
+		/// <param name="nameSuffix">Name suffix</param>
+		/// <param name="tabs">Number of tabs in the formatting</param>
+		/// <returns>Destination builder</returns>
+		public static StringBuilder AppendConstInitMemberDecl(this StringBuilder builder, UhtProperty property, IUhtPropertyMemberContext context, string name, string nameSuffix, int tabs)
+		{
+			return property.AppendConstInitMemberDecl(builder, context, name, nameSuffix, tabs);
+		}
+
 		/// <summary>
 		/// Append the member definition
 		/// </summary>
@@ -2331,6 +2542,23 @@ namespace EpicGames.UHT.Types
 			return property.AppendMemberDef(builder, context, name, nameSuffix, offset, tabs);
 		}
 
+		/// <summary>
+		/// Append a definition of the constinit variable representing the property
+		/// </summary>
+		/// <param name="builder">Destination builder</param>
+		/// <param name="property">Property in question</param>
+		/// <param name="context">Context of the property</param>
+		/// <param name="name">Property name</param>
+		/// <param name="nameSuffix">Name suffix</param>
+		/// <param name="outerFunc">Optional function to add expression for outer property</param>
+		/// <param name="offset">Offset of the property in the parent</param>
+		/// <param name="tabs">Number of tabs in the formatting</param>
+		/// <returns>Destination builder</returns>
+		public static StringBuilder AppendConstInitMemberDef(this StringBuilder builder, UhtProperty property, IUhtPropertyMemberContext context, string name, string nameSuffix, Action<StringBuilder>? outerFunc, string? offset, int tabs)
+		{
+			return property.AppendConstInitMemberDef(builder, context, name, nameSuffix, outerFunc, offset, tabs);
+		}
+
 		/// <summary>
 		/// Append the member pointer
 		/// </summary>
@@ -2346,6 +2574,26 @@ namespace EpicGames.UHT.Types
 			return property.AppendMemberPtr(builder, context, name, nameSuffix, tabs);
 		}
 
+		/// <summary>
+		/// Append an expression resolving to a pointer to the constinit object instance for this property
+		/// </summary>
+		/// <param name="builder">Destination builder</param>
+		/// <param name="property">Property in question</param>
+		/// <param name="context">Context of the property</param>
+		/// <param name="name">Property name</param>
+		/// <param name="nameSuffix">Name suffix</param>
+		/// <param name="tabs">Number of tabs in the formatting</param>
+		/// <param name="post">String to append afterwards</param>
+		/// <returns>Destination builder</returns>
+		public static StringBuilder AppendConstInitMemberPtr(this StringBuilder builder, UhtProperty? property, IUhtPropertyMemberContext context, string? name, string nameSuffix, int tabs, string? post)
+		{
+			if (property is null)
+			{
+				return builder.Append("nullptr").Append(post);
+			}
+			return property.AppendConstInitMemberPtr(builder, context, name ?? property.EngineName, nameSuffix, tabs).Append(post);
+		}
+
 		/// <summary>
 		/// Append the full declaration including such things as const, *, and &amp;
 		/// </summary>
@@ -2573,13 +2821,13 @@ namespace EpicGames.UHT.Types
 		/// Append the singleton name for the given type
 		/// </summary>
 		/// <param name="builder">Destination builder</param>
-		/// <param name="context">Context of the property</param>
+		/// <param name="provider">Object to provide the names of singletons</param>
 		/// <param name="type">Type to append</param>
-		/// <param name="registered">If true, append the registered type singleton.</param>
+		/// <param name="singletonType">Type of declaration to append</param>
 		/// <returns>Destination builder</returns>
-		public static StringBuilder AppendSingletonName(this StringBuilder builder, IUhtPropertyMemberContext context, UhtObject? type, bool registered)
+		public static StringBuilder AppendSingletonName(this StringBuilder builder, IUhtSingletonNameProvider provider, UhtObject? type, UhtSingletonType singletonType)
 		{
-			return builder.Append(context.GetSingletonName(type, registered));
+			return builder.Append(provider.GetSingletonName(type, singletonType));
 		}
 
 		/// <summary>
diff --git a/Engine/Source/Programs/Shared/EpicGames.UHT/Types/UhtScriptStruct.cs b/Engine/Source/Programs/Shared/EpicGames.UHT/Types/UhtScriptStruct.cs
--- a/Engine/Source/Programs/Shared/EpicGames.UHT/Types/UhtScriptStruct.cs
+++ b/Engine/Source/Programs/Shared/EpicGames.UHT/Types/UhtScriptStruct.cs
@@ -5,6 +5,7 @@ using System.Collections.Generic;
 using System.Text;
 using System.Text.Json.Serialization;
 using EpicGames.Core;
+using EpicGames.UHT.Exporters.CodeGen;
 using EpicGames.UHT.Parsers;
 using EpicGames.UHT.Tables;
 using EpicGames.UHT.Utils;
@@ -725,7 +726,13 @@ namespace EpicGames.UHT.Types
 		public bool HasNoOpConstructor => ScriptStructExportFlags.HasAnyFlags(UhtScriptStructExportFlags.HasNoOpConstructor);
 
 		/// <inheritdoc/>
-		public override string EngineClassName => "ScriptStruct";
+		public override string EngineClassName => IsVerseField ? "VerseStruct" : "ScriptStruct";
+
+		/// <inheritdoc/>
+		public override string EngineLinkClassName => "ScriptStruct";
+
+		/// <inheritdoc/>
+		public override UhtClass? EngineClass => IsVerseField ? Session.UVerseStruct : Session.UScriptStruct;
 
 		/// <inheritdoc/>
 		[JsonIgnore]
@@ -996,17 +1003,23 @@ namespace EpicGames.UHT.Types
 				collector.AddSingleton(this);
 			}
 			collector.AddExportType(this);
-			collector.AddDeclaration(this, true);
-			collector.AddCrossModuleReference(this, true);
+			collector.AddDeclaration(this, UhtSingletonType.Registered);
+			collector.AddCrossModuleReference(this, UhtSingletonType.Registered);
 			if (SuperScriptStruct != null)
 			{
-				collector.AddCrossModuleReference(SuperScriptStruct, true);
+				collector.AddCrossModuleReference(SuperScriptStruct, UhtSingletonType.Registered);
 			}
-			collector.AddCrossModuleReference(Package, true);
+			collector.AddCrossModuleReference(Package, UhtSingletonType.Registered);
 			foreach (UhtType child in Children)
 			{
 				child.CollectReferences(collector);
 			}
+
+			collector.AddCrossModuleReference(Session.UScriptStruct, UhtSingletonType.ConstInit);
+			for (UhtStruct? superStruct = SuperScriptStruct; superStruct is not null; superStruct = superStruct.SuperStruct)
+			{
+				collector.AddCrossModuleReference(superStruct, UhtSingletonType.ConstInit);
+			}
 		}
 
 		private void CollectRigVMMembers()
diff --git a/Engine/Source/Programs/Shared/EpicGames.UHT/Types/UhtStruct.cs b/Engine/Source/Programs/Shared/EpicGames.UHT/Types/UhtStruct.cs
--- a/Engine/Source/Programs/Shared/EpicGames.UHT/Types/UhtStruct.cs
+++ b/Engine/Source/Programs/Shared/EpicGames.UHT/Types/UhtStruct.cs
@@ -7,6 +7,7 @@ using System.Linq;
 using System.Text.Json.Serialization;
 using System.Threading;
 using EpicGames.Core;
+using EpicGames.UHT.Exporters.CodeGen;
 using EpicGames.UHT.Parsers;
 using EpicGames.UHT.Tables;
 using EpicGames.UHT.Tokenizer;
@@ -79,8 +80,8 @@ namespace EpicGames.UHT.Types
 		public virtual string EngineNamePrefix => "F";
 
 		/// <inheritdoc/>
-		public override string EngineClassName => "Struct";
-
+		public override string EngineClassName => IsVerseField ? "VerseStruct" : "Struct";
+		
 		/// <summary>
 		/// Super type
 		/// </summary>
@@ -112,6 +113,22 @@ namespace EpicGames.UHT.Types
 		[JsonIgnore]
 		public UhtStruct? SuperStruct => Super;
 
+		/// <summary>
+		/// Super struct type for inheritance, i.e. functions do not inherit properties from their Super
+		/// </summary>
+		[JsonIgnore]
+		public virtual UhtStruct? InheritanceSuperStruct => Super;
+
+		/// <summary>
+		/// After properties are fully resolved, points to the first property in this struct
+		/// </summary>
+		public UhtDefineScopeLink<UhtProperty>? FirstProperty { get; set; }
+
+		/// <summary>
+		/// For classes, struct type for sparse class data shared between all instances
+		/// </summary>
+		public UhtScriptStruct? SparseClassDataStruct { get; set; }
+
 		/// <summary>
 		/// Stack used to test to see if we are recursing in a ScanForInstanceReferenced call
 		/// </summary>
@@ -169,6 +186,32 @@ namespace EpicGames.UHT.Types
 			base.ResolveSuper(resolvePhase);
 		}
 
+		/// <inheritdoc/>
+		protected override void ResolveChildren(UhtResolvePhase phase)
+		{
+			base.ResolveChildren(phase);
+
+			if (phase == UhtResolvePhase.Final)
+			{
+				// Build a linked list of each property of this struct/class/function
+				using (UhtDefineScopeListBuilder<UhtProperty> propertyList = new(DefineScope))
+				{
+					foreach (UhtType type in Children)
+					{
+						if (type is UhtProperty property)
+						{
+							if (property.NextProperty is not null)
+							{
+								throw new Exception("Unexpectedly initialized next link on property");
+							}
+							property.NextProperty = propertyList.Add(property);
+						}
+					}
+					FirstProperty = propertyList.Head;
+				}
+			}
+		}
+
 		/// <summary>
 		/// Check properties to see if any instances are referenced.
 		/// This method does NOT cache the result.
@@ -412,33 +455,33 @@ namespace EpicGames.UHT.Types
 
 			foreach (string sparseClassDataTypeName in sparseClassDataTypes)
 			{
-				UhtScriptStruct? sparseScriptStruct = Session.FindType(null, UhtFindOptions.EngineName | UhtFindOptions.ScriptStruct, sparseClassDataTypeName) as UhtScriptStruct;
+				SparseClassDataStruct = Session.FindType(null, UhtFindOptions.EngineName | UhtFindOptions.ScriptStruct, sparseClassDataTypeName) as UhtScriptStruct;
 
 				// make sure the sparse class data struct actually exists
-				if (sparseScriptStruct == null)
+				if (SparseClassDataStruct  == null)
 				{
 					this.LogError($"Unable to find sparse data type '{sparseClassDataTypeName}' for class '{SourceName}'");
 					continue;
 				}
 
 				// check the data struct for invalid properties
-				foreach (UhtProperty property in sparseScriptStruct.Properties)
+				foreach (UhtProperty property in SparseClassDataStruct  .Properties)
 				{
 					if (property.PropertyFlags.HasAnyFlags(EPropertyFlags.BlueprintAssignable))
 					{
-						property.LogError($"Sparse class data types can not contain blueprint assignable delegates. Type '{sparseScriptStruct.EngineName}' Delegate '{property.SourceName}'");
+						property.LogError($"Sparse class data types can not contain blueprint assignable delegates. Type '{SparseClassDataStruct  .EngineName}' Delegate '{property.SourceName}'");
 					}
 
 					// all sparse properties should have EditDefaultsOnly
 					if (!property.PropertyFlags.HasFlag(EPropertyFlags.Deprecated) && !property.PropertyFlags.HasAllFlags(EPropertyFlags.Edit | EPropertyFlags.DisableEditOnInstance))
 					{
-						property.LogError($"Sparse class data types must be VisibleDefaultsOnly or EditDefaultsOnly. Type '{sparseScriptStruct.EngineName}' Property '{property.SourceName}'");
+						property.LogError($"Sparse class data types must be VisibleDefaultsOnly or EditDefaultsOnly. Type '{SparseClassDataStruct  .EngineName}' Property '{property.SourceName}'");
 					}
 
 					// no sparse properties should have BlueprintReadWrite
 					if (property.PropertyFlags.HasAnyFlags(EPropertyFlags.BlueprintVisible) && !property.PropertyFlags.HasAnyFlags(EPropertyFlags.BlueprintReadOnly))
 					{
-						property.LogError($"Sparse class data types must not be BlueprintReadWrite. Type '{sparseScriptStruct.EngineName}' Property '{property.SourceName}'");
+						property.LogError($"Sparse class data types must not be BlueprintReadWrite. Type '{SparseClassDataStruct  .EngineName}' Property '{property.SourceName}'");
 					}
 				}
 
@@ -454,11 +497,11 @@ namespace EpicGames.UHT.Types
 							UhtScriptStruct? superSparseScriptStruct = Session.FindType(null, UhtFindOptions.EngineName | UhtFindOptions.ScriptStruct, superSparseClassDataTypeName) as UhtScriptStruct;
 							if (superSparseScriptStruct != null)
 							{
-								if (!sparseScriptStruct.IsChildOf(superSparseScriptStruct))
+								if (!SparseClassDataStruct  .IsChildOf(superSparseScriptStruct))
 								{
 									this.LogError(
 										$"Class '{SourceName}' is a child of '{superClass.SourceName}' but its sparse class data struct " +
-										$"'{sparseScriptStruct.EngineName}', does not inherit from '{superSparseScriptStruct.EngineName}'.");
+										$"'{SparseClassDataStruct  .EngineName}', does not inherit from '{superSparseScriptStruct.EngineName}'.");
 								}
 							}
 						}
@@ -467,5 +510,12 @@ namespace EpicGames.UHT.Types
 			}
 		}
 		#endregion
+
+		/// <inheritdoc/>
+		public override void CollectReferences(IUhtReferenceCollector collector)
+		{
+			base.CollectReferences(collector);
+			collector.AddCrossModuleReference(SparseClassDataStruct, UhtSingletonType.ConstInit);
+		}
 	}
 }
diff --git a/Engine/Source/Programs/Shared/EpicGames.UHT/Types/UhtType.cs b/Engine/Source/Programs/Shared/EpicGames.UHT/Types/UhtType.cs
--- a/Engine/Source/Programs/Shared/EpicGames.UHT/Types/UhtType.cs
+++ b/Engine/Source/Programs/Shared/EpicGames.UHT/Types/UhtType.cs
@@ -61,6 +61,11 @@ namespace EpicGames.UHT.Types
 		/// </summary>
 		Delegate,
 
+		/// <summary>
+		/// Type is a sparse delegate (UhtFunction)
+		/// </summary>
+		SparseDelegate,
+
 		/// <summary>
 		/// Type is a property (UhtProperty)
 		/// </summary>
@@ -90,6 +95,7 @@ namespace EpicGames.UHT.Types
 				case UhtEngineType.ScriptStruct: return "struct";
 				case UhtEngineType.Enum: return "enum";
 				case UhtEngineType.Function:
+				case UhtEngineType.SparseDelegate:
 				case UhtEngineType.Delegate: return "function";
 				case UhtEngineType.Property: return "property";
 				default:
@@ -114,6 +120,7 @@ namespace EpicGames.UHT.Types
 				case UhtEngineType.ScriptStruct: return "Struct";
 				case UhtEngineType.Enum: return "Enum";
 				case UhtEngineType.Function:
+				case UhtEngineType.SparseDelegate:
 				case UhtEngineType.Delegate: return "Function";
 				case UhtEngineType.Property: return "Property";
 				default:
@@ -139,6 +146,7 @@ namespace EpicGames.UHT.Types
 				case UhtEngineType.Enum: return "enum";
 				case UhtEngineType.Function: return "function";
 				case UhtEngineType.Delegate: return "delegate";
+				case UhtEngineType.SparseDelegate: return "sparse delegate";
 				case UhtEngineType.Property: return "property";
 				default:
 					throw new UhtIceException("Invalid engine extended type");
@@ -163,6 +171,7 @@ namespace EpicGames.UHT.Types
 				case UhtEngineType.Enum: return "Enum";
 				case UhtEngineType.Function: return "Function";
 				case UhtEngineType.Delegate: return "Delegate";
+				case UhtEngineType.SparseDelegate: return "SparseDelegate";
 				case UhtEngineType.Property: return "Property";
 				default:
 					throw new UhtIceException("Invalid engine extended type");
@@ -279,6 +288,7 @@ namespace EpicGames.UHT.Types
 				case UhtEngineType.Enum: return UhtFindOptions.Enum;
 				case UhtEngineType.Function: return UhtFindOptions.Function;
 				case UhtEngineType.Delegate: return UhtFindOptions.DelegateFunction;
+				case UhtEngineType.SparseDelegate: return UhtFindOptions.DelegateFunction;
 				case UhtEngineType.Property: return UhtFindOptions.Property;
 				default:
 					throw new UhtIceException("Invalid engine extended type");
@@ -837,10 +847,22 @@ namespace EpicGames.UHT.Types
 		public abstract UhtEngineType EngineType { get; }
 
 		/// <summary>
-		/// The name of the engine's class for this type
+		/// The name of the engine's class for this type i.e. Class, Function
 		/// </summary>
 		public abstract string EngineClassName { get; }
 
+		/// <summary>
+		/// The name of the engine's class for this type for linking with registration functions. 
+		/// i.e. EngineClassName may be VerseStruct, but EngineLinkClassName may be ScriptStruct
+		/// </summary>
+		public virtual string EngineLinkClassName => EngineClassName;
+
+		/// <summary>
+		/// Return the runtime type of this object in the engine after parsing, when the Session
+		/// has references to such types available.
+		/// </summary>
+		public abstract UhtClass? EngineClass { get; }
+
 		/// <summary>
 		/// Return true if the type has been deprecated
 		/// </summary>
diff --git a/Engine/Source/Programs/Shared/EpicGames.UHT/Utils/UhtReferenceCollector.cs b/Engine/Source/Programs/Shared/EpicGames.UHT/Utils/UhtReferenceCollector.cs
--- a/Engine/Source/Programs/Shared/EpicGames.UHT/Utils/UhtReferenceCollector.cs
+++ b/Engine/Source/Programs/Shared/EpicGames.UHT/Utils/UhtReferenceCollector.cs
@@ -4,6 +4,7 @@ using System;
 using System.Collections.Generic;
 using System.Text;
 using EpicGames.Core;
+using EpicGames.UHT.Exporters.CodeGen;
 using EpicGames.UHT.Types;
 
 namespace EpicGames.UHT.Utils
@@ -19,16 +20,16 @@ namespace EpicGames.UHT.Utils
 		/// <summary>
 		/// Add a cross module reference to a given object type.
 		/// </summary>
-		/// <param name="obj">Object type being referenced</param>
-		/// <param name="registered">If true, the method being invoked must return the registered object.  This only applies to classes.</param>
-		void AddCrossModuleReference(UhtObject? obj, bool registered);
+		/// <param name="obj">Object/type being referenced</param>
+		/// <param name="type">Type of reference required</param>
+		void AddCrossModuleReference(UhtObject? obj, UhtSingletonType type);
 
 		/// <summary>
 		/// Add an object declaration
 		/// </summary>
 		/// <param name="obj">Object in question</param>
-		/// <param name="registered">If true, the method being invoked must return the registered object.  This only applies to classes.</param>
-		void AddDeclaration(UhtObject obj, bool registered);
+		/// <param name="type">Type of reference required</param>
+		void AddDeclaration(UhtObject obj, UhtSingletonType type);
 
 		/// <summary>
 		/// Add a field as a singleton for exporting
@@ -57,96 +58,86 @@ namespace EpicGames.UHT.Utils
 		void AddForwardDeclaration(UhtNamespace? namespaceObj, string? declaration);
 	}
 
-	/// <summary>
-	/// Delegate used to fetch the string associated with a reference
-	/// </summary>
-	/// <param name="objectIndex">Index of the referenced object</param>
-	/// <param name="registered">If true return the registered string, otherwise the unregistered string.  Classes have an unregistered version.</param>
-	/// <returns>The requested string</returns>
-	public delegate string GetReferenceStringDelegate(int objectIndex, bool registered);
-
 	/// <summary>
 	/// Maintains a list of referenced object indices.
 	/// </summary>
 	public class UhtUniqueReferenceCollection
 	{
+		private const int Shift = 2;
+		private const int Mask = 0x3;
+
 		/// <summary>
-		/// Collection use to quickly detect if a reference is already in the collection
+		/// Returns whether there is anything in this collection
 		/// </summary>
-		private HashSet<int> Uniques { get; } = new HashSet<int>();
+		public bool IsEmpty => Uniques.Count == 0;
 
 		/// <summary>
 		/// List of all unique reference keys.  Use UngetKey to get the object index and the flag.
 		/// </summary>
-		public List<int> References { get; } = new List<int>();
+		private HashSet<int> Uniques { get; } = new HashSet<int>();
 
 		/// <summary>
-		/// Return an encoded key that represents the object and registered flag. 
+		/// Return an encoded key that represents the object and type of reference.
 		/// If the object has the alternate object set (i.e. native interfaces), then
 		/// that object's index is used to generate the key.
 		/// </summary>
 		/// <param name="obj">Object being referenced</param>
-		/// <param name="registered">If true, then the API that ensures the object is registered is returned.</param>
+		/// <param name="type">The type of reference required for linking.</param>
 		/// <returns>Integer key value.</returns>
-		public static int GetKey(UhtObject obj, bool registered)
+		public static int GetKey(UhtObject obj, UhtSingletonType type)
 		{
-			return obj.AlternateObject != null ? GetKey(obj.AlternateObject, registered) : (obj.ObjectTypeIndex << 1) + (registered ? 1 : 0);
+			return obj.AlternateObject != null ? GetKey(obj.AlternateObject, type) : (obj.ObjectTypeIndex << Shift) | ((int)type);
 		}
 
 		/// <summary>
 		/// Given a key, return the object index and registered flag
 		/// </summary>
 		/// <param name="key">The key in question</param>
-		/// <param name="objectIndex">Index of the referenced object</param>
-		/// <param name="registered">True if referencing the registered API.</param>
-		public static void UngetKey(int key, out int objectIndex, out bool registered)
+		public static (int objectIndex, UhtSingletonType type) UngetKey(int key)
 		{
-			objectIndex = key >> 1;
-			registered = (key & 1) != 0;
+			return (key >> Shift, (UhtSingletonType)(key & Mask));
 		}
 
 		/// <summary>
 		/// Add the given object to the references
 		/// </summary>
 		/// <param name="obj">Object to be added</param>
-		/// <param name="registered">True if the registered API is being returned.</param>
-		public void Add(UhtObject? obj, bool registered)
+		/// <param name="type">Type of reference required</param>
+		public void Add(UhtObject? obj, UhtSingletonType type)
 		{
 			if (obj != null)
 			{
-				int key = GetKey(obj, registered);
-				if (Uniques.Add(key))
-				{
-					References.Add(key);
-				}
+				Uniques.Add(GetKey(obj, type));
 			}
 		}
 
 		/// <summary>
 		/// Return the collection of references sorted by the API string returned by the delegate.
 		/// </summary>
-		/// <param name="referenceStringDelegate">Delegate to invoke to return the requested object API string</param>
+		/// <param name="referenceToString">Function to invoke to return the requested object API string, taking the object index and the UhtSingletonType.</param>
 		/// <returns>Read only memory region of all the string.</returns>
-		public ReadOnlyMemory<string> GetSortedReferences(GetReferenceStringDelegate referenceStringDelegate)
+		public IReadOnlyList<string> GetSortedReferences(Func<int, UhtSingletonType, string?> referenceToString)
 		{
 			// Collect the unsorted array
-			string[] sorted = new string[References.Count];
-			for (int index = 0; index < References.Count; ++index)
+			List<string> sorted = new(Uniques.Count);
+			foreach(int key in Uniques)
 			{
-				int key = References[index];
-				UngetKey(key, out int objectIndex, out bool registered);
-				sorted[index] = referenceStringDelegate(objectIndex, registered);
+				(int objectIndex, UhtSingletonType type) = UngetKey(key);
+				if (referenceToString(objectIndex, type) is string decl)
+				{
+					sorted.Add(decl);
+				}
 			}
 
 			// Sort the array
-			Array.Sort(sorted, StringComparerUE.OrdinalIgnoreCase);
+			sorted.Sort(StringComparerUE.OrdinalIgnoreCase);
 
 			// Remove duplicates.  In some instances the different keys might return the same string.
 			// This removes those duplicates
-			if (References.Count > 1)
+			if (sorted.Count > 1)
 			{
 				int priorOut = 0;
-				for (int index = 1; index < sorted.Length; ++index)
+				for (int index = 1; index < sorted.Count; ++index)
 				{
 					if (sorted[index] != sorted[priorOut])
 					{
@@ -154,12 +145,13 @@ namespace EpicGames.UHT.Utils
 						sorted[priorOut] = sorted[index];
 					}
 				}
-				return sorted.AsMemory(0, priorOut + 1);
-			}
-			else
-			{
-				return sorted.AsMemory();
+
+				if (priorOut < sorted.Count - 1)
+				{
+					sorted.RemoveRange(priorOut + 1, sorted.Count - priorOut - 1);
+				}
 			}
+			return sorted;
 		}
 	}
 
@@ -205,19 +197,19 @@ namespace EpicGames.UHT.Utils
 		public HashSet<UhtHeaderFile> ReferencedHeaders { get; } = [];
 
 		/// <inheritdoc/>
-		public void AddCrossModuleReference(UhtObject? obj, bool registered)
+		public void AddCrossModuleReference(UhtObject? obj, UhtSingletonType type)
 		{
-			CrossModule.Add(obj, registered);
-			if (obj != null && obj is not UhtPackage && registered)
+			CrossModule.Add(obj, type);
+			if (obj != null && obj is not UhtPackage && type == UhtSingletonType.Registered)
 			{
 				ReferencedHeaders.Add(obj.HeaderFile);
 			}
 		}
 
 		/// <inheritdoc/>
-		public void AddDeclaration(UhtObject obj, bool registered)
+		public void AddDeclaration(UhtObject obj, UhtSingletonType type)
 		{
-			Declaration.Add(obj, registered);
+			Declaration.Add(obj, type);
 		}
 
 		/// <inheritdoc/>
diff --git a/Engine/Source/Programs/Shared/EpicGames.UHT/Utils/UhtSession.cs b/Engine/Source/Programs/Shared/EpicGames.UHT/Utils/UhtSession.cs
--- a/Engine/Source/Programs/Shared/EpicGames.UHT/Utils/UhtSession.cs
+++ b/Engine/Source/Programs/Shared/EpicGames.UHT/Utils/UhtSession.cs
@@ -729,6 +729,16 @@ namespace EpicGames.UHT.Utils
 		/// </summary>
 		public bool WarningsAsErrors { get; set; } = false;
 
+		/// <summary>
+		/// What format of code to generate for compiled-in UObject instances such as UClass instances
+		/// </summary>
+		public UhtCompiledInObjectFormat CompiledInObjectFormat => Manifest?.TargetSettings.CompiledInObjectFormat ?? UhtCompiledInObjectFormat.Default;
+
+		/// <summary>
+		/// Access to target-wide settings for code generation etc.
+		/// </summary>
+		public UHTTargetSettings? TargetSettings => Manifest?.TargetSettings;
+
 		/// <summary>
 		/// If true, include relative file paths in the log file
 		/// </summary>
@@ -863,6 +873,142 @@ namespace EpicGames.UHT.Utils
 			}
 		}
 
+		/// <summary>
+		/// After headers are parsed, returns the UClass class.
+		/// </summary>
+		public UhtClass UVerseClass
+		{
+			get
+			{
+				if (_uverseclass == null)
+				{
+					throw new UhtIceException("UVerseClass was not defined.");
+				}
+				return _uverseclass;
+			}
+		}
+
+		/// <summary>
+		/// After headers are parsed, returns the UScriptStruct class.
+		/// </summary>
+		public UhtClass UScriptStruct
+		{
+			get
+			{
+				if (_uscriptstruct == null)
+				{
+					throw new UhtIceException("UScriptStruct was not defined.");
+				}
+				return _uscriptstruct;
+			}
+		}
+
+		/// <summary>
+		/// After headers are parsed, returns the UScriptStruct class.
+		/// </summary>
+		public UhtClass UVerseStruct
+		{
+			get
+			{
+				if (_uversestruct == null)
+				{
+					throw new UhtIceException("UVerseStruct was not defined.");
+				}
+				return _uversestruct;
+			}
+		}
+
+		/// <summary>
+		/// After headers are parsed, returns the UFunction class.
+		/// </summary>
+		public UhtClass UFunction
+		{
+			get
+			{
+				if (_ufunction == null)
+				{
+					throw new UhtIceException("UFunction was not defined.");
+				}
+				return _ufunction;
+			}
+		}
+
+		/// <summary>
+		/// After headers are parsed, returns the UFunction class.
+		/// </summary>
+		public UhtClass UVerseFunction
+		{
+			get
+			{
+				if (_uversefunction == null)
+				{
+					throw new UhtIceException("UVerseFunction was not defined.");
+				}
+				return _uversefunction;
+			}
+		}
+
+		/// <summary>
+		/// After headers are parsed, returns the UDelegateFunction class.
+		/// </summary>
+		public UhtClass UDelegateFunction
+		{
+			get
+			{
+				if (_udelegatefunction == null)
+				{
+					throw new UhtIceException("UDelegateFunction was not defined.");
+				}
+				return _udelegatefunction;
+			}
+		}
+
+		/// <summary>
+		/// After headers are parsed, returns the USparseDelegateFunction class.
+		/// </summary>
+		public UhtClass USparseDelegateFunction
+		{
+			get
+			{
+				if (_usparsedelegatefunction == null)
+				{
+					throw new UhtIceException("USparseDelegateFunction was not defined.");
+				}
+				return _usparsedelegatefunction;
+			}
+		}
+
+		/// <summary>
+		/// After headers are parsed, returns the UEnum class.
+		/// </summary>
+		public UhtClass UEnum
+		{
+			get
+			{
+				if (_uenum == null)
+				{
+					throw new UhtIceException("UEnum was not defined.");
+				}
+				return _uenum;
+			}
+		}
+
+		/// <summary>
+		/// After headers are parsed, returns the UEnum class.
+		/// </summary>
+		public UhtClass UVerseEnum
+		{
+			get
+			{
+				if (_uverseenum == null)
+				{
+					throw new UhtIceException("UVerseEnum was not defined.");
+				}
+				return _uverseenum;
+			}
+		}
+
+
 		/// <summary>
 		/// After headers are parsed, returns the UInterface class.
 		/// </summary>
@@ -945,6 +1091,15 @@ namespace EpicGames.UHT.Utils
 		private UhtClass? _uobject = null;
 		private UhtClass? _upackage = null;
 		private UhtClass? _uclass = null;
+		private UhtClass? _uverseclass = null;
+		private UhtClass? _uscriptstruct = null;
+		private UhtClass? _uversestruct = null;
+		private UhtClass? _ufunction = null;
+		private UhtClass? _uversefunction = null;
+		private UhtClass? _udelegatefunction = null;
+		private UhtClass? _usparsedelegatefunction = null;
+		private UhtClass? _uenum = null;
+		private UhtClass? _uverseenum = null;
 		private UhtClass? _uinterface = null;
 		private UhtClass? _iinterface = null;
 		private readonly TypeCounter _typeCounter = new();
@@ -1150,6 +1305,19 @@ namespace EpicGames.UHT.Utils
 			return Tables!.StructDefaultValueTable.TryGet(name, out structDefaultValue);
 		}
 
+		/// <summary>
+		/// Returns whether UHT should generate a certain type of code for compiled-in object instances
+		/// </summary>
+		/// <param name="format">Format to test for</param>
+		/// <returns>Whether the given format is active - multiple formats may be active in certain development modes</returns>
+		public bool IsUsingCompiledInObjectFormat(UhtCompiledInObjectFormat format) => (CompiledInObjectFormat & format) != UhtCompiledInObjectFormat.None;
+
+		/// <summary>
+		/// Returns whether UHT is emitting code for multiple compiled-in object formats, and therefore such
+		/// code should be surrounded by preprocessor checks
+		/// </summary>
+		public bool IsUsingMultipleCompiledInObjectFormats => Int32.PopCount((int)CompiledInObjectFormat) > 1;
+
 		/// <summary>
 		/// Run UHT on the given manifest.  Use the bHasError property to see if process was successful.
 		/// </summary>
@@ -2219,6 +2387,15 @@ namespace EpicGames.UHT.Utils
 				_uobject = (UhtClass?)FindType(null, UhtFindOptions.SourceName | UhtFindOptions.Class, "UObject");
 				_upackage = (UhtClass?)FindType(null, UhtFindOptions.SourceName | UhtFindOptions.Class, "UPackage");
 				_uclass = (UhtClass?)FindType(null, UhtFindOptions.SourceName | UhtFindOptions.Class, "UClass");
+				_uverseclass = (UhtClass?)FindType(null, UhtFindOptions.SourceName | UhtFindOptions.Class, "UVerseClass");
+				_uscriptstruct = (UhtClass?)FindType(null, UhtFindOptions.SourceName | UhtFindOptions.Class, "UScriptStruct");
+				_uversestruct = (UhtClass?)FindType(null, UhtFindOptions.SourceName | UhtFindOptions.Class, "UVerseStruct");
+				_ufunction = (UhtClass?)FindType(null, UhtFindOptions.SourceName | UhtFindOptions.Class, "UFunction");
+				_uversefunction = (UhtClass?)FindType(null, UhtFindOptions.SourceName | UhtFindOptions.Class, "UVerseFunction");
+				_udelegatefunction = (UhtClass?)FindType(null, UhtFindOptions.SourceName | UhtFindOptions.Class, "UDelegateFunction");
+				_usparsedelegatefunction = (UhtClass?)FindType(null, UhtFindOptions.SourceName | UhtFindOptions.Class, "USparseDelegateFunction");
+				_uenum = (UhtClass?)FindType(null, UhtFindOptions.SourceName | UhtFindOptions.Class, "UEnum");
+				_uverseenum = (UhtClass?)FindType(null, UhtFindOptions.SourceName | UhtFindOptions.Class, "UVerseEnum");
 				_uinterface = (UhtClass?)FindType(null, UhtFindOptions.SourceName | UhtFindOptions.Class, "UInterface");
 				_iinterface = (UhtClass?)FindType(null, UhtFindOptions.SourceName | UhtFindOptions.Class, "IInterface");
 				EVerseNativeCallResult = (UhtEnum?)FindType(null, UhtFindOptions.SourceName | UhtFindOptions.Enum, "EVerseNativeCallResult");
diff --git a/Engine/Source/Programs/Shared/EpicGames.UHT/Utils/UhtStringBuilderExtensions.cs b/Engine/Source/Programs/Shared/EpicGames.UHT/Utils/UhtStringBuilderExtensions.cs
--- a/Engine/Source/Programs/Shared/EpicGames.UHT/Utils/UhtStringBuilderExtensions.cs
+++ b/Engine/Source/Programs/Shared/EpicGames.UHT/Utils/UhtStringBuilderExtensions.cs
@@ -3,8 +3,10 @@
 using System;
 using System.Collections.Generic;
 using System.Globalization;
+using System.Linq.Expressions;
 using System.Text;
 using EpicGames.Core;
+using EpicGames.UHT.Exporters.CodeGen;
 using EpicGames.UHT.Types;
 
 namespace EpicGames.UHT.Utils
@@ -66,6 +68,48 @@ namespace EpicGames.UHT.Utils
 			return builder;
 		}
 
+		/// <summary>
+		/// Append a list of items, transforming each into a StringBuilder append operation
+		/// </summary>
+		/// <typeparam name="T">Type of element</typeparam>
+		/// <param name="builder">StringBuilder to append to </param>
+		/// <param name="items">List of items</param>
+		/// <param name="transform">Action to transform each item into some kind of Append call on the stringbuilder</param>
+		/// <returns></returns>
+		public static StringBuilder AppendEach<T>(this StringBuilder builder, IReadOnlyCollection<T> items, Action<StringBuilder, T> transform)
+		{
+			foreach (T item in items)
+			{
+				transform(builder, item);
+			}
+			return builder;
+		}
+
+		/// <summary>
+		/// Optionally append a list of items with a prefix and suffix. If the list is empty, the prefix and suffix are also omitted
+		/// </summary>
+		/// <typeparam name="T"></typeparam>
+		/// <param name="builder"></param>
+		/// <param name="items"></param>
+		/// <param name="prefix"></param>
+		/// <param name="transform"></param>
+		/// <param name="suffix"></param>
+		/// <returns></returns>
+		public static StringBuilder AppendEach<T>(this StringBuilder builder, IReadOnlyCollection<T> items, Action<StringBuilder> prefix, Action<StringBuilder, T> transform, Action<StringBuilder> suffix)
+		{
+			if (items.Count == 0)
+			{
+				return builder;
+			}
+			prefix(builder);
+			foreach (T item in items)
+			{
+				transform(builder, item);
+			}
+			suffix(builder);
+			return builder;
+		}
+
 		/// <summary>
 		/// Append a name declaration to the builder
 		/// </summary>
@@ -123,6 +167,51 @@ namespace EpicGames.UHT.Utils
 			return builder.AppendNameDef(context.StaticsName, context.NamePrefix, name, nameSuffix);
 		}
 
+		/// <summary>
+		/// Append an array view construction expression for the given variable
+		/// </summary>
+		/// <param name="builder">Destination builder</param>
+		/// <param name="staticsName">Name of the statics section</param>
+		/// <param name="arrayName">The name of the arrray, pass null to insert an expression for an empty array view</param>
+		/// <param name="tabs">Number of tabs to start the line</param>
+		/// <param name="endl">Text to end the line</param>
+		/// <returns>Destination builder</returns>
+		public static StringBuilder AppendArrayView(this StringBuilder builder, string? staticsName, string? arrayName, int tabs, string? endl) 
+		{
+			if (arrayName is null)
+			{
+				return builder.AppendTabs(tabs).Append("{}").Append(endl);
+			}
+			if (staticsName is not null) 
+			{
+				return builder.AppendTabs(tabs).Append("MakeArrayView(").Append(staticsName).Append("::").Append(arrayName).Append(')').Append(endl);
+			}
+			return builder.AppendTabs(tabs).Append("MakeArrayView(").Append(arrayName).Append(')').Append(endl);
+		}
+
+		/// <summary>
+		/// Append an array view construction expression for metadata if the metadata collection is not empty.
+		/// </summary>
+		/// <param name="builder">Destination builder</param>
+		/// <param name="metaData">MetaData object to check for emptiness</param>
+		/// <param name="staticsName">Name of the statics section</param>
+		/// <param name="arrayName">The name of the arrray</param>
+		/// <param name="tabs">Number of tabs to start the line</param>
+		/// <param name="endl">Text to end the line</param>
+		/// <returns>Destination builder</returns>
+		public static StringBuilder AppendMetaDataArrayView(this StringBuilder builder, UhtMetaData metaData, string? staticsName, string arrayName, int tabs, string endl) 
+		{
+			if (metaData.IsEmpty())
+			{
+				return builder.AppendTabs(tabs).Append("{}").Append(endl);
+			}
+			if (!String.IsNullOrEmpty(staticsName))
+			{
+				return builder.AppendTabs(tabs).Append("MakeArrayView(").Append(staticsName).Append("::").Append(arrayName).Append(')').Append(endl);
+			}
+			return builder.AppendTabs(tabs).Append("MakeArrayView(").Append(arrayName).Append(')').Append(endl);
+		}
+
 		/// <summary>
 		/// Append the meta data parameters.  This is intended to be used as arguments to a function call.
 		/// </summary>
@@ -180,57 +269,67 @@ namespace EpicGames.UHT.Utils
 		}
 
 		/// <summary>
-		/// Append the meta data declaration.
+		/// Append the metadata definition, i.e. a definition for a array of structures containing name/value string pairs.
 		/// </summary>
 		/// <param name="builder">Destination builder</param>
 		/// <param name="type">Source type containing the meta data</param>
+		/// <param name="settings">Access to necessary info for code generation such as code format</param>
 		/// <param name="namePrefix">Name prefix</param>
 		/// <param name="name">Name</param>
 		/// <param name="nameSuffix">Optional name suffix</param>
 		/// <param name="metaNameSuffix">Optional meta data name suffix</param>
 		/// <param name="tabs">Number of tabs to indent</param>
 		/// <returns>Destination builder</returns>
-		public static StringBuilder AppendMetaDataDecl(this StringBuilder builder, UhtType type, string? namePrefix, string name, string? nameSuffix, string? metaNameSuffix, int tabs)
+		public static StringBuilder AppendMetaDataDefinition(this StringBuilder builder, UhtType type, IUhtSessionSettingsProvider settings, string? namePrefix, string name, string? nameSuffix, string? metaNameSuffix, int tabs)
 		{
 			if (!type.MetaData.IsEmpty())
 			{
 				bool isPartOfEngine = type.Module.IsPartOfEngine;
 				List<KeyValuePair<string, string>> sortedMetaData = type.MetaData.GetSorted();
-				builder.AppendTabs(tabs).Append("static constexpr UECodeGen_Private::FMetaDataPairParam ").AppendNameDecl(namePrefix, name, nameSuffix).Append(metaNameSuffix).Append("[] = {\r\n");
-				foreach (KeyValuePair<string, string> kvp in sortedMetaData)
+				if (settings.IsUsingCompiledInObjectFormat(UhtCompiledInObjectFormat.ConstInit))
 				{
-					bool restricted = !isPartOfEngine && s_hiddenMetaDataNames.Contains(kvp.Key);
-					if (restricted)
+					using UhtConditionalMacroBlock block = new(builder, "UE_WITH_CONSTINIT_UOBJECT", settings.IsUsingMultipleCompiledInObjectFormats);
+					builder.AppendTabs(tabs).Append("static constexpr UE::CodeGen::ConstInit::FMetaData ").AppendNameDecl(namePrefix, name, nameSuffix).Append(metaNameSuffix).Append("[] = {\r\n");
+					foreach (KeyValuePair<string, string> kvp in sortedMetaData)
 					{
-						builder.Append("#if !UE_BUILD_SHIPPING\r\n");
+						bool restricted = !isPartOfEngine && s_hiddenMetaDataNames.Contains(kvp.Key);
+						if (restricted)
+						{
+							builder.Append("#if !UE_BUILD_SHIPPING\r\n");
+						}
+						builder.AppendTabs(tabs + 1).Append("{ UTF8TEXT(").AppendUTF8LiteralString(kvp.Key).Append("), UTF8TEXT(").AppendUTF8LiteralString(kvp.Value).Append(") },\r\n");
+						if (restricted)
+						{
+							builder.Append("#endif\r\n");
+						}
 					}
-					builder.AppendTabs(tabs + 1).Append("{ ").AppendUTF8LiteralString(kvp.Key).Append(", ").AppendUTF8LiteralString(kvp.Value).Append(" },\r\n");
-					if (restricted)
+
+					builder.AppendTabs(tabs).Append("};\r\n");
+				}
+				if (settings.IsUsingCompiledInObjectFormat(UhtCompiledInObjectFormat.Params))
+				{
+					using UhtConditionalMacroBlock block = new(builder, "!UE_WITH_CONSTINIT_UOBJECT", settings.IsUsingMultipleCompiledInObjectFormats);
+					builder.AppendTabs(tabs).Append("static constexpr UECodeGen_Private::FMetaDataPairParam ").AppendNameDecl(namePrefix, name, nameSuffix).Append(metaNameSuffix).Append("[] = {\r\n");
+					foreach (KeyValuePair<string, string> kvp in sortedMetaData)
 					{
-						builder.Append("#endif\r\n");
+						bool restricted = !isPartOfEngine && s_hiddenMetaDataNames.Contains(kvp.Key);
+						if (restricted)
+						{
+							builder.Append("#if !UE_BUILD_SHIPPING\r\n");
+						}
+						builder.AppendTabs(tabs + 1).Append("{ ").AppendUTF8LiteralString(kvp.Key).Append(", ").AppendUTF8LiteralString(kvp.Value).Append(" },\r\n");
+						if (restricted)
+						{
+							builder.Append("#endif\r\n");
+						}
 					}
-				}
 
-				builder.AppendTabs(tabs).Append("};\r\n");
+					builder.AppendTabs(tabs).Append("};\r\n");
+				}
 			}
 			return builder;
 		}
 
-		/// <summary>
-		/// Append the meta data declaration
-		/// </summary>
-		/// <param name="builder">Destination builder</param>
-		/// <param name="property">Source type containing the meta data</param>
-		/// <param name="context">Property context used to get the statics name and name prefix</param>
-		/// <param name="name">Name</param>
-		/// <param name="nameSuffix">Optional name suffix</param>
-		/// <param name="tabs">Number of tabs to indent</param>
-		/// <returns>Destination builder</returns>
-		public static StringBuilder AppendMetaDataDecl(this StringBuilder builder, UhtProperty property, IUhtPropertyMemberContext context, string name, string nameSuffix, int tabs)
-		{
-			return property.AppendMetaDataDecl(builder, context, name, nameSuffix, tabs);
-		}
-
 		/// <summary>
 		/// Append the given text as a UTF8 encoded string
 		/// </summary>
@@ -265,7 +364,7 @@ namespace EpicGames.UHT.Utils
 		{
 			if (text == null)
 			{
-				builder.Append("");
+				builder.Append("\"\"");
 			}
 			else
 			{
@@ -436,17 +535,7 @@ namespace EpicGames.UHT.Utils
 		/// <exception cref="NotImplementedException">Class has an unexpected class type</exception>
 		public static StringBuilder AppendClassSourceNameOrInterfaceName(this StringBuilder builder, UhtClass classObj)
 		{
-			switch (classObj.ClassType)
-			{
-				case UhtClassType.Class:
-				case UhtClassType.NativeInterface:
-				case UhtClassType.VModule:
-					return builder.Append(classObj.SourceName);
-				case UhtClassType.Interface:
-					return builder.Append('I').Append(classObj.SourceName[1..]);
-				default:
-					throw new NotImplementedException();
-			}
+			return builder.Append(classObj.Namespace.FullSourceName).Append(classObj.NativeFunctionCallName);
 		}
 
 		/// <summary>
diff --git a/Engine/Source/Programs/UnrealBuildTool/Configuration/ReadOnlyTargetRules.cs b/Engine/Source/Programs/UnrealBuildTool/Configuration/ReadOnlyTargetRules.cs
--- a/Engine/Source/Programs/UnrealBuildTool/Configuration/ReadOnlyTargetRules.cs
+++ b/Engine/Source/Programs/UnrealBuildTool/Configuration/ReadOnlyTargetRules.cs
@@ -185,6 +185,8 @@ namespace UnrealBuildTool
 
 		public bool bCompileAgainstCoreUObject => Inner.bCompileAgainstCoreUObject;
 
+		public bool bEnableConstInitUObject => Inner.bEnableConstInitUObject;
+
 		public bool bCompileAgainstApplicationCore => Inner.bCompileAgainstApplicationCore;
 
 		public bool bEnableTrace => Inner.bEnableTrace;
diff --git a/Engine/Source/Programs/UnrealBuildTool/Configuration/TargetRules.cs b/Engine/Source/Programs/UnrealBuildTool/Configuration/TargetRules.cs
--- a/Engine/Source/Programs/UnrealBuildTool/Configuration/TargetRules.cs
+++ b/Engine/Source/Programs/UnrealBuildTool/Configuration/TargetRules.cs
@@ -1211,6 +1211,19 @@ namespace UnrealBuildTool
 		}
 		private bool bCompileAgainstCoreUObjectPrivate = true;
 
+		/// <summary>
+		/// Enable compilation of UClass and related reflection objects as statically initialized global variables.
+		/// </summary>
+		[RequiresUniqueBuildEnvironment]
+		[CommandLine("-ConstInitUObject")]
+		[CommandLine("-NoConstInitUObject", Value="false")]
+		public virtual bool bEnableConstInitUObject 
+		{
+			get => bEnableConstInitUObjectPrivate;
+			set => bEnableConstInitUObjectPrivate = value;
+		}
+		private bool bEnableConstInitUObjectPrivate;
+
 		/// <summary>
 		/// Enabled for builds that need to initialize the ApplicationCore module. Command line utilities do not normally need this.
 		/// </summary>
@@ -1457,6 +1470,8 @@ namespace UnrealBuildTool
 		/// </summary>
 		[RequiresUniqueBuildEnvironment]
 		[XmlConfigFile(Category = "BuildConfiguration")]
+		[CommandLine("-LiveCoding")]
+		[CommandLine("-NoLiveCoding", Value = "false")]
 		public bool bWithLiveCoding
 		{
 			get => bWithLiveCodingPrivate ?? (Platform == UnrealTargetPlatform.Win64 && Architecture.bIsX64 && Configuration != UnrealTargetConfiguration.Shipping && Configuration != UnrealTargetConfiguration.Test && Type != TargetType.Program);
diff --git a/Engine/Source/Programs/UnrealBuildTool/Configuration/UEBuildTarget.cs b/Engine/Source/Programs/UnrealBuildTool/Configuration/UEBuildTarget.cs
--- a/Engine/Source/Programs/UnrealBuildTool/Configuration/UEBuildTarget.cs
+++ b/Engine/Source/Programs/UnrealBuildTool/Configuration/UEBuildTarget.cs
@@ -1311,6 +1311,16 @@ namespace UnrealBuildTool
 				}
 			}
 
+			if (RulesObject.bEnableConstInitUObject && RulesObject.bWithLiveCoding)
+			{
+				throw new BuildException("bEnableConstInitUObject and bWithLiveCoding cannot currently be set simultaneously");
+			}
+
+			if (RulesObject.bEnableConstInitUObject && RulesObject.LinkType != TargetLinkType.Monolithic)
+			{
+				throw new BuildException("bEnableConstInitUObject can currently only be set on monolithic targets");
+			}
+
 			// If we're using the shared build environment, make sure all the settings are valid
 			if (RulesObject.BuildEnvironment == TargetBuildEnvironment.Shared)
 			{
@@ -2658,6 +2668,14 @@ namespace UnrealBuildTool
 				};
 			}
 
+			if (Rules.bEnableConstInitUObject)
+			{
+				Makefile.UHTTargetSettings.CompiledInObjectFormat = UhtCompiledInObjectFormat.ConstInit;
+				// Generate max name/value pairs for UENUMs in the compiled-in data
+				Makefile.UHTTargetSettings.GenerateEnumMaxValues = true;
+			}
+			Logger.LogInformation($"UHT compiled-in object format {Makefile.UHTTargetSettings.CompiledInObjectFormat}");
+
 			// UHT mode uses this to create the makefile at least to the point where UHT would have valid manifest information
 			if (bInitOnly)
 			{
@@ -5824,6 +5842,15 @@ namespace UnrealBuildTool
 				GlobalCompileEnvironment.Definitions.Add("WITH_COREUOBJECT=0");
 			}
 
+			if (Rules.bEnableConstInitUObject)
+			{
+				GlobalCompileEnvironment.Definitions.Add("UE_WITH_CONSTINIT_UOBJECT=1");
+			}
+			else
+			{
+				GlobalCompileEnvironment.Definitions.Add("UE_WITH_CONSTINIT_UOBJECT=0");
+			}
+
 			if (Rules.bEnableTrace)
 			{
 				GlobalCompileEnvironment.Definitions.Add("UE_TRACE_ENABLED=1");
diff --git a/Engine/Source/Programs/UnrealBuildTool/System/ExternalExecution.cs b/Engine/Source/Programs/UnrealBuildTool/System/ExternalExecution.cs
--- a/Engine/Source/Programs/UnrealBuildTool/System/ExternalExecution.cs
+++ b/Engine/Source/Programs/UnrealBuildTool/System/ExternalExecution.cs
@@ -1083,6 +1083,17 @@ namespace UnrealBuildTool
 		{
 			return FileReference.Combine(Makefile.ProjectIntermediateDirectoryNoArch, TargetName + ".uhtmanifest");
 		}
+		
+		/// <summary>
+		/// Return the file reference for an intermediate file containing the UHT settings last used to build a given target
+		/// </summary>
+		/// <param name="Makefile">Input makefile</param>
+		/// <param name="TargetName">Name of the target</param>
+		/// <returns>Manifest file name</returns>
+		public static FileReference GetUHTSettingsFileName(TargetMakefile Makefile, string TargetName)
+		{
+			return FileReference.Combine(Makefile.ProjectIntermediateDirectoryNoArch, TargetName + ".uhtsettings");
+		}
 
 		/// <summary>
 		/// Return the file reference for the UHT deps file
@@ -1155,6 +1166,7 @@ namespace UnrealBuildTool
 				IsGameTarget = Makefile.TargetType != TargetType.Program,
 				RootLocalPath = Unreal.RootDirectory.FullName,
 				ExternalDependenciesFile = DepsFileName.FullName,
+				TargetSettings = Makefile.UHTTargetSettings,
 				Modules = Modules,
 				UhtPlugins = UhtPlugins,
 			};
@@ -1215,13 +1227,34 @@ namespace UnrealBuildTool
 					}
 				}
 
-				// ensure the headers are up to date
 				bool bUHTNeedsToRun = false;
+
+				// Check explicit settings passed to UHT - this info is also saved in the manifest but we load just the settings here
+				FileReference UHTSettingsFileName = GetUHTSettingsFileName(Makefile, TargetName);
+				if (!FileReference.Exists(UHTSettingsFileName))
+				{
+					bUHTNeedsToRun = true;
+				}
+				else
+				{
+					using (BinaryArchiveReader reader = new(UHTSettingsFileName))
+					{
+						// Note that even if default/omitted values cause the settings to compare equal here, the UHT module version check below 
+						// will ensure that we rerun UHT when new properties are added to UHTTargetSettings
+						UHTTargetSettings settings = new(reader);
+						if (settings != Makefile.UHTTargetSettings)
+						{
+							bUHTNeedsToRun = true;
+						}
+					}
+				}
+
+				// ensure the headers are up to date
 				if (BuildConfiguration.bForceHeaderGeneration)
 				{
 					bUHTNeedsToRun = true;
 				}
-				else if (AreGeneratedCodeFilesOutOfDate(Makefile.UObjectModules, CompositeTimestamp, Logger))
+				else if (!bUHTNeedsToRun && AreGeneratedCodeFilesOutOfDate(Makefile.UObjectModules, CompositeTimestamp, Logger))
 				{
 					bUHTNeedsToRun = true;
 				}
@@ -1328,7 +1361,12 @@ namespace UnrealBuildTool
 					// Update the tool info file
 					DirectoryReference.CreateDirectory(ToolInfoFile.Directory);
 					FileReference.WriteAllText(ToolInfoFile, UhtAssemblyPath);
-
+					// Update saved UHT settings
+					using (BinaryArchiveWriter writer = new(UHTSettingsFileName))
+					{
+						Makefile.UHTTargetSettings.Write(writer);
+					}
+					
 					// Now that UHT has successfully finished generating code, we need to update all cached FileItems in case their last write time has changed.
 					// Otherwise UBT might not detect changes UHT made.
 					using (GlobalTracer.Instance.BuildSpan("ExternalExecution.ResetCachedHeaderInfo()").StartActive())
diff --git a/Engine/Source/Programs/UnrealBuildTool/System/TargetMakefile.cs b/Engine/Source/Programs/UnrealBuildTool/System/TargetMakefile.cs
--- a/Engine/Source/Programs/UnrealBuildTool/System/TargetMakefile.cs
+++ b/Engine/Source/Programs/UnrealBuildTool/System/TargetMakefile.cs
@@ -178,7 +178,7 @@ namespace UnrealBuildTool
 		/// <summary>
 		/// The version number to write
 		/// </summary>
-		public const int CurrentVersion = 36;
+		public const int CurrentVersion = 37;
 
 		/// <summary>
 		/// The time at which the makefile was created
@@ -336,6 +336,11 @@ namespace UnrealBuildTool
 		/// </summary>
 		public List<UHTModuleHeaderInfo> UObjectModuleHeaders = new List<UHTModuleHeaderInfo>();
 
+		/// <summary>
+		/// Target-wide settings for UHT that should not be in UHTModuleInfo
+		/// </summary>
+		public UHTTargetSettings UHTTargetSettings = new();
+
 #if __VPROJECT_AVAILABLE__
 		/// <summary>
 		/// All build modules containing Verse code
@@ -507,6 +512,7 @@ namespace UnrealBuildTool
 			CandidatesForWorkingSet = Reader.ReadHashSet(() => Reader.ReadFileItem())!;
 			UObjectModules = Reader.ReadList(() => new UHTModuleInfo(Reader))!;
 			UObjectModuleHeaders = Reader.ReadList(() => new UHTModuleHeaderInfo(Reader))!;
+			UHTTargetSettings = new UHTTargetSettings(Reader);
 #if __VPROJECT_AVAILABLE__
 			VNIModules = Reader.ReadList(() => new VNIModuleInfo(Reader))!;
 			bUseVerseBPVM = Reader.ReadBool();
@@ -581,6 +587,7 @@ namespace UnrealBuildTool
 			Writer.WriteHashSet(CandidatesForWorkingSet, x => Writer.WriteFileItem(x));
 			Writer.WriteList(UObjectModules, e => e.Write(Writer));
 			Writer.WriteList(UObjectModuleHeaders, x => x.Write(Writer));
+			UHTTargetSettings.Write(Writer);
 #if __VPROJECT_AVAILABLE__
 			Writer.WriteList(VNIModules, e => e.Write(Writer));
 			Writer.WriteBool(bUseVerseBPVM);
