Patching verse path information through the the UVerseClasses representing interfaces

diff --git a/Engine/Source/Programs/Shared/EpicGames.UHT/Exporters/CodeGen/UhtHeaderCodeGeneratorCppFile.cs b/Engine/Source/Programs/Shared/EpicGames.UHT/Exporters/CodeGen/UhtHeaderCodeGeneratorCppFile.cs
--- a/Engine/Source/Programs/Shared/EpicGames.UHT/Exporters/CodeGen/UhtHeaderCodeGeneratorCppFile.cs
+++ b/Engine/Source/Programs/Shared/EpicGames.UHT/Exporters/CodeGen/UhtHeaderCodeGeneratorCppFile.cs
@@ -2941,7 +2941,7 @@ namespace EpicGames.UHT.Exporters.CodeGen
 					{
 						EVerseClassFlags verseClassFlags = classObj.ClassType == UhtClassType.VModule ? EVerseClassFlags.Module : EVerseClassFlags.None;
 						builder.Append("\t},\r\n");
-						if (classObj.ClassType == UhtClassType.Class || classObj.ClassType == UhtClassType.VModule)
+						if (classObj.ClassType == UhtClassType.Class || classObj.ClassType == UhtClassType.VModule || classObj.ClassType == UhtClassType.Interface)
 						{
 							builder.Append("\t\"").AppendVerseScopeAndName(classObj, UhtVerseNameMode.PackageRelative).Append("\",\r\n");
 						}
