EpicGames.UBA: Don't pass private collections through events

#rnx

diff --git a/Engine/Source/Programs/Shared/EpicGames.UBA/Utils.cs b/Engine/Source/Programs/Shared/EpicGames.UBA/Utils.cs
--- a/Engine/Source/Programs/Shared/EpicGames.UBA/Utils.cs
+++ b/Engine/Source/Programs/Shared/EpicGames.UBA/Utils.cs
@@ -24,21 +24,21 @@ namespace EpicGames.UBA
 		/// <summary>
 		/// Paths that are not allowed to be transferred over the network for UBA remote agents.
 		/// </summary>
-		/// <returns>enumerable of disallowed paths</returns>
-		public static IEnumerable<string> DisallowedPaths => s_disallowedPaths.Distinct().Order();
+		/// <returns>IEnumerable of disallowed paths</returns>
+		public static IEnumerable<string> DisallowedPaths => s_disallowedPaths.Order().Distinct();
 		static readonly ConcurrentBag<string> s_disallowedPaths = [];
 
 		/// <summary>
 		/// Mapping of binary paths for cross architecture host binaries, to allow for using helpers of a different architecture.
 		/// </summary>
-		/// <returns>dictionary of binary mappings, where the key is the binary for the current host architecture</returns>
-		public static IReadOnlyDictionary<string, string> CrossArchitecturePaths => s_crossArchitecturePaths;
+		/// <returns>IEnumerable of binary mappings, where the key is the binary for the current host architecture</returns>
+		public static IEnumerable<KeyValuePair<string, string>> CrossArchitecturePaths => s_crossArchitecturePaths.OrderBy(x => x.Key);
 		static readonly ConcurrentDictionary<string, string> s_crossArchitecturePaths = [];
 
 		/// <summary>
 		/// Mapping of a folder path to a single hash, to allow for hashing an entire folder so each individual file does not need to be processed.
 		/// </summary>
-		public static IReadOnlyDictionary<string, string> PathHashes => s_pathHashes;
+		public static IEnumerable<KeyValuePair<string, string>> PathHashes => s_pathHashes.OrderBy(x => x.Key);
 		static readonly ConcurrentDictionary<string, string> s_pathHashes = [];
 
 		/// <summary>
@@ -54,7 +54,7 @@ namespace EpicGames.UBA
 					s_disallowedPaths.Add(path);
 				}
 			}
-			DisallowedPathRegistered?.Invoke(s_disallowedPaths, new(paths));
+			DisallowedPathRegistered?.Invoke(DisallowedPaths, new(paths));
 		}
 
 		/// <summary>
@@ -66,7 +66,7 @@ namespace EpicGames.UBA
 		{
 			if (s_crossArchitecturePaths.TryAdd(path, crossPath))
 			{
-				CrossArchitecturePathRegistered?.Invoke(s_crossArchitecturePaths, new(path, crossPath));
+				CrossArchitecturePathRegistered?.Invoke(CrossArchitecturePaths, new(path, crossPath));
 			}
 		}
 
@@ -79,7 +79,7 @@ namespace EpicGames.UBA
 		{
 			if (s_pathHashes.TryAdd(path, hash))
 			{
-				PathHashRegistered?.Invoke(s_pathHashes, new(path, hash));
+				PathHashRegistered?.Invoke(PathHashes, new(path, hash));
 			}
 		}
 
@@ -88,21 +88,21 @@ namespace EpicGames.UBA
 		/// </summary>
 		/// <param name="sender">collection that is being changed</param>
 		/// <param name="e">event args containing which paths were added</param>
-		public delegate void DisallowedPathRegisteredEventHandler(IReadOnlyCollection<string> sender, DisallowedPathRegisteredEventArgs e);
+		public delegate void DisallowedPathRegisteredEventHandler(IEnumerable<string> sender, DisallowedPathRegisteredEventArgs e);
 
 		/// <summary>
 		/// Delegate for registering a cross architecture path
 		/// </summary>
 		/// <param name="sender">collection that is being changed</param>
 		/// <param name="e">event args containing which path was added</param>
-		public delegate void CrossArchitecturePathRegisteredEventHandler(IReadOnlyDictionary<string, string> sender, CrossArchitecturePathRegisteredEventArgs e);
+		public delegate void CrossArchitecturePathRegisteredEventHandler(IEnumerable<KeyValuePair<string, string>> sender, CrossArchitecturePathRegisteredEventArgs e);
 
 		/// <summary>
 		/// Delegate for registering a path hash
 		/// </summary>
 		/// <param name="sender">collection that is being changed</param>
 		/// <param name="e">event args containing which path was added</param>
-		public delegate void PathHashRegisteredEventHandler(IReadOnlyDictionary<string, string> sender, PathHashRegisteredEventArgs e);
+		public delegate void PathHashRegisteredEventHandler(IEnumerable<KeyValuePair<string, string>> sender, PathHashRegisteredEventArgs e);
 
 		/// <summary>
 		/// Remote disallowed path registered event handler
