EpicGames.UBA: ConcurrentDictionary enumerator is not a snapshot

#rnx

diff --git a/Engine/Source/Programs/Shared/EpicGames.UBA/Utils.cs b/Engine/Source/Programs/Shared/EpicGames.UBA/Utils.cs
--- a/Engine/Source/Programs/Shared/EpicGames.UBA/Utils.cs
+++ b/Engine/Source/Programs/Shared/EpicGames.UBA/Utils.cs
@@ -32,13 +32,13 @@ namespace EpicGames.UBA
 		/// Mapping of binary paths for cross architecture host binaries, to allow for using helpers of a different architecture.
 		/// </summary>
 		/// <returns>IEnumerable of binary mappings, where the key is the binary for the current host architecture</returns>
-		public static IEnumerable<KeyValuePair<string, string>> CrossArchitecturePaths => s_crossArchitecturePaths.OrderBy(x => x.Key);
+		public static IEnumerable<KeyValuePair<string, string>> CrossArchitecturePaths => s_crossArchitecturePaths.ToArray().OrderBy(x => x.Key);
 		static readonly ConcurrentDictionary<string, string> s_crossArchitecturePaths = [];
 
 		/// <summary>
 		/// Mapping of a folder path to a single hash, to allow for hashing an entire folder so each individual file does not need to be processed.
 		/// </summary>
-		public static IEnumerable<KeyValuePair<string, string>> PathHashes => s_pathHashes.OrderBy(x => x.Key);
+		public static IEnumerable<KeyValuePair<string, string>> PathHashes => s_pathHashes.ToArray().OrderBy(x => x.Key);
 		static readonly ConcurrentDictionary<string, string> s_pathHashes = [];
 
 		/// <summary>
