Updating the Horde retry warnings to no longer create Build Health issues.

diff --git a/Engine/Source/Programs/Shared/EpicGames.Horde/HordeHttpMessageHandler.cs b/Engine/Source/Programs/Shared/EpicGames.Horde/HordeHttpMessageHandler.cs
--- a/Engine/Source/Programs/Shared/EpicGames.Horde/HordeHttpMessageHandler.cs
+++ b/Engine/Source/Programs/Shared/EpicGames.Horde/HordeHttpMessageHandler.cs
@@ -47,7 +47,7 @@ namespace EpicGames.Horde
 			// Wait 30 seconds for operations to timeout
 			Task OnTimeoutAsync(Context context, TimeSpan timespan, Task timeoutTask)
 			{
-				logger.LogWarning(KnownLogEvents.Systemic_Horde_Http, "{Method} {Url} timed out after {Time}s.", request.Method, request.RedactedRequestUri(), (int)timespan.TotalSeconds);
+				logger.LogWarning(KnownLogEvents.Horde_BuildHealth_Ignore, "{Method} {Url} timed out after {Time}s.", request.Method, request.RedactedRequestUri(), (int)timespan.TotalSeconds);
 				return Task.CompletedTask;
 			}
 
@@ -56,7 +56,7 @@ namespace EpicGames.Horde
 			// Retry twice after a timeout
 			void OnRetry(Exception ex, TimeSpan timespan)
 			{
-				logger.LogWarning(KnownLogEvents.Systemic_Horde_Http, ex, "{Method} {Url} retrying after {Time}s.", request.Method, request.RedactedRequestUri(), timespan.TotalSeconds);
+				logger.LogWarning(KnownLogEvents.Horde_BuildHealth_Ignore, ex, "{Method} {Url} retrying after {Time}s.", request.Method, request.RedactedRequestUri(), timespan.TotalSeconds);
 			}
 
 			TimeSpan[] retryTimes = new[] { TimeSpan.FromSeconds(5.0), TimeSpan.FromSeconds(10.0), TimeSpan.FromSeconds(15.0) };
