Enable clang unreachable warnings

diff --git a/Engine/Source/Programs/Shared/UnrealBuildTool.Tests/CppCompileWarningsTests.cs b/Engine/Source/Programs/Shared/UnrealBuildTool.Tests/CppCompileWarningsTests.cs
--- a/Engine/Source/Programs/Shared/UnrealBuildTool.Tests/CppCompileWarningsTests.cs
+++ b/Engine/Source/Programs/Shared/UnrealBuildTool.Tests/CppCompileWarningsTests.cs
@@ -61,6 +61,7 @@ namespace UnrealBuildTool.Tests
 
 			Assert.AreEqual(WarningLevel.Warning, defaultBuildSettings.ShadowVariableWarningLevel);
 			Assert.AreEqual(WarningLevel.Off, defaultBuildSettings.UndefinedIdentifierWarningLevel);
+			Assert.AreEqual(WarningLevel.Off, defaultBuildSettings.UnreachableCodeWarningLevel);
 			Assert.AreEqual(WarningLevel.Off, defaultBuildSettings.UnsafeTypeCastWarningLevel);
 			Assert.AreEqual(WarningLevel.Off, defaultBuildSettings.SwitchUnhandledEnumeratorWarningLevel);
 		}
@@ -71,6 +72,7 @@ namespace UnrealBuildTool.Tests
 			CppCompileWarnings defaultBuildSettings = new CppCompileWarnings();
 			defaultBuildSettings.ShadowVariableWarningLevel = WarningLevel.Error;
 			defaultBuildSettings.UndefinedIdentifierWarningLevel = WarningLevel.Warning;
+			defaultBuildSettings.UnreachableCodeWarningLevel = WarningLevel.Warning;
 			defaultBuildSettings.UnsafeTypeCastWarningLevel = WarningLevel.Error;
 			defaultBuildSettings.SwitchUnhandledEnumeratorWarningLevel = WarningLevel.Warning;
 
@@ -81,6 +83,7 @@ namespace UnrealBuildTool.Tests
 			{
 				Assert.AreEqual(defaultBuildSettings.ShadowVariableWarningLevel, cloned.ShadowVariableWarningLevel);
 				Assert.AreEqual(defaultBuildSettings.UndefinedIdentifierWarningLevel, cloned.UndefinedIdentifierWarningLevel);
+				Assert.AreEqual(defaultBuildSettings.UnreachableCodeWarningLevel, cloned.UnreachableCodeWarningLevel);
 				Assert.AreEqual(defaultBuildSettings.UnsafeTypeCastWarningLevel, cloned.UnsafeTypeCastWarningLevel);
 				Assert.AreEqual(defaultBuildSettings.SwitchUnhandledEnumeratorWarningLevel, cloned.SwitchUnhandledEnumeratorWarningLevel);
 			}
@@ -89,11 +92,13 @@ namespace UnrealBuildTool.Tests
 			{
 				cloned.ShadowVariableWarningLevel = WarningLevel.Warning;
 				cloned.UndefinedIdentifierWarningLevel = WarningLevel.Error;
+				cloned.UnreachableCodeWarningLevel = WarningLevel.Error;
 				cloned.UnsafeTypeCastWarningLevel = WarningLevel.Warning;
 				cloned.SwitchUnhandledEnumeratorWarningLevel = WarningLevel.Error;
 
 				Assert.AreNotEqual(defaultBuildSettings.ShadowVariableWarningLevel, cloned.ShadowVariableWarningLevel);
 				Assert.AreNotEqual(defaultBuildSettings.UndefinedIdentifierWarningLevel, cloned.UndefinedIdentifierWarningLevel);
+				Assert.AreNotEqual(defaultBuildSettings.UnreachableCodeWarningLevel, cloned.UnreachableCodeWarningLevel);
 				Assert.AreNotEqual(defaultBuildSettings.UnsafeTypeCastWarningLevel, cloned.UnsafeTypeCastWarningLevel);
 				Assert.AreNotEqual(defaultBuildSettings.SwitchUnhandledEnumeratorWarningLevel, cloned.SwitchUnhandledEnumeratorWarningLevel);
 			}
@@ -134,19 +139,22 @@ namespace UnrealBuildTool.Tests
 			// Assert invariants
 			{
 				Assert.AreEqual(WarningLevel.Default, newWarnings.ShadowVariableWarningLevel);
-				Assert.AreEqual(WarningLevel.Default, newWarnings.UnsafeTypeCastWarningLevel);
 				Assert.AreEqual(WarningLevel.Default, newWarnings.UndefinedIdentifierWarningLevel);
+				Assert.AreEqual(WarningLevel.Default, newWarnings.UnreachableCodeWarningLevel);
+				Assert.AreEqual(WarningLevel.Default, newWarnings.UnsafeTypeCastWarningLevel);
 			}
 
 			// Assert new parent values
 			{
 				parentWarnings.ShadowVariableWarningLevel = WarningLevel.Error;
 				parentWarnings.UndefinedIdentifierWarningLevel = WarningLevel.Error;
+				parentWarnings.UnreachableCodeWarningLevel = WarningLevel.Error;
 				parentWarnings.UnsafeTypeCastWarningLevel = WarningLevel.Error;
 
 				Assert.AreEqual(WarningLevel.Error, newWarnings.ShadowVariableWarningLevel);
-				Assert.AreEqual(WarningLevel.Error, newWarnings.UnsafeTypeCastWarningLevel);
 				Assert.AreEqual(WarningLevel.Error, newWarnings.UndefinedIdentifierWarningLevel);
+				Assert.AreEqual(WarningLevel.Error, newWarnings.UnreachableCodeWarningLevel);
+				Assert.AreEqual(WarningLevel.Error, newWarnings.UnsafeTypeCastWarningLevel);
 			}
 		}
 
@@ -162,11 +170,13 @@ namespace UnrealBuildTool.Tests
 			{
 				parentWarnings.ShadowVariableWarningLevel = WarningLevel.Error;
 				parentWarnings.UndefinedIdentifierWarningLevel = WarningLevel.Error;
+				parentWarnings.UnreachableCodeWarningLevel = WarningLevel.Error;
 				parentWarnings.UnsafeTypeCastWarningLevel = WarningLevel.Error;
 
 				Assert.AreEqual(WarningLevel.Error, newWarnings.ShadowVariableWarningLevel);
-				Assert.AreEqual(WarningLevel.Error, newWarnings.UnsafeTypeCastWarningLevel);
 				Assert.AreEqual(WarningLevel.Error, newWarnings.UndefinedIdentifierWarningLevel);
+				Assert.AreEqual(WarningLevel.Error, newWarnings.UnreachableCodeWarningLevel);
+				Assert.AreEqual(WarningLevel.Error, newWarnings.UnsafeTypeCastWarningLevel);
 			}
 
 			// Assert reparenting
@@ -175,11 +185,13 @@ namespace UnrealBuildTool.Tests
 				newWarnings.AddParent(parentWarnings2);
 				parentWarnings2.ShadowVariableWarningLevel = WarningLevel.Warning;
 				parentWarnings2.UndefinedIdentifierWarningLevel = WarningLevel.Warning;
+				parentWarnings2.UnreachableCodeWarningLevel = WarningLevel.Warning;
 				parentWarnings2.UnsafeTypeCastWarningLevel = WarningLevel.Warning;
 
 				Assert.AreEqual(WarningLevel.Warning, newWarnings.ShadowVariableWarningLevel);
-				Assert.AreEqual(WarningLevel.Warning, newWarnings.UnsafeTypeCastWarningLevel);
 				Assert.AreEqual(WarningLevel.Warning, newWarnings.UndefinedIdentifierWarningLevel);
+				Assert.AreEqual(WarningLevel.Warning, newWarnings.UnreachableCodeWarningLevel);
+				Assert.AreEqual(WarningLevel.Warning, newWarnings.UnsafeTypeCastWarningLevel);
 			}
 
 			// Assert transitive to eldest ancestor
@@ -188,11 +200,12 @@ namespace UnrealBuildTool.Tests
 				Assert.AreEqual(parentWarnings.ShadowVariableWarningLevel, newWarnings.ShadowVariableWarningLevel);
 			}
 
-			// Assert removal of parentWarnings2 - UnsafeTypeCastWarningLevel & UndefinedIdentifierWarningLevel will now be the reoslved items
+			// Assert removal of parentWarnings2 - UndefinedIdentifierWarningLevel, UnreachableCodeWarningLevel, and UnsafeTypeCastWarningLevel will now be the resolved items
 			{
 				newWarnings.RemoveParent();
-				Assert.AreEqual(parentWarnings.UnsafeTypeCastWarningLevel, newWarnings.UnsafeTypeCastWarningLevel);
 				Assert.AreEqual(parentWarnings.UndefinedIdentifierWarningLevel, newWarnings.UndefinedIdentifierWarningLevel);
+				Assert.AreEqual(parentWarnings.UnreachableCodeWarningLevel, newWarnings.UnreachableCodeWarningLevel);
+				Assert.AreEqual(parentWarnings.UnsafeTypeCastWarningLevel, newWarnings.UnsafeTypeCastWarningLevel);
 
 			}
 
@@ -200,8 +213,9 @@ namespace UnrealBuildTool.Tests
 			{
 				newWarnings.RemoveParent();
 				Assert.AreEqual(WarningLevel.Default, newWarnings.ShadowVariableWarningLevel);
-				Assert.AreEqual(WarningLevel.Default, newWarnings.UnsafeTypeCastWarningLevel);
 				Assert.AreEqual(WarningLevel.Default, newWarnings.UndefinedIdentifierWarningLevel);
+				Assert.AreEqual(WarningLevel.Default, newWarnings.UnreachableCodeWarningLevel);
+				Assert.AreEqual(WarningLevel.Default, newWarnings.UnsafeTypeCastWarningLevel);
 			}
 		}
 
@@ -215,25 +229,29 @@ namespace UnrealBuildTool.Tests
 
 			parentWarnings.ShadowVariableWarningLevel = WarningLevel.Error;
 			parentWarnings.UndefinedIdentifierWarningLevel = WarningLevel.Error;
+			parentWarnings.UnreachableCodeWarningLevel = WarningLevel.Error;
 			parentWarnings.UnsafeTypeCastWarningLevel = WarningLevel.Error;
 
 			// Assert invariants
 			{
 				Assert.AreEqual(WarningLevel.Error, newWarnings.ShadowVariableWarningLevel);
-				Assert.AreEqual(WarningLevel.Error, newWarnings.UnsafeTypeCastWarningLevel);
 				Assert.AreEqual(WarningLevel.Error, newWarnings.UndefinedIdentifierWarningLevel);
+				Assert.AreEqual(WarningLevel.Error, newWarnings.UnreachableCodeWarningLevel);
+				Assert.AreEqual(WarningLevel.Error, newWarnings.UnsafeTypeCastWarningLevel);
 			}
 
 			// Child overrides;
 			newWarnings.ShadowVariableWarningLevel = WarningLevel.Warning;
 			newWarnings.UndefinedIdentifierWarningLevel = WarningLevel.Warning;
+			newWarnings.UnreachableCodeWarningLevel = WarningLevel.Warning;
 			newWarnings.UnsafeTypeCastWarningLevel = WarningLevel.Warning;
 
 			// Assert new parent values
 			{
 				Assert.AreEqual(WarningLevel.Warning, newWarnings.ShadowVariableWarningLevel);
-				Assert.AreEqual(WarningLevel.Warning, newWarnings.UnsafeTypeCastWarningLevel);
 				Assert.AreEqual(WarningLevel.Warning, newWarnings.UndefinedIdentifierWarningLevel);
+				Assert.AreEqual(WarningLevel.Warning, newWarnings.UnreachableCodeWarningLevel);
+				Assert.AreEqual(WarningLevel.Warning, newWarnings.UnsafeTypeCastWarningLevel);
 			}
 		}
 
@@ -321,6 +339,23 @@ namespace UnrealBuildTool.Tests
 					baseWarnings.UndefinedIdentifierWarningLevel = WarningLevel.Default;
 				}
 
+				// Specialized VC toolchain check; UnreachableCode
+				{
+					baseWarnings.UnreachableCodeWarningLevel = WarningLevel.Warning;
+					args = baseWarnings.GenerateWarningCommandLineArgs(newCompileEnvironment, typeof(VCToolChain));
+
+					Assert.IsTrue(args.Count() == 1);
+					Assert.IsTrue(args.Contains("/w44702"));
+
+					CppCompileEnvironment preprocessOnlyCompileEnvironment = new CppCompileEnvironment(newCompileEnvironment);
+					preprocessOnlyCompileEnvironment.bPreprocessOnly = true;
+					args = baseWarnings.GenerateWarningCommandLineArgs(preprocessOnlyCompileEnvironment, typeof(VCToolChain));
+
+					Assert.IsFalse(args.Any());
+
+					baseWarnings.UnreachableCodeWarningLevel = WarningLevel.Default;
+				}
+
 				{
 					// DeterministicWarningLevel
 					VersionNumber newVersionNumber = new VersionNumber(18, 9, 9);
diff --git a/Engine/Source/Programs/UnrealBuildTool/Configuration/CppCompileWarnings.cs b/Engine/Source/Programs/UnrealBuildTool/Configuration/CppCompileWarnings.cs
--- a/Engine/Source/Programs/UnrealBuildTool/Configuration/CppCompileWarnings.cs
+++ b/Engine/Source/Programs/UnrealBuildTool/Configuration/CppCompileWarnings.cs
@@ -155,6 +155,28 @@ namespace UnrealBuildTool
 		}
 		private WarningLevel _undefinedIdentifierWarningLevel;
 
+		/// <summary>
+		/// Indicates what warning/error level to treat unreachable code.
+		/// MSVC - 
+		///		https://learn.microsoft.com/en-us/cpp/error-messages/compiler-warnings/compiler-warning-level-4-c4702 
+		///		4702 - unreachable code
+		///		44702 - Note: The extra 4 is not a typo, /wLXXXX sets warning XXXX to level L
+		/// Clang - 
+		///		https://clang.llvm.org/docs/DiagnosticsReference.html#wunreachable-code
+		/// </summary>
+		[XmlConfigFile(Category = "BuildConfiguration")]
+		[CommandLine("-UnreachableCodeWarningLevel=")]
+		[WarningsVCToolChainAttribute(null, ["/w44702"], ["/we4702"])]
+		[WarningsClangToolChain(null, ["-Wunreachable-code-aggressive", "-Wno-error=unreachable-code-aggressive"], ["-Wunreachable-code-aggressive"])]
+		[VersionWarningLevelDefault(WarningLevel.Off, BuildSettingsVersion.V1, BuildSettingsVersion.V5, InitializationContext.Any)]
+		[VersionWarningLevelDefault(WarningLevel.Error, BuildSettingsVersion.V6, BuildSettingsVersion.Latest, InitializationContext.Any)]
+		public WarningLevel UnreachableCodeWarningLevel
+		{
+			get => ResolveWarning(ParentCppCompileWarnings?.UnreachableCodeWarningLevel, _unreachableCodeWarningLevel);
+			set => _unreachableCodeWarningLevel = value;
+		}
+		private WarningLevel _unreachableCodeWarningLevel;
+
 		/// <summary>
 		/// Indicates what warning/error level to treat unhandled enumerators in switches on enumeration-typed values.
 		/// MSVC - 
diff --git a/Engine/Source/Programs/UnrealBuildTool/Configuration/ReadOnlyCppCompileWarnings.cs b/Engine/Source/Programs/UnrealBuildTool/Configuration/ReadOnlyCppCompileWarnings.cs
--- a/Engine/Source/Programs/UnrealBuildTool/Configuration/ReadOnlyCppCompileWarnings.cs
+++ b/Engine/Source/Programs/UnrealBuildTool/Configuration/ReadOnlyCppCompileWarnings.cs
@@ -41,6 +41,7 @@ namespace UnrealBuildTool
 
 		public WarningLevel ShadowVariableWarningLevel => _inner.ShadowVariableWarningLevel;
 		public WarningLevel UndefinedIdentifierWarningLevel => _inner.UndefinedIdentifierWarningLevel;
+		public WarningLevel UnreachableCodeWarningLevel => _inner.UnreachableCodeWarningLevel;
 		public WarningLevel UnsafeTypeCastWarningLevel => _inner.UnsafeTypeCastWarningLevel;
 		public WarningLevel SwitchUnhandledEnumeratorWarningLevel => _inner.SwitchUnhandledEnumeratorWarningLevel;
 		public WarningLevel ShortenSizeTToIntWarningLevel => _inner.ShortenSizeTToIntWarningLevel;
diff --git a/Engine/Source/Programs/UnrealBuildTool/Configuration/TargetRules.cs b/Engine/Source/Programs/UnrealBuildTool/Configuration/TargetRules.cs
--- a/Engine/Source/Programs/UnrealBuildTool/Configuration/TargetRules.cs
+++ b/Engine/Source/Programs/UnrealBuildTool/Configuration/TargetRules.cs
@@ -156,6 +156,7 @@ namespace UnrealBuildTool
 		/// New defaults for 5.7:
 		/// * DanglingWarningLevel = WarningLevel.Warning
 		/// * UndefinedIdentifierWarningLevel = WarningLevel.Error
+		/// * UnreachableCodeWarningLevel = WarningLevel.Error
 		/// </summary>
 		V6,
 
diff --git a/Engine/Source/Programs/UnrealBuildTool/Platform/Linux/LinuxToolChain.cs b/Engine/Source/Programs/UnrealBuildTool/Platform/Linux/LinuxToolChain.cs
--- a/Engine/Source/Programs/UnrealBuildTool/Platform/Linux/LinuxToolChain.cs
+++ b/Engine/Source/Programs/UnrealBuildTool/Platform/Linux/LinuxToolChain.cs
@@ -280,7 +280,6 @@ namespace UnrealBuildTool
 		{
 			base.GetCompileArguments_WarningsAndErrors(CompileEnvironment, Arguments);
 
-			//Arguments.Add("-Wunreachable-code");            // additional warning not normally included in Wall: warns if there is code that will never be executed - not helpful due to bIsGCC and similar
 			Arguments.Add("-Wno-dangling");		// kinda scary, but matches what we currently do on MS platforms when clang is used
 
 			Arguments.Add("-Wno-undefined-bool-conversion"); // hides checking if 'this' pointer is null
diff --git a/Engine/Source/Runtime/Core/Public/Clang/ClangPlatformCompilerPreSetup.h b/Engine/Source/Runtime/Core/Public/Clang/ClangPlatformCompilerPreSetup.h
--- a/Engine/Source/Runtime/Core/Public/Clang/ClangPlatformCompilerPreSetup.h
+++ b/Engine/Source/Runtime/Core/Public/Clang/ClangPlatformCompilerPreSetup.h
@@ -65,11 +65,14 @@
 
 // Unreachable not being evaluated on clang... yet
 #ifndef PRAGMA_DISABLE_UNREACHABLE_CODE_WARNINGS
-	#define PRAGMA_DISABLE_UNREACHABLE_CODE_WARNINGS
+	#define PRAGMA_DISABLE_UNREACHABLE_CODE_WARNINGS \
+		_Pragma("clang diagnostic push") \
+		_Pragma("clang diagnostic ignored \"-Wunreachable-code-aggressive\"")
 #endif // PRAGMA_DISABLE_UNREACHABLE_CODE_WARNINGS
 
 #ifndef PRAGMA_RESTORE_UNREACHABLE_CODE_WARNINGS
-	#define PRAGMA_RESTORE_UNREACHABLE_CODE_WARNINGS
+	#define PRAGMA_RESTORE_UNREACHABLE_CODE_WARNINGS \
+		_Pragma("clang diagnostic pop")
 #endif // PRAGMA_RESTORE_UNREACHABLE_CODE_WARNINGS
 
 #ifndef PRAGMA_PUSH_ATTRIBUTE_MINSIZE_FUNCTIONS
