Small generated code fixes for constinit uobjects

diff --git a/Engine/Source/Programs/Shared/EpicGames.UHT/Exporters/CodeGen/UhtHeaderCodeGeneratorCppFile.cs b/Engine/Source/Programs/Shared/EpicGames.UHT/Exporters/CodeGen/UhtHeaderCodeGeneratorCppFile.cs
--- a/Engine/Source/Programs/Shared/EpicGames.UHT/Exporters/CodeGen/UhtHeaderCodeGeneratorCppFile.cs
+++ b/Engine/Source/Programs/Shared/EpicGames.UHT/Exporters/CodeGen/UhtHeaderCodeGeneratorCppFile.cs
@@ -83,6 +83,7 @@ namespace EpicGames.UHT.Exporters.CodeGen
 				}
 				if (HeaderFile.Children.Any(x => x is UhtStruct structObj && x.Children.Any(y => y is UhtVerseValueProperty)))
 				{
+					includesToAdd.Add("VerseVM/VBPVMDynamicProperty.h");  // Only needed when WITH_VERSE_BPVM
 					includesToAdd.Add("UObject/VerseValueProperty.h"); 
 				}
 
@@ -1362,9 +1363,9 @@ namespace EpicGames.UHT.Exporters.CodeGen
 					if (function.IsVerseField)
 					{
 						builder.Append("\tUE::CodeGen::ConstInit::FVerseFunctionParams{\r\n")
-							.Append("\t\t.CompiledInAlternateNameUTF8 = IF_WITH_VERSE_VM(")
+							.Append("\t\t.AlternateNameUTF8 = IF_WITH_VERSE_VM(")
 							.Append("UTF8TEXT(").AppendUTF8LiteralString(function.GetVerseScopeAndName(UhtVerseNameMode.Default)).Append("), ")
-							.Append("UTF8TEXT(").AppendUTF8LiteralString(function.GetEncodedVerseScopeAndName(UhtVerseNameMode.Default)).Append("), ")
+							.Append("UTF8TEXT(").AppendUTF8LiteralString(function.GetEncodedVerseScopeAndName(UhtVerseNameMode.Default)).Append(") ")
 							.Append("), \r\n")
 							.Append("\t},\r\n");
 					}
@@ -2897,8 +2898,8 @@ namespace EpicGames.UHT.Exporters.CodeGen
 					if (classObj.IsVerseField)
 					{
 						builder.Append("\tUE::CodeGen::ConstInit::FVerseClassParams{\r\n")
-							.Append("\t\t.PackageRelativeVersePath = ").AppendVerseScopeAndName(classObj, UhtVerseNameMode.PackageRelative).Append(",\r\n")
-							.Append("\t\t.MangledPackageVersePath = ").Append(VerseNameMangling.MangleCasedName(classObj.Module.Module.VersePath).Result).Append(",\r\n");
+							.Append("\t\t.PackageRelativeVersePath = UTF8TEXT(\"").AppendVerseScopeAndName(classObj, UhtVerseNameMode.PackageRelative).Append("\"),\r\n")
+							.Append("\t\t.MangledPackageVersePath = UTF8TEXT(\"").Append(VerseNameMangling.MangleCasedName(classObj.Module.Module.VersePath).Result).Append("\"),\r\n");
 						if (hasVerseCallableFunctions)
 						{
 							builder.Append("\t\tIF_WITH_VERSE_VM(.VerseCallableThunks = ").AppendArrayView(staticsName, "VerseFuncs", 0, ",)\r\n");
