EpicGames.MsBuild: Retry on MacOS when building due to intermittent SIGBUS errors
Print log lines as informational if dotnet build fails
Lower retry time to 1 second

#rnx

diff --git a/Engine/Build/BatchFiles/DotnetRetry.sh b/Engine/Build/BatchFiles/DotnetRetry.sh
--- a/Engine/Build/BatchFiles/DotnetRetry.sh
+++ b/Engine/Build/BatchFiles/DotnetRetry.sh
@@ -28,7 +28,7 @@ elif [ "$(uname)" = "Linux" ]; then
 fi
 
 MAX_RETRIES=3
-RETRY_DELAY=5
+RETRY_DELAY=1
 COUNT=0
 ERRORCODE=0
 
diff --git a/Engine/Source/Programs/Shared/EpicGames.MsBuild/CsProjBuilder.cs b/Engine/Source/Programs/Shared/EpicGames.MsBuild/CsProjBuilder.cs
--- a/Engine/Source/Programs/Shared/EpicGames.MsBuild/CsProjBuilder.cs
+++ b/Engine/Source/Programs/Shared/EpicGames.MsBuild/CsProjBuilder.cs
@@ -16,7 +16,6 @@ using EpicGames.Core;
 using Microsoft.Build.Evaluation;
 using Microsoft.Build.Locator;
 using Microsoft.Extensions.Logging;
-using Microsoft.Extensions.Logging.Abstractions;
 
 namespace EpicGames.MsBuild
 {
@@ -421,8 +420,6 @@ namespace EpicGames.MsBuild
 
 		private static async Task<bool> RunDotnetAsync(IEnumerable<string> arguments, CsProjBuildHook hook, ILogger logger, CancellationToken cancellationToken)
 		{
-			string outputPrefix = $"dotnet {arguments.First()}";
-
 			ProcessStartInfo startInfo = new ProcessStartInfo
 			{
 				FileName = hook.DotnetPath.FullName,
@@ -436,6 +433,7 @@ namespace EpicGames.MsBuild
 			startInfo.EnvironmentVariables["DOTNET_MULTILEVEL_LOOKUP"] = "0"; // use only the bundled dotnet installation - ignore any other/system dotnet install
 
 			logger.LogDebug("Running: {Application} {Arguments}", startInfo.FileName, startInfo.Arguments);
+			PrefixLogger prefixLogger = new($"[dotnet {arguments.First()}]", logger);
 			using Process dotnetProcess = new Process();
 			dotnetProcess.StartInfo = startInfo;
 			dotnetProcess.OutputDataReceived += (sender, args) =>
@@ -443,7 +441,7 @@ namespace EpicGames.MsBuild
 				string output = args.Data?.TrimEnd() ?? String.Empty;
 				if (!String.IsNullOrEmpty(output))
 				{
-					logger.LogDebug("[{Prefix}] {Message}", outputPrefix, output);
+					prefixLogger.LogInformation("{Message}", output);
 				}
 			};
 			dotnetProcess.ErrorDataReceived += (sender, args) =>
@@ -451,7 +449,7 @@ namespace EpicGames.MsBuild
 				string output = args.Data?.TrimEnd() ?? String.Empty;
 				if (!String.IsNullOrEmpty(output))
 				{
-					logger.LogError("[{Prefix}] {Message}", outputPrefix, output);
+					prefixLogger.LogError("{Message}", output);
 				}
 			};
 			dotnetProcess.Start();
@@ -502,7 +500,47 @@ namespace EpicGames.MsBuild
 				.. globalProperties.Select(x => $"\"/p:{x.Key}={x.Value}\""),
 			];
 
-			bool result = await RunDotnetAsync(arguments, hook, logger, cancellationToken);
+			bool result = false;
+			CaptureLogger captureLogger = new();
+
+			// MacOS has been experiencing intermittent SIGBUS errors when building, so retry a few times on failure
+			if (OperatingSystem.IsMacOS())
+			{
+				int maxRetries = 3;
+				TimeSpan retryDelay = new(0, 0, 1);
+				int count = 0;
+
+				while (count < maxRetries && !cancellationToken.IsCancellationRequested)
+				{
+					captureLogger = new();
+					result = await RunDotnetAsync(arguments, hook, captureLogger, cancellationToken);
+					if (result)
+					{
+						break;
+					}
+					count++;
+					if (count < maxRetries)
+					{
+						logger.LogInformation("dotnet command failed. Retrying in {Seconds} seconds... ({Count}/{MaxRetries})", retryDelay.Seconds, count, maxRetries);
+						await Task.Delay(retryDelay, cancellationToken);
+					}
+					else
+					{
+						logger.LogInformation("dotnet command failed after {MaxRetries} attempts", maxRetries);
+					}
+				}
+			}
+			else
+			{
+				result = await RunDotnetAsync(arguments, hook, captureLogger, cancellationToken);
+			}
+
+			// If successful downgrade logs to Debug
+			if (result)
+			{
+				captureLogger.Events.ForEach(x => x.Level = LogLevel.Debug);
+			}
+			captureLogger.RenderTo(logger);
 
 			logger.LogInformation("Build {Count} projects time: {TimeSeconds:0.00} s", projects.Count(), stopwatch.Elapsed.TotalSeconds);
 
