[EpicGames.Build]
* Fixed bug in Parallel2.For when from and to are the same

diff --git a/Engine/Source/Programs/Shared/EpicGames.Build/System/Parallel2.cs b/Engine/Source/Programs/Shared/EpicGames.Build/System/Parallel2.cs
--- a/Engine/Source/Programs/Shared/EpicGames.Build/System/Parallel2.cs
+++ b/Engine/Source/Programs/Shared/EpicGames.Build/System/Parallel2.cs
@@ -130,6 +130,16 @@ namespace UnrealBuildBase
 		/// <param name="useTasks">Use tasks or threads</param>
 		public static void For(int from, int to, Action<int> body, int helperCount = -1, bool useTasks = true)
 		{
+			int distance = to - from;
+			if (distance <= 1)
+			{
+				if (distance == 1)
+				{
+					body(from);
+				}
+				return;
+			}
+
 			using EventWaitHandle Handle = new(false, EventResetMode.ManualReset);
 
 			ForContext context = new() { Handle = Handle, Index = from, End = to, Finished = from };
@@ -168,7 +178,7 @@ namespace UnrealBuildBase
 
 			if (helperCount > 0)
 			{
-				helperCount = Math.Min(to - from, helperCount);
+				helperCount = Math.Min(distance, helperCount);
 				if (helperCount > 0)
 				{
 					--helperCount;
