[EpicGames.Core]
* Added StackStringBuffer which is a string buffer that must reside on the stack
* Changed FileSystemReference.CombineStrings to use StackStringBuffer and removed thread local usage

diff --git a/Engine/Source/Programs/Shared/EpicGames.Core/FileSystemReference.cs b/Engine/Source/Programs/Shared/EpicGames.Core/FileSystemReference.cs
--- a/Engine/Source/Programs/Shared/EpicGames.Core/FileSystemReference.cs
+++ b/Engine/Source/Programs/Shared/EpicGames.Core/FileSystemReference.cs
@@ -4,6 +4,7 @@ using System;
 using System.Collections.Generic;
 using System.IO;
 using System.Linq;
+using System.Runtime.CompilerServices;
 using System.Text;
 using System.Threading;
 
@@ -47,83 +48,104 @@ namespace EpicGames.Core
 		/// </summary>
 		private readonly Lazy<int> _hashCode;
 
-		static readonly ThreadLocal<StringBuilder> s_combineStringsStringBuilder = new ThreadLocal<StringBuilder>(() => new StringBuilder(260));
-
 		/// <summary>
 		/// Create a full path by concatenating multiple strings
 		/// </summary>
+		[MethodImpl(MethodImplOptions.AggressiveOptimization)]
 		protected static string CombineStrings(DirectoryReference baseDirectory, params string[] fragments)
 		{
-			// Get the initial string to append to, and strip any root directory suffix from it
-			StringBuilder newFullName = s_combineStringsStringBuilder.Value!.Clear().Append(baseDirectory.FullName);
-			if (newFullName.Length > 0 && newFullName[^1] == Path.DirectorySeparatorChar)
-			{
-				newFullName.Remove(newFullName.Length - 1, 1);
-			}
+			char Sep = Path.DirectorySeparatorChar;
+			char Alt = Path.AltDirectorySeparatorChar;
 
-			// Scan through the fragments to append, appending them to a string and updating the base length as we go
-			foreach (string fragment in fragments)
+			Span<char> initial = stackalloc char[512];
+			using StackStringBuilder sb = new(initial);
+
+			foreach (string? fragment in fragments)
 			{
 				// Check if this fragment is an absolute path
 				if ((fragment.Length >= 2 && fragment[1] == ':') || (fragment.Length >= 1 && (fragment[0] == '\\' || fragment[0] == '/')))
 				{
 					// It is. Reset the new name to the full version of this path.
-					newFullName.Clear();
-					newFullName.Append(Path.GetFullPath(fragment).TrimEnd(Path.DirectorySeparatorChar));
+					sb.Resize(0);
+					ReadOnlySpan<char> abs = Path.GetFullPath(fragment).AsSpan();
+					if (abs.Length == 3 && abs[2] == ':')
+					{
+						sb.Append(abs.Slice(0, 2));
+					}
+					else
+					{
+						sb.Append(abs);
+					}
+					continue;
 				}
-				else
+				else if (sb.Length == 0)
 				{
-					// Append all the parts of this fragment to the end of the existing path.
-					int startIdx = 0;
-					while (startIdx < fragment.Length)
+					ReadOnlySpan<char> baseSpan = baseDirectory.FullName.AsSpan();
+
+					// If only drive, remove separator (otherwise there should never be a separator)
+					if (baseSpan.Length == 3 && baseSpan[2] == ':')
 					{
-						// Find the end of this fragment. We may have been passed multiple paths in the same string.
-						int endIdx = startIdx;
-						while (endIdx < fragment.Length && fragment[endIdx] != '\\' && fragment[endIdx] != '/')
-						{
-							endIdx++;
-						}
+						sb.Append(baseSpan.Slice(0, 2));
+					}
+					else
+					{
+						sb.Append(baseSpan);
+					}
+				}
 
-						// Ignore any empty sections, like leading or trailing slashes, and '.' directory references.
-						int length = endIdx - startIdx;
-						if (length == 0)
-						{
-							// Multiple directory separators in a row; illegal.
-							throw new ArgumentException(String.Format("Path fragment '{0}' contains invalid directory separators.", fragment));
-						}
-						else if (length == 2 && fragment[startIdx] == '.' && fragment[startIdx + 1] == '.')
+				// Relative fragment: allow separators (including repeated); skip empty segments
+				ReadOnlySpan<char> span = fragment.AsSpan();
+				int idx = 0;
+
+				while (idx <= span.Length)
+				{
+					int rel = (idx < span.Length) ? span.Slice(idx).IndexOfAny(Sep, Alt) : -1;
+					int segLen = (rel < 0) ? (span.Length - idx) : rel;
+
+					if (segLen == 0)
+					{
+						// Empty segment (e.g., leading/trailing or repeated separators) — skip
+					}
+					else if (segLen == 1 && span[idx] == '.')
+					{
+						// Skip "."
+					}
+					else if (segLen == 2 && span[idx] == '.' && span[idx + 1] == '.')
+					{
+						// Pop last segment; if none, leave as-is
+						for (int i = sb.Length - 1; i >= 0; i--)
 						{
-							// Remove the last directory name
-							for (int separatorIdx = newFullName.Length - 1; separatorIdx >= 0; separatorIdx--)
+							if (sb[i] == Sep)
 							{
-								if (newFullName[separatorIdx] == Path.DirectorySeparatorChar)
-								{
-									newFullName.Remove(separatorIdx, newFullName.Length - separatorIdx);
-									break;
-								}
+								sb.Resize(i);
+								break;
 							}
 						}
-						else if (length != 1 || fragment[startIdx] != '.')
+					}
+					else
+					{
+						// Append separator only if there's already content and the last char isn't a separator
+						if (sb.Length > 0)
 						{
-							// Append this fragment
-							newFullName.Append(Path.DirectorySeparatorChar);
-							newFullName.Append(fragment, startIdx, length);
+							if (sb[sb.Length - 1] != Sep)
+							{
+								sb.Append(Sep);
+							}
 						}
 
-						// Move to the next part
-						startIdx = endIdx + 1;
+						// Append segment (no separators inside by construction)
+						sb.Append(span.Slice(idx, segLen));
 					}
-				}
-			}
 
-			// Append the directory separator
-			if (newFullName.Length == 0 || (newFullName.Length == 2 && newFullName[1] == ':'))
-			{
-				newFullName.Append(Path.DirectorySeparatorChar);
-			}
+					if (rel < 0)
+					{
+						break;
+					}
 
-			// Set the new path variables
-			return newFullName.ToString();
+					idx += rel + 1; // move past the separator
+				}
+			}
+			return sb.ToString();
 		}
 
 		/// <summary>
@@ -231,6 +253,7 @@ namespace EpicGames.Core
 		/// </summary>
 		/// <param name="directory"></param>
 		/// <param name="result"></param>
+		[MethodImpl(MethodImplOptions.AggressiveOptimization)]
 		public void WriteRelativeTo(DirectoryReference directory, StringBuilder result)
 		{
 			// Find how much of the path is common between the two paths. This length does not include a trailing directory separator character.
diff --git a/Engine/Source/Programs/Shared/EpicGames.Core/StackStringBuffer.cs b/Engine/Source/Programs/Shared/EpicGames.Core/StackStringBuffer.cs
new file mode 100644
--- /dev/null
+++ b/Engine/Source/Programs/Shared/EpicGames.Core/StackStringBuffer.cs
@@ -0,0 +1,103 @@
+// Copyright Epic Games, Inc. All Rights Reserved.
+
+using System;
+using System.Buffers;
+using System.Runtime.CompilerServices;
+
+namespace EpicGames.Core
+{
+	/// <summary>
+	/// StringBuffer that can only live on the stack and does not make any allocations as long as initial buffer is big enough
+	/// Usage:
+	/// Span<char> initial = stackalloc char[512];
+	/// StackStringBuilder sb = new(initial);
+	/// sb.Append("Foo");
+	/// sb.Dispose(); // Needed in case StackStringBuffer reallocated
+	/// </summary>
+	ref struct StackStringBuilder
+	{
+		private Span<char> _chars;
+		private int _pos;
+		private char[]? _arrayToReturnToPool;
+
+		/// <summary>
+		/// Ctor
+		/// </summary>
+		/// <param name="initialBuffer">Initial buffer</param>
+		public StackStringBuilder(Span<char> initialBuffer)
+		{
+			_chars = initialBuffer;
+		}
+
+		public int Length => _pos;
+
+		public ref char this[int index] => ref _chars[index];
+
+		/// <summary>
+		/// Resize string buffer
+		/// </summary>
+		/// <param name="newSize"></param>
+		public void Resize(int newSize)
+		{
+			_pos = newSize;
+		}
+
+		/// <summary>
+		/// Append character
+		/// </summary>
+		/// <param name="c"></param>
+		public void Append(char c)
+		{
+			int pos = _pos;
+			if ((uint)pos < (uint)_chars.Length)
+			{
+				_chars[pos] = c;
+				_pos = pos + 1;
+			}
+			else
+			{
+				Grow(1);
+				Append(c);
+			}
+		}
+
+		/// <summary>
+		/// Append span of characters
+		/// </summary>
+		/// <param name="s"></param>
+		public void Append(ReadOnlySpan<char> s)
+		{
+			int newPos = _pos + s.Length;
+			if (newPos > _chars.Length)
+			{
+				Grow(s.Length);
+			}
+			s.CopyTo(_chars.Slice(_pos));
+			_pos = newPos;
+		}
+
+		public override string ToString()
+		{
+			return new string(_chars.Slice(0, _pos));
+		}
+
+		public void Dispose()
+		{
+			char[]? toReturn = _arrayToReturnToPool;
+			this = default;
+			if (toReturn is not null)
+			{
+				ArrayPool<char>.Shared.Return(toReturn);
+			}
+		}
+
+		private void Grow(int requiredAdditionalCapacity)
+		{
+			int newSize = Math.Max(_pos + requiredAdditionalCapacity, _chars.Length * 2);
+			char[] poolArray = ArrayPool<char>.Shared.Rent(newSize);
+			_chars.Slice(0, _pos).CopyTo(poolArray);
+			_chars = poolArray;
+			_arrayToReturnToPool = poolArray;
+		}
+	}
+}
\ No newline at end of file
