[EpicGames.Build]
* Fixed how exceptions are marshalled to calling thread for Parallel2.ForEach, Parallel2.For and ParallelQueue

diff --git a/Engine/Source/Programs/Shared/EpicGames.Build/System/Parallel2.cs b/Engine/Source/Programs/Shared/EpicGames.Build/System/Parallel2.cs
--- a/Engine/Source/Programs/Shared/EpicGames.Build/System/Parallel2.cs
+++ b/Engine/Source/Programs/Shared/EpicGames.Build/System/Parallel2.cs
@@ -4,6 +4,7 @@ using System;
 using System.Collections.Concurrent;
 using System.Collections.Generic;
 using System.Linq;
+using System.Runtime.ExceptionServices;
 using System.Threading;
 using System.Threading.Tasks;
 
@@ -37,7 +38,7 @@ namespace UnrealBuildBase
 			void WorkerLoop()
 			{
 				int finished = 0;
-				Exception? exception = null;
+				ExceptionDispatchInfo? exception = null;
 				while (true)
 				{
 					Monitor.Enter(context);
@@ -76,9 +77,9 @@ namespace UnrealBuildBase
 					{
 						body(c);
 					}
-					catch (Exception e)
+					catch (Exception ex)
 					{
-						exception ??= e;
+						exception ??= ExceptionDispatchInfo.Capture(ex);
 					}
 					finished = 1;
 				}
@@ -116,7 +117,7 @@ namespace UnrealBuildBase
 
 			if (context.Exception != null)
 			{
-				throw context.Exception;
+				context.Exception.Throw();
 			}
 		}
 
@@ -148,9 +149,9 @@ namespace UnrealBuildBase
 					{
 						body(index);
 					}
-					catch (Exception e)
+					catch (Exception ex)
 					{
-						Interlocked.CompareExchange(ref context.Exception, e, null);
+						Interlocked.CompareExchange(ref context.Exception, ExceptionDispatchInfo.Capture(ex), null);
 					}
 
 					if (Interlocked.Increment(ref context.Finished) == context.End)
@@ -193,7 +194,7 @@ namespace UnrealBuildBase
 
 			if (context.Exception != null)
 			{
-				throw context.Exception;
+				context.Exception.Throw();
 			}
 		}
 
@@ -204,7 +205,7 @@ namespace UnrealBuildBase
 			internal int Count;
 			internal int Finished;
 			internal bool Done;
-			internal Exception? Exception;
+			internal ExceptionDispatchInfo? Exception;
 		}
 
 		private class ForContext
@@ -213,7 +214,7 @@ namespace UnrealBuildBase
 			internal int Index;
 			internal int End;
 			internal int Finished;
-			internal Exception? Exception;
+			internal ExceptionDispatchInfo? Exception;
 		}
 
 		public static readonly int MaxHelperCount = Math.Max(Environment.ProcessorCount - 1, 0);
@@ -284,7 +285,7 @@ namespace UnrealBuildBase
 					}
 					catch (Exception ex)
 					{
-						_exception = ex;
+						_exception = ExceptionDispatchInfo.Capture(ex);
 					}
 
 					if (Interlocked.Decrement(ref _outstanding) != 0)
@@ -317,7 +318,7 @@ namespace UnrealBuildBase
 
 			if (_exception != null)
 			{
-				throw _exception;
+				_exception.Throw();
 			}
 		}
 
@@ -329,7 +330,7 @@ namespace UnrealBuildBase
 		private int _accepting = 1;
 		private int _done;
 		private int _workerCount;
-		private Exception? _exception;
+		private ExceptionDispatchInfo? _exception;
 
 		public static readonly int MaxHelperCount = Math.Max(Environment.ProcessorCount - 1, 0);
 	}
