[EpicGames.Core]
* Fixed so DirectoryReference.Combine(dir); returns dir

diff --git a/Engine/Source/Programs/Shared/EpicGames.Core/DirectoryReference.cs b/Engine/Source/Programs/Shared/EpicGames.Core/DirectoryReference.cs
--- a/Engine/Source/Programs/Shared/EpicGames.Core/DirectoryReference.cs
+++ b/Engine/Source/Programs/Shared/EpicGames.Core/DirectoryReference.cs
@@ -3,6 +3,7 @@
 using System;
 using System.Collections.Generic;
 using System.ComponentModel;
+using System.Diagnostics;
 using System.Diagnostics.CodeAnalysis;
 using System.Globalization;
 using System.IO;
@@ -161,7 +162,7 @@ namespace EpicGames.Core
 		/// <param name="baseDirectory">The base directory</param>
 		/// <param name="fragments">Fragments to combine with the base directory</param>
 		/// <returns>The new directory name</returns>
-		public static DirectoryReference Combine(DirectoryReference baseDirectory, params string[] fragments) => new(CombineStrings(baseDirectory, fragments), Sanitize.None);
+		public static DirectoryReference Combine(DirectoryReference baseDirectory, params string[] fragments) => fragments.Length == 0 ? baseDirectory : new(CombineStrings(baseDirectory, fragments), Sanitize.None);
 
 		/// <summary>
 		/// Compares two filesystem object names for equality. Uses the canonical name representation, not the display name representation.
