[EpicGames.Core]
* Fixes for static analyzer  (coding standard compliance)

diff --git a/Engine/Source/Programs/Shared/EpicGames.Core/FileSystemReference.cs b/Engine/Source/Programs/Shared/EpicGames.Core/FileSystemReference.cs
--- a/Engine/Source/Programs/Shared/EpicGames.Core/FileSystemReference.cs
+++ b/Engine/Source/Programs/Shared/EpicGames.Core/FileSystemReference.cs
@@ -54,8 +54,8 @@ namespace EpicGames.Core
 		[MethodImpl(MethodImplOptions.AggressiveOptimization)]
 		protected static string CombineStrings(DirectoryReference baseDirectory, params string[] fragments)
 		{
-			char Sep = Path.DirectorySeparatorChar;
-			char Alt = Path.AltDirectorySeparatorChar;
+			char sep = Path.DirectorySeparatorChar;
+			char alt = Path.AltDirectorySeparatorChar;
 
 			Span<char> initial = stackalloc char[512];
 			using StackStringBuilder sb = new(initial);
@@ -99,7 +99,7 @@ namespace EpicGames.Core
 
 				while (idx <= span.Length)
 				{
-					int rel = (idx < span.Length) ? span.Slice(idx).IndexOfAny(Sep, Alt) : -1;
+					int rel = (idx < span.Length) ? span.Slice(idx).IndexOfAny(sep, alt) : -1;
 					int segLen = (rel < 0) ? (span.Length - idx) : rel;
 
 					if (segLen == 0)
@@ -115,7 +115,7 @@ namespace EpicGames.Core
 						// Pop last segment; if none, leave as-is
 						for (int i = sb.Length - 1; i >= 0; i--)
 						{
-							if (sb[i] == Sep)
+							if (sb[i] == sep)
 							{
 								sb.Resize(i);
 								break;
@@ -127,9 +127,9 @@ namespace EpicGames.Core
 						// Append separator only if there's already content and the last char isn't a separator
 						if (sb.Length > 0)
 						{
-							if (sb[sb.Length - 1] != Sep)
+							if (sb[sb.Length - 1] != sep)
 							{
-								sb.Append(Sep);
+								sb.Append(sep);
 							}
 						}
 
diff --git a/Engine/Source/Programs/Shared/EpicGames.Core/StackStringBuffer.cs b/Engine/Source/Programs/Shared/EpicGames.Core/StackStringBuffer.cs
--- a/Engine/Source/Programs/Shared/EpicGames.Core/StackStringBuffer.cs
+++ b/Engine/Source/Programs/Shared/EpicGames.Core/StackStringBuffer.cs
@@ -2,7 +2,6 @@
 
 using System;
 using System.Buffers;
-using System.Runtime.CompilerServices;
 
 namespace EpicGames.Core
 {
@@ -10,9 +9,8 @@ namespace EpicGames.Core
 	/// StringBuffer that can only live on the stack and does not make any allocations as long as initial buffer is big enough
 	/// Usage:
 	/// Span<char> initial = stackalloc char[512];
-	/// StackStringBuilder sb = new(initial);
+	/// using StackStringBuilder sb = new(initial);
 	/// sb.Append("Foo");
-	/// sb.Dispose(); // Needed in case StackStringBuffer reallocated
 	/// </summary>
 	ref struct StackStringBuilder
 	{
