[EpicGames.Core]
* Escaped <> inside summary documentation info to make static analysis happy

diff --git a/Engine/Source/Programs/Shared/EpicGames.Core/StackStringBuffer.cs b/Engine/Source/Programs/Shared/EpicGames.Core/StackStringBuffer.cs
--- a/Engine/Source/Programs/Shared/EpicGames.Core/StackStringBuffer.cs
+++ b/Engine/Source/Programs/Shared/EpicGames.Core/StackStringBuffer.cs
@@ -8,9 +8,9 @@ namespace EpicGames.Core
 	/// <summary>
 	/// StringBuffer that can only live on the stack and does not make any allocations as long as initial buffer is big enough
 	/// Usage:
-	/// Span<char> initial = stackalloc char[512];
-	/// using StackStringBuilder sb = new(initial);
-	/// sb.Append("Foo");
+	///   Span&lt;char&gt; initial = stackalloc char[512];
+	///   using StackStringBuilder sb = new(initial);
+	///   sb.Append("Foo");
 	/// </summary>
 	ref struct StackStringBuilder
 	{
