[EpicGames.Build]
* Implemented platform optimized DirectoryEnumerator that is using NtCreateFile directly on windows (and DirectoryInfo.EnumerateFileSystemEntries on non-windows). It speeds up search a little bit but above all fix so it is possible to start ubt with debugger attached without having to wait a minute for file discovery.
* Changed DirectoryItem and FileItem to use DirectoryEnumerator instead of FileInfo and DirectoryInfo

diff --git a/Engine/Source/Programs/Shared/EpicGames.Build/System/DirectoryEnumerator.cs b/Engine/Source/Programs/Shared/EpicGames.Build/System/DirectoryEnumerator.cs
new file mode 100644
--- /dev/null
+++ b/Engine/Source/Programs/Shared/EpicGames.Build/System/DirectoryEnumerator.cs
@@ -0,0 +1,366 @@
+// Copyright Epic Games, Inc. All Rights Reserved.
+
+using System;
+using System.Buffers;
+using System.Collections.Generic;
+using System.ComponentModel;
+using System.IO;
+using System.IO.Enumeration;
+using System.Runtime.CompilerServices;
+using System.Runtime.InteropServices;
+
+namespace UnrealBuildBase
+{
+	/// <summary>
+	/// Platform optimized utility class to enumerate directory to fetch information needed by FileItem and DirectoryItem
+	/// </summary>
+	public static class DirectoryEnumerator
+	{
+		/// <summary>
+		/// Struct representing a file system entry
+		/// </summary>
+		/// <param name="Name"></param>
+		/// <param name="Size"></param>
+		/// <param name="Attributes"></param>
+		/// <param name="CreationTimeUtc"></param>
+		/// <param name="LastWriteTimeUtc"></param>
+		public readonly record struct Entry(string Name, long Size, FileAttributes Attributes, DateTime CreationTimeUtc, DateTime LastWriteTimeUtc);
+
+		/// <summary>
+		/// Enumerate a directory. If directory does not exist it will return null
+		/// </summary>
+		/// <param name="directoryPath">Path to directory</param>
+		/// <param name="bufferSize">Buffer size used internally to reduce number of kernel calls</param>
+		/// <returns>List of files</returns>
+		public static unsafe List<Entry>? Enumerate(string directoryPath, int bufferSize = 128 * 1024)
+		{
+			if (OperatingSystem.IsWindows())
+			{
+				string ntPath;
+				if (directoryPath[0] == '\\')
+				{
+					ntPath = @"\??\UNC\" + directoryPath.TrimStart('\\'); // UNC -> \??\UNC\server\share\...
+				}
+				else
+				{
+					ntPath = @"\??\" + directoryPath; // Drive -> \??\C:\...
+				}
+
+				IntPtr dirHandle = IntPtr.Zero;
+
+				fixed (char* p = ntPath)
+				{
+					UNICODE_STRING name = new()
+					{
+						Length = (ushort)(ntPath.Length * sizeof(char)),
+						MaximumLength = (ushort)((ntPath.Length + 1) * sizeof(char)),
+						Buffer = p
+					};
+
+					OBJECT_ATTRIBUTES oa = new()
+					{
+						Length = (uint)sizeof(OBJECT_ATTRIBUTES),
+						RootDirectory = IntPtr.Zero,
+						ObjectName = &name,
+						Attributes = OBJ_CASE_INSENSITIVE,
+						SecurityDescriptor = IntPtr.Zero,
+						SecurityQualityOfService = IntPtr.Zero
+					};
+
+					long alloc = 0;
+
+					int status = NtCreateFile(out dirHandle, FILE_LIST_DIRECTORY | SYNCHRONIZE, ref oa, out IO_STATUS_BLOCK iosb, ref alloc, 0u, FILE_SHARE_READ | FILE_SHARE_WRITE | FILE_SHARE_DELETE, FILE_OPEN, FILE_DIRECTORY_FILE | FILE_SYNCHRONOUS_IO_NONALERT | FILE_OPEN_FOR_BACKUP_INTENT, IntPtr.Zero, 0u);
+					if (status != STATUS_SUCCESS)
+					{
+						return null; //ThrowIo(status, $"NtCreateFile('{directoryPath}')");
+					}
+				}
+
+				try
+				{
+					List<Entry> entries = new(256);
+
+					if (bufferSize <= 128 * 1024)
+					{
+						byte* buf = stackalloc byte[bufferSize];
+						EnumerateWithBuffer(dirHandle, buf, bufferSize, entries);
+					}
+					else
+					{
+						byte[] rented = ArrayPool<byte>.Shared.Rent(bufferSize);
+						try
+						{
+							fixed (byte* buf = rented)
+							{
+								EnumerateWithBuffer(dirHandle, buf, bufferSize, entries);
+							}
+						}
+						finally
+						{
+							ArrayPool<byte>.Shared.Return(rented);
+						}
+					}
+
+					return entries;
+				}
+				finally
+				{
+					if (dirHandle != IntPtr.Zero)
+					{
+						int v = NtClose(dirHandle);
+					}
+				}
+			}
+			else
+			{
+				List<Entry> entries = new(256);
+				foreach (FileSystemInfo entry in (new DirectoryInfo(directoryPath)).EnumerateFileSystemInfos())
+				{
+					long size = (entry.Attributes & FileAttributes.Directory) == 0 ? ((FileInfo)entry).Length : 0;
+					entries.Add(new Entry(entry.Name, size, entry.Attributes, entry.CreationTimeUtc, entry.LastWriteTimeUtc));
+				}
+
+				return entries;
+			}
+		}
+
+		/// <summary>
+		/// Try to get file from path
+		/// </summary>
+		/// <param name="path"></param>
+		/// <param name="entry"></param>
+		/// <returns>Returns false if file is not found or is a directory</returns>
+		public static bool TryGetFileFromPath(string path, out Entry entry)
+		{
+			if (OperatingSystem.IsWindows())
+			{
+				return TryGetEntryFromPath(path, out entry) && (entry.Attributes & FileAttributes.Directory) == 0;
+			}
+			FileInfo info = new(path);
+			if (!info.Exists)
+			{
+				entry = default;
+				return false;
+			}
+			entry = new("", info.Length, info.Attributes, info.CreationTimeUtc, info.LastWriteTimeUtc);
+			return true;
+		}
+
+		/// <summary>
+		/// Try to get directory from path
+		/// </summary>
+		/// <param name="path"></param>
+		/// <param name="entry"></param>
+		/// <returns>Returns false if directory is not found or is a file</returns>
+		public static bool TryGetDirectoryFromPath(string path, out Entry entry)
+		{
+			if (OperatingSystem.IsWindows())
+			{
+				return TryGetEntryFromPath(path, out entry) && (entry.Attributes & FileAttributes.Directory) != 0;
+			}
+			DirectoryInfo info = new(path);
+			if (!info.Exists)
+			{
+				entry = default;
+				return false;
+			}
+			entry = new("", 0, info.Attributes, info.CreationTimeUtc, info.LastWriteTimeUtc);
+			return true;
+		}
+
+		private static bool TryGetEntryFromPath(string path, out Entry entry)
+		{
+			if (!GetFileAttributesExW(path, GET_FILEEX_INFO_LEVELS.GetFileExInfoStandard, out WIN32_FILE_ATTRIBUTE_DATA data))
+			{
+				entry = default;
+				return false;
+			}
+
+			FileAttributes attrs = (FileAttributes)data.dwFileAttributes;
+			long size = ((long)data.nFileSizeHigh << 32) | data.nFileSizeLow;
+			DateTime c = DateTimeOffset.FromFileTime(((long)data.ftCreationTime.dwHighDateTime << 32) | data.ftCreationTime.dwLowDateTime).UtcDateTime;
+			DateTime w = DateTimeOffset.FromFileTime(((long)data.ftLastWriteTime.dwHighDateTime << 32) | data.ftLastWriteTime.dwLowDateTime).UtcDateTime;
+
+			entry = new("", size, attrs, c, w);
+			return true;
+		}
+
+		private static unsafe void EnumerateWithBuffer(IntPtr dirHandle, byte* buf, int bufferSize, List<Entry> entries)
+		{
+			bool restart = true;
+
+			while (true)
+			{
+				int status = NtQueryDirectoryFile(dirHandle, IntPtr.Zero, IntPtr.Zero, IntPtr.Zero, out IO_STATUS_BLOCK iosbLocal, buf, (uint)bufferSize, FILE_INFORMATION_CLASS.FileIdBothDirectoryInformation, false, null, restart);
+
+				if (status == STATUS_NO_MORE_FILES)
+				{
+					break;
+				}
+
+				if (status != STATUS_SUCCESS)
+				{
+					ThrowIo(status, "NtQueryDirectoryFile");
+				}
+
+				byte* cursor = buf;
+				while (true)
+				{
+					FILE_ID_BOTH_DIR_INFORMATION* info = (FILE_ID_BOTH_DIR_INFORMATION*)cursor;
+
+					int nameBytes = checked((int)info->FileNameLength);
+					char* namePtr = &info->FileName;
+					int nameChars = nameBytes / 2;
+
+					if (!(nameChars == 1 && namePtr[0] == '.') && !(nameChars == 2 && namePtr[0] == '.' && namePtr[1] == '.'))
+					{
+						string name = new string(namePtr, 0, nameChars);
+						entries.Add(new Entry(name, info->EndOfFile, (FileAttributes)info->FileAttributes, DateTimeOffset.FromFileTime(info->CreationTime).UtcDateTime, DateTimeOffset.FromFileTime(info->LastWriteTime).UtcDateTime));
+					}
+
+					if (info->NextEntryOffset == 0)
+					{
+						break;
+					}
+
+					cursor += info->NextEntryOffset;
+				}
+
+				restart = false;
+			}
+		}
+
+		private static void ThrowIo(int ntstatus, string op)
+		{
+			int win32 = RtlNtStatusToDosError(ntstatus);
+			Win32Exception ex = new(unchecked(win32));
+			throw new IOException($"{op} Failed: 0x{ntstatus:X8} ({ex.Message})", ex);
+		}
+
+		private const int STATUS_SUCCESS = 0x00000000;
+		private const int STATUS_NO_MORE_FILES = unchecked((int)0x80000006);
+
+		private const uint OBJ_CASE_INSENSITIVE = 0x00000040;
+
+		private const uint FILE_LIST_DIRECTORY = 0x0001;
+		private const uint FILE_SHARE_READ = 0x00000001;
+		private const uint FILE_SHARE_WRITE = 0x00000002;
+		private const uint FILE_SHARE_DELETE = 0x00000004;
+		private const uint FILE_OPEN = 0x00000001;
+		private const uint FILE_DIRECTORY_FILE = 0x00000001;
+		private const uint FILE_SYNCHRONOUS_IO_NONALERT = 0x00000020;
+		private const uint FILE_OPEN_FOR_BACKUP_INTENT = 0x00004000;
+		private const uint SYNCHRONIZE = 0x00100000;
+
+		[DllImport("ntdll.dll")]
+		private static extern int NtClose(IntPtr handle);
+
+		[DllImport("ntdll.dll", ExactSpelling = true)]
+		private static extern int RtlNtStatusToDosError(int status);
+
+		[DllImport("ntdll.dll", ExactSpelling = true)]
+		private static extern unsafe int NtCreateFile(
+			out IntPtr FileHandle,
+			uint DesiredAccess,
+			ref OBJECT_ATTRIBUTES ObjectAttributes,
+			out IO_STATUS_BLOCK IoStatusBlock,
+			ref long AllocationSize,
+			uint FileAttributes,
+			uint ShareAccess,
+			uint CreateDisposition,
+			uint CreateOptions,
+			IntPtr EaBuffer,
+			uint EaLength);
+
+		[DllImport("ntdll.dll", ExactSpelling = true)]
+		private static extern unsafe int NtQueryDirectoryFile(
+			IntPtr FileHandle,
+			IntPtr Event,
+			IntPtr ApcRoutine,
+			IntPtr ApcContext,
+			out IO_STATUS_BLOCK IoStatusBlock,
+			void* FileInformation,
+			uint Length,
+			FILE_INFORMATION_CLASS FileInformationClass,
+			[MarshalAs(UnmanagedType.U1)] bool ReturnSingleEntry,
+			UNICODE_STRING* FileName,
+			[MarshalAs(UnmanagedType.U1)] bool RestartScan);
+
+		[DllImport("kernel32.dll", CharSet = CharSet.Unicode, SetLastError = true)]
+		private static extern bool GetFileAttributesExW(string lpFileName, GET_FILEEX_INFO_LEVELS fInfoLevelId, out WIN32_FILE_ATTRIBUTE_DATA fileData);
+
+		private enum GET_FILEEX_INFO_LEVELS : int
+		{
+			GetFileExInfoStandard = 0
+		}
+
+		[StructLayout(LayoutKind.Sequential)]
+		private struct FILETIME
+		{
+			public uint dwLowDateTime;
+			public uint dwHighDateTime;
+		}
+
+		[StructLayout(LayoutKind.Sequential)]
+		private struct WIN32_FILE_ATTRIBUTE_DATA
+		{
+			public uint dwFileAttributes;
+			public FILETIME ftCreationTime;
+			public FILETIME ftLastAccessTime;
+			public FILETIME ftLastWriteTime;
+			public uint nFileSizeHigh;
+			public uint nFileSizeLow;
+		}
+
+		[StructLayout(LayoutKind.Sequential)]
+		private unsafe struct UNICODE_STRING
+		{
+			public ushort Length;
+			public ushort MaximumLength;
+			public char* Buffer;
+		}
+
+		[StructLayout(LayoutKind.Sequential)]
+		private unsafe struct OBJECT_ATTRIBUTES
+		{
+			public uint Length;
+			public IntPtr RootDirectory;
+			public UNICODE_STRING* ObjectName;
+			public uint Attributes;
+			public IntPtr SecurityDescriptor;
+			public IntPtr SecurityQualityOfService;
+		}
+
+		[StructLayout(LayoutKind.Sequential)]
+		private struct IO_STATUS_BLOCK
+		{
+			public IntPtr Status;
+			public nuint Information;
+		}
+
+		private enum FILE_INFORMATION_CLASS : int
+		{
+			FileIdBothDirectoryInformation = 37
+		}
+
+		[StructLayout(LayoutKind.Sequential)]
+		private unsafe struct FILE_ID_BOTH_DIR_INFORMATION
+		{
+			public uint NextEntryOffset;
+			public uint FileIndex;
+			public long CreationTime;
+			public long LastAccessTime;
+			public long LastWriteTime;
+			public long ChangeTime;
+			public long EndOfFile;
+			public long AllocationSize;
+			public uint FileAttributes;
+			public uint FileNameLength;
+			public uint EaSize;
+			public byte ShortNameLength;
+			private fixed char ShortName[12];
+			public long FileId;
+			public char FileName;
+		}
+	}
+}
diff --git a/Engine/Source/Programs/Shared/EpicGames.Build/System/DirectoryItem.cs b/Engine/Source/Programs/Shared/EpicGames.Build/System/DirectoryItem.cs
--- a/Engine/Source/Programs/Shared/EpicGames.Build/System/DirectoryItem.cs
+++ b/Engine/Source/Programs/Shared/EpicGames.Build/System/DirectoryItem.cs
@@ -22,54 +22,19 @@ namespace UnrealBuildBase
 		public readonly DirectoryReference Location;
 
 		/// <summary>
-		/// Cached value for whether the directory exists
-		/// </summary>
-		Lazy<DirectoryInfo> Info;
-
-		/// <summary>
-		/// Cached maps of name to subdirectory and name to file
-		/// </summary>
-		Lazy<(Dictionary<string, DirectoryItem>, Dictionary<string, FileItem>)> Cache;
-
-		/// <summary>
-		/// Global map of location to item
+		/// The name of this directory
 		/// </summary>
-		static ConcurrentDictionary<DirectoryReference, DirectoryItem> LocationToItem = [];
-
-		/// <summary>
-		/// Constructor
-		/// </summary>
-		/// <param name="Location">Path to this directory</param>
-		/// <param name="Info">Information about this directory</param>
-		private DirectoryItem(DirectoryReference Location, DirectoryInfo Info)
-		{
-			this.Location = Location;
-			Cache = new(Scan);
-
-			// For some reason we need to call an extra Refresh on Linux/Mac to not get wrong results from "Exists"
-			this.Info = OperatingSystem.IsWindows()
-				? new(Info)
-				: new(() =>
-				{
-					Info.Refresh();
-					return Info;
-				});
-		}
+		public readonly string Name;
 
 		/// <summary>
 		/// Accessor for map of name to subdirectory item
 		/// </summary>
-		Dictionary<string, DirectoryItem> Directories => Cache.Value.Item1;
+		Dictionary<string, DirectoryItem> Directories => _cache.Value.Item1;
 
 		/// <summary>
 		/// Accessor for map of name to file
 		/// </summary>
-		Dictionary<string, FileItem> Files => Cache.Value.Item2;
-
-		/// <summary>
-		/// The name of this directory
-		/// </summary>
-		public string Name => Info.Value.Name;
+		Dictionary<string, FileItem> Files => _cache.Value.Item2;
 
 		/// <summary>
 		/// The full name of this directory
@@ -79,31 +44,25 @@ namespace UnrealBuildBase
 		/// <summary>
 		/// Whether the directory exists or not
 		/// </summary>
-		public bool Exists => Info.Value.Exists;
+		public bool Exists => Update().Exists;
 
 		/// <summary>
 		/// The last write time of the file.
 		/// </summary>
-		public DateTime LastWriteTimeUtc => Info.Value.LastWriteTimeUtc;
+		public DateTime LastWriteTimeUtc => Update().LastWriteTimeUtc;
 
 		/// <summary>
 		/// The creation time of the file.
 		/// </summary>
-		public DateTime CreationTimeUtc => Info.Value.CreationTimeUtc;
+		public DateTime CreationTimeUtc => Update().CreationTimeUtc;
 
 		/// <summary>
 		/// Gets the parent directory item
 		/// </summary>
 		public DirectoryItem? GetParentDirectoryItem()
 		{
-			if (Info.Value.Parent == null)
-			{
-				return null;
-			}
-			else
-			{
-				return GetItemByDirectoryInfo(Info.Value.Parent);
-			}
+			DirectoryReference? parent = Location.ParentDirectory;
+			return parent != null ? DirectoryItem.GetItemByDirectoryReference(parent) : null;
 		}
 
 		/// <summary>
@@ -128,9 +87,7 @@ namespace UnrealBuildBase
 		/// <param name="Location">Path to the directory</param>
 		/// <returns>The directory item for this location</returns>
 		public static DirectoryItem GetItemByDirectoryReference(DirectoryReference Location)
-			=> LocationToItem.TryGetValue(Location, out DirectoryItem? Result)
-				? Result
-				: LocationToItem.GetOrAdd(Location, new DirectoryItem(Location, new(Location.FullName)));
+			=> LocationToItem.GetOrAdd(Location, static loc => new DirectoryItem(loc, loc.GetDirectoryName(), null));
 
 		/// <summary>
 		/// Finds or creates a directory item from a DirectoryInfo object
@@ -144,16 +101,11 @@ namespace UnrealBuildBase
 		/// </summary>
 		public void ResetCachedInfo()
 		{
-			Info = new(() =>
-			{
-				DirectoryInfo Info = Location.ToDirectoryInfo();
-				Info.Refresh();
-				return Info;
-			});
+			_info = null;
 
-			if (Cache.IsValueCreated)
+			if (_cache.IsValueCreated)
 			{
-				(Dictionary<string, DirectoryItem> dirs, Dictionary<string, FileItem> files) = Cache.Value;
+				(Dictionary<string, DirectoryItem> dirs, Dictionary<string, FileItem> files) = _cache.Value;
 				foreach (DirectoryItem SubDirectory in dirs.Values)
 				{
 					SubDirectory.ResetCachedInfo();
@@ -162,7 +114,7 @@ namespace UnrealBuildBase
 				{
 					File.ResetCachedInfo();
 				}
-				Cache = new(Scan);
+				_cache = new(Scan);
 			}
 		}
 
@@ -184,13 +136,8 @@ namespace UnrealBuildBase
 		{
 			LocationToItem.Values.AsParallel().ForAll(Item =>
 			{
-				Item.Info = new Lazy<DirectoryInfo>(() =>
-				{
-					DirectoryInfo Info = Item.Location.ToDirectoryInfo();
-					Info.Refresh();
-					return Info;
-				});
-				Item.Cache = new(Item.Scan);
+				Item._info = null;
+				Item._cache = new(Item.Scan);
 			});
 			FileItem.ResetAllCachedInfo_SLOW();
 		}
@@ -198,7 +145,7 @@ namespace UnrealBuildBase
 		/// <summary>
 		/// Caches the subdirectories of this directories
 		/// </summary>
-		public bool CacheDirectories() => Cache.Value.Item1 != null;
+		public bool CacheDirectories() => _cache.Value.Item1 != null;
 
 		/// <summary>
 		/// Enumerates all the subdirectories
@@ -236,45 +183,88 @@ namespace UnrealBuildBase
 			return Directories.TryGetValue(Name, out OutDirectory);
 		}
 
+		/// <summary>
+		/// Constructor
+		/// </summary>
+		/// <param name="Location">Path to this directory</param>
+		private DirectoryItem(DirectoryReference location, string name, Info? info)
+		{
+			Location = location;
+			Name = name;
+			_info = info;
+
+			_cache = new(Scan);
+		}
+
+		Info Update()
+		{
+			if (_info != null)
+			{
+				return _info;
+			}
+			Info newInfo = new();
+			if (DirectoryEnumerator.TryGetDirectoryFromPath(FullName, out DirectoryEnumerator.Entry entry))
+			{
+				newInfo.Exists = true;
+				newInfo.LastWriteTimeUtc = entry.LastWriteTimeUtc;
+				newInfo.CreationTimeUtc = entry.CreationTimeUtc;
+			}
+			_info = newInfo;
+			return newInfo;
+		}
+
 		/// <summary>
 		/// Scans the directory for directories and files, used for lazy initialization.
 		/// </summary>
 		/// <returns></returns>
 		(Dictionary<string, DirectoryItem> , Dictionary<string, FileItem>) Scan()
 		{
+			if (_info != null && !_info.Exists)
+			{
+				return ([], []);
+			}
+
 			try
 			{
-				// We want to turn enumerator to a list here since using EnumerateFiles.Count and then enumerate EnumerateFiles cause two iterations of findfirst/findnext
-				// Also, we don't check if exists to minimize kernel call count
-				List<FileSystemInfo> infos = Info.Value.EnumerateFileSystemInfos().ToList();
+				List<DirectoryEnumerator.Entry>? entries = DirectoryEnumerator.Enumerate(FullName);
 
+				if (entries == null)
+				{
+					_info = new Info() { Exists = false };
+					return ([], []);
+				}
 				int dirCount = 0;
-				foreach (FileSystemInfo info in infos)
+				foreach (DirectoryEnumerator.Entry entry in entries)
 				{
-					if (info.Attributes.HasFlag(FileAttributes.Directory))
+					if ((entry.Attributes & FileAttributes.Directory) != 0)
 					{
 						++dirCount;
 					}
 				}
 
 				Dictionary<string, DirectoryItem>? newDirs = new(dirCount, FileReference.Comparer);
-				Dictionary<string, FileItem>? newFiles = new(infos.Count - dirCount, FileReference.Comparer);
+				Dictionary<string, FileItem>? newFiles = new(entries.Count - dirCount, FileReference.Comparer);
 
-				foreach (FileSystemInfo info in infos)
+				foreach (DirectoryEnumerator.Entry entry in entries)
 				{
-					if (info.Attributes.HasFlag(FileAttributes.Directory))
+					if ((entry.Attributes & FileAttributes.Directory) != 0)
 					{
-						newDirs.Add(info.Name, DirectoryItem.GetItemByDirectoryInfo((DirectoryInfo)info));
+						Info info = new() { Exists = true, LastWriteTimeUtc = entry.LastWriteTimeUtc, CreationTimeUtc = entry.CreationTimeUtc };
+						DirectoryReference loc = DirectoryReference.Combine(Location, entry.Name);
+						DirectoryItem dir = LocationToItem.GetOrAdd(loc, loc => new DirectoryItem(loc, entry.Name, info));
+						newDirs.Add(entry.Name, dir);
 					}
 					else
 					{
-						FileItem FileItem = FileItem.GetItemByFileInfo((FileInfo)info);
+						FileItem.Info info = new() { Attributes = entry.Attributes, Length = entry.Size, LastWriteTimeUtc = entry.LastWriteTimeUtc, CreationTimeUtc = entry.CreationTimeUtc };
+						FileReference loc = FileReference.Combine(Location, entry.Name);
+						FileItem file = FileItem.UniqueSourceFileMap.GetOrAdd(loc, loc => new FileItem(loc, entry.Name, info));
 
 						// There are folders in linux sdk that has files with same name but different casing.
 						// Ideally FileReference.Comparer should be case sensitive on linux/mac but I don't dare changing that right now
-						if (newFiles.TryAdd(info.Name, FileItem))
+						if (newFiles.TryAdd(entry.Name, file))
 						{
-							FileItem.UpdateCachedDirectory(this);
+							file.UpdateCachedDirectory(this);
 						}
 					}
 				}
@@ -297,7 +287,7 @@ namespace UnrealBuildBase
 		/// <summary>
 		/// Caches the files in this directory
 		/// </summary>
-		public bool CacheFiles() => Cache.Value.Item2 != null;
+		public bool CacheFiles() => _cache.Value.Item2 != null;
 
 		/// <summary>
 		/// Enumerates all the files
@@ -418,6 +408,17 @@ namespace UnrealBuildBase
 			return left is null ? right is null : left.CompareTo(right) >= 0;
 		}
 		#endregion
+
+		class Info
+		{
+			internal bool Exists;
+			internal DateTime LastWriteTimeUtc;
+			internal DateTime CreationTimeUtc;
+		}
+
+		Info? _info;
+		Lazy<(Dictionary<string, DirectoryItem>, Dictionary<string, FileItem>)> _cache;
+		static ConcurrentDictionary<DirectoryReference, DirectoryItem> LocationToItem = [];
 	}
 
 	/// <summary>
diff --git a/Engine/Source/Programs/Shared/EpicGames.Build/System/FileItem.cs b/Engine/Source/Programs/Shared/EpicGames.Build/System/FileItem.cs
--- a/Engine/Source/Programs/Shared/EpicGames.Build/System/FileItem.cs
+++ b/Engine/Source/Programs/Shared/EpicGames.Build/System/FileItem.cs
@@ -28,40 +28,10 @@ namespace UnrealBuildBase
 		/// </summary>
 		public readonly FileReference Location;
 
-		/// <summary>
-		/// The information about the file.
-		/// </summary>
-		Lazy<FileInfo> Info;
-
-		/// <summary>
-		/// A case-insensitive dictionary that's used to map each unique file name to a single FileItem object.
-		/// </summary>
-		static ConcurrentDictionary<FileReference, FileItem> UniqueSourceFileMap = new(-1, 100000);// [];
-
-		/// <summary>
-		/// Constructor
-		/// </summary>
-		/// <param name="Location">Location of the file</param>
-		/// <param name="Info">File info</param>
-		private FileItem(FileReference Location, FileInfo Info)
-		{
-			this.Location = Location;
-			CachedDirectory = new(() => DirectoryItem.GetItemByDirectoryReference(Location.Directory));
-
-			// For some reason we need to call an extra Refresh on Linux/Mac to not get wrong results from "Exists"
-			this.Info = OperatingSystem.IsWindows()
-				? new(Info)
-				: new(() =>
-				{
-					Info.Refresh();
-					return Info;
-				});
-		}
-
 		/// <summary>
 		/// Name of this file
 		/// </summary>
-		public string Name => Info.Value.Name;
+		public readonly string Name;
 
 		/// <summary>
 		/// Full name of this file
@@ -81,27 +51,27 @@ namespace UnrealBuildBase
 		/// <summary>
 		/// Whether the file exists.
 		/// </summary>
-		public bool Exists => Info.Value.Exists;
+		public bool Exists => Update().Attributes != 0;
 
 		/// <summary>
 		/// Size of the file if it exists, otherwise -1
 		/// </summary>
-		public long Length => Info.Value.Length;
+		public long Length => Update().Length;
 
 		/// <summary>
 		/// The attributes for this file
 		/// </summary>
-		public FileAttributes Attributes => Info.Value.Attributes;
+		public FileAttributes Attributes => Update().Attributes;
 
 		/// <summary>
 		/// The last write time of the file.
 		/// </summary>
-		public DateTime LastWriteTimeUtc => Info.Value.LastWriteTimeUtc;
+		public DateTime LastWriteTimeUtc => Update().LastWriteTimeUtc;
 
 		/// <summary>
 		/// The creation time of the file.
 		/// </summary>
-		public DateTime CreationTimeUtc => Info.Value.CreationTimeUtc;
+		public DateTime CreationTimeUtc => Update().CreationTimeUtc;
 
 		/// <summary>
 		/// Determines if the file has the given extension
@@ -146,7 +116,7 @@ namespace UnrealBuildBase
 		/// <param name="Location">Location of the file</param>
 		/// <returns>The FileItem that represents the given a full file path.</returns>
 		public static FileItem GetItemByFileReference(FileReference Location)
-			=> UniqueSourceFileMap.GetOrAdd(Location, static loc => new FileItem(loc, loc.ToFileInfo()));
+			=> UniqueSourceFileMap.GetOrAdd(Location, static loc => new FileItem(loc, loc.GetFileName(), null));
 
 		/// <summary>
 		/// Deletes the file.
@@ -200,12 +170,7 @@ namespace UnrealBuildBase
 		/// </summary>
 		public void ResetCachedInfo()
 		{
-			Info = new(() =>
-			{
-				FileInfo Info = Location.ToFileInfo();
-				Info.Refresh();
-				return Info;
-			});
+			_info = null;
 		}
 
 		/// <summary>
@@ -294,6 +259,50 @@ namespace UnrealBuildBase
 			return left is null ? right is null : left.CompareTo(right) >= 0;
 		}
 		#endregion
+
+
+		internal FileItem(FileReference Location, string name, Info? info)
+		{
+			this.Location = Location;
+			Name = name;
+			_info = info;
+
+			CachedDirectory = new(() => DirectoryItem.GetItemByDirectoryReference(Location.Directory)); // Take this in
+		}
+
+		private Info Update()
+		{
+			if (_info != null)
+			{
+				return _info;
+			}
+
+			Info newInfo = new();
+			if (DirectoryEnumerator.TryGetFileFromPath(FullName, out DirectoryEnumerator.Entry entry))
+			{
+				newInfo.Attributes = entry.Attributes;
+				newInfo.Length = entry.Size;
+				newInfo.LastWriteTimeUtc = entry.LastWriteTimeUtc;
+				newInfo.CreationTimeUtc = entry.CreationTimeUtc;
+			}
+			_info = newInfo;
+			return newInfo;
+		}
+
+		/// <summary>
+		/// A case-insensitive dictionary that's used to map each unique file name to a single FileItem object.
+		/// </summary>
+		internal static ConcurrentDictionary<FileReference, FileItem> UniqueSourceFileMap = new(-1, 100000);// [];
+
+		internal class Info
+		{
+			internal FileAttributes Attributes;
+			internal long Length;
+			internal DateTime LastWriteTimeUtc;
+			internal DateTime CreationTimeUtc;
+		}
+
+		Info? _info;
 	}
 
 	/// <summary>
