[EpicGames.Core]
* Changed so FileSystemReference is not using Lazy<int> for hash code but instead lazily calculate it on GetHashCode call with an interlocked writing to the variable. It does not matter if this is done one extra time if same reference is used in multiple dictionaries

diff --git a/Engine/Source/Programs/Shared/EpicGames.Core/FileSystemReference.cs b/Engine/Source/Programs/Shared/EpicGames.Core/FileSystemReference.cs
--- a/Engine/Source/Programs/Shared/EpicGames.Core/FileSystemReference.cs
+++ b/Engine/Source/Programs/Shared/EpicGames.Core/FileSystemReference.cs
@@ -37,16 +37,32 @@ namespace EpicGames.Core
 		protected FileSystemReference(string fullName)
 		{
 			FullName = fullName;
-			_hashCode = new(() => Comparer.GetHashCode(FullName));
 		}
 
 		/// <inheritdoc/>
-		public override int GetHashCode() => _hashCode.Value;
+		public override int GetHashCode()
+		{
+			int h = _hashCode;
+			if (h != 0)
+			{
+				return h;
+			}
+
+			h = Comparer.GetHashCode(FullName);
+
+			if (h != 0)
+			{
+				Interlocked.CompareExchange(ref _hashCode, h, 0);
+				return _hashCode; // return whatever won the race
+			}
+
+			return 0;
+		}
 
 		/// <summary>
 		/// Lazy initialized hash code
 		/// </summary>
-		private readonly Lazy<int> _hashCode;
+		private int _hashCode;
 
 		/// <summary>
 		/// Create a full path by concatenating multiple strings
