[EpicGames.Build]
* Removed comments on private ctor because they were wrong and not needed

diff --git a/Engine/Source/Programs/Shared/EpicGames.Build/System/DirectoryItem.cs b/Engine/Source/Programs/Shared/EpicGames.Build/System/DirectoryItem.cs
--- a/Engine/Source/Programs/Shared/EpicGames.Build/System/DirectoryItem.cs
+++ b/Engine/Source/Programs/Shared/EpicGames.Build/System/DirectoryItem.cs
@@ -183,11 +183,7 @@ namespace UnrealBuildBase
 			return Directories.TryGetValue(Name, out OutDirectory);
 		}
 
-		/// <summary>
-		/// Constructor
-		/// </summary>
-		/// <param name="Location">Path to this directory</param>
-		private DirectoryItem(DirectoryReference location, string name, Info? info)
+		DirectoryItem(DirectoryReference location, string name, Info? info)
 		{
 			Location = location;
 			Name = name;
