[EpicGames.Core]
* Moved LegacyEventLogger to its own file

diff --git a/Engine/Source/Programs/Shared/EpicGames.Core/LegacyEventLogger.cs b/Engine/Source/Programs/Shared/EpicGames.Core/LegacyEventLogger.cs
new file mode 100644
--- /dev/null
+++ b/Engine/Source/Programs/Shared/EpicGames.Core/LegacyEventLogger.cs
@@ -0,0 +1,53 @@
+// Copyright Epic Games, Inc. All Rights Reserved.
+
+using System;
+using Microsoft.Extensions.Logging;
+
+namespace EpicGames.Core
+{
+	/// <summary>
+	/// Wrapper around a custom logger interface which flushes the event parser when switching between legacy
+	/// and native structured logging
+	/// </summary>
+	sealed class LegacyEventLogger : ILogger, IDisposable
+	{
+		private ILogger _inner;
+		private readonly LogEventParser _parser;
+
+		public LogEventParser Parser => _parser;
+
+		public LegacyEventLogger(ILogger inner)
+		{
+			_inner = inner;
+			_parser = new LogEventParser(inner);
+		}
+
+		public void Dispose()
+		{
+			_parser.Dispose();
+		}
+
+		public void SetInnerLogger(ILogger inner)
+		{
+			lock (_parser)
+			{
+				_parser.Flush();
+				_inner = inner;
+				_parser.Logger = inner;
+			}
+		}
+
+		public IDisposable? BeginScope<TState>(TState state) where TState : notnull => _inner.BeginScope(state);
+
+		public bool IsEnabled(LogLevel logLevel) => _inner.IsEnabled(logLevel);
+
+		public void Log<TState>(LogLevel logLevel, EventId eventId, TState state, Exception? exception, Func<TState, Exception?, string> formatter)
+		{
+			lock (_parser)
+			{
+				_parser.Flush();
+			}
+			_inner.Log(logLevel, eventId, state, exception, formatter);
+		}
+	}
+}
\ No newline at end of file
diff --git a/Engine/Source/Programs/Shared/EpicGames.Core/Log.cs b/Engine/Source/Programs/Shared/EpicGames.Core/Log.cs
--- a/Engine/Source/Programs/Shared/EpicGames.Core/Log.cs
+++ b/Engine/Source/Programs/Shared/EpicGames.Core/Log.cs
@@ -891,52 +891,6 @@ namespace EpicGames.Core
 		}
 	}
 
-	/// <summary>
-	/// Wrapper around a custom logger interface which flushes the event parser when switching between legacy
-	/// and native structured logging
-	/// </summary>
-	sealed class LegacyEventLogger : ILogger, IDisposable
-	{
-		private ILogger _inner;
-		private readonly LogEventParser _parser;
-
-		public LogEventParser Parser => _parser;
-
-		public LegacyEventLogger(ILogger inner)
-		{
-			_inner = inner;
-			_parser = new LogEventParser(inner);
-		}
-
-		public void Dispose()
-		{
-			_parser.Dispose();
-		}
-
-		public void SetInnerLogger(ILogger inner)
-		{
-			lock (_parser)
-			{
-				_parser.Flush();
-				_inner = inner;
-				_parser.Logger = inner;
-			}
-		}
-
-		public IDisposable? BeginScope<TState>(TState state) where TState : notnull => _inner.BeginScope(state);
-
-		public bool IsEnabled(LogLevel logLevel) => _inner.IsEnabled(logLevel);
-
-		public void Log<TState>(LogLevel logLevel, EventId eventId, TState state, Exception? exception, Func<TState, Exception?, string> formatter)
-		{
-			lock (_parser)
-			{
-				_parser.Flush();
-			}
-			_inner.Log(logLevel, eventId, state, exception, formatter);
-		}
-	}
-
 	/// <summary>
 	/// Default log output device
 	/// </summary>
