[EpicGames.Build]
* Changed UEBuildPlatformSDK ProcessJsonFile to use JsonObject instead JsonSerializer.Deserialize<>. Removes 1/3 of Registersdk time

diff --git a/Engine/Source/Programs/Shared/EpicGames.Build/UEBuildPlatformSDK.cs b/Engine/Source/Programs/Shared/EpicGames.Build/UEBuildPlatformSDK.cs
--- a/Engine/Source/Programs/Shared/EpicGames.Build/UEBuildPlatformSDK.cs
+++ b/Engine/Source/Programs/Shared/EpicGames.Build/UEBuildPlatformSDK.cs
@@ -1148,11 +1148,8 @@ namespace EpicGames.Core
 		});
 		private void ProcessJsonFile(FileReference SDKConfigFile, Dictionary<string, string> VersionMap, Dictionary<string, string[]> VersionArrayMap)
 		{
-			string Contents = FileReference.ReadAllText(SDKConfigFile);
-
-			// load the json into a dictionary
-			Dictionary<string, JsonElement>? LoadedDictionary = JsonSerializer.Deserialize<Dictionary<string, JsonElement>>(Contents, ProcessJsonFileOptions.Value);
-			if (LoadedDictionary == null)
+			// load the json
+			if (!JsonObject.TryRead(SDKConfigFile, out JsonObject? jsonObject))
 			{
 				throw new Exception($"Failed to parse SDK version file '{SDKConfigFile}'");
 			}
@@ -1161,13 +1158,9 @@ namespace EpicGames.Core
 
 			// look for special key to chain to other - after processing everything else (we only add from the "parent" settings not already in the map)
 			List<string> Parents = [];
-			foreach (string Key in LoadedDictionary.Keys)
+			foreach (string Key in jsonObject.KeyNames)
 			{
-				JsonElement Obj = LoadedDictionary[Key];
-				string? StringValue = Obj.ValueKind == JsonValueKind.String ? Obj.GetString() : null;
-				string[]? ArrayValue = Obj.ValueKind == JsonValueKind.Array ? Obj.EnumerateArray().Select(x => x.GetString()!).ToArray() : null;
-
-				if (StringValue != null)
+				if (jsonObject.TryGetStringField(Key, out string? StringValue))
 				{
 					if (Key.Equals("ParentSDKFile", StringComparison.OrdinalIgnoreCase))
 					{
@@ -1176,7 +1169,7 @@ namespace EpicGames.Core
 					// add this key if it's not already in the (case-insensitive) dictionary 
 					VersionMap.TryAdd(Key, StringValue);
 				}
-				else if (ArrayValue != null)
+				else if (jsonObject.TryGetStringArrayField(Key, out string[]? ArrayValue))
 				{
 					VersionArrayMap.Add(Key, ArrayValue);
 				}
