[EpicGames.Core]
* Added some documentation to methods

diff --git a/Engine/Source/Programs/Shared/EpicGames.Core/Log.cs b/Engine/Source/Programs/Shared/EpicGames.Core/Log.cs
--- a/Engine/Source/Programs/Shared/EpicGames.Core/Log.cs
+++ b/Engine/Source/Programs/Shared/EpicGames.Core/Log.cs
@@ -205,15 +205,17 @@ namespace EpicGames.Core
 			await s_defaultLogger.FlushAsync();
 		}
 
+		/// <summary>
+		/// Enable threaded logging. This can be used for performance reasons to prevent stalls when lots of logging is made from multiple threads
+		/// </summary>
 		public static void EnableThreaded()
 		{
 			s_legacyLogger.EnableThreaded();
 		}
 
 		/// <summary>
-		/// Finish async and fallback to sync logging
+		/// Disable threaded logging. will flush queued log entries and get back to synchronous logging
 		/// </summary>
-		/// <returns></returns>
 		public static void DisableThreaded()
 		{
 			s_legacyLogger.DisableThreaded();
@@ -1066,7 +1068,7 @@ namespace EpicGames.Core
 		}
 
 		/// <summary>
-		/// Flush the stream
+		/// Flush the log queue
 		/// </summary>
 		/// <returns></returns>
 		public async Task FlushAsync()
