[UBT/UBA]
* Changed so uba trace is initialized as early as possible and hooked into ubt timline system. This makes it possible to get a trace file that contains ubt specific things

diff --git a/Engine/Source/Programs/Shared/EpicGames.UBA/ITrace.cs b/Engine/Source/Programs/Shared/EpicGames.UBA/ITrace.cs
--- a/Engine/Source/Programs/Shared/EpicGames.UBA/ITrace.cs
+++ b/Engine/Source/Programs/Shared/EpicGames.UBA/ITrace.cs
@@ -40,17 +40,16 @@ namespace EpicGames.UBA
 		/// <param name="statusLink">Optional hyperlink that can be used to make text clickable in visualizer</param>
 		public abstract void UpdateStatus(uint statusRow, uint statusColumn, string statusText, LogEntryType statusType, string? statusLink = null);
 
-		/// <summary>
-		/// Writes trace to file. Only works if created with ITrace.Create
-		/// </summary>
-		public abstract void CloseAndWrite(string traceFileName);
-
 		/// <summary>
 		/// Creates a new trace
 		/// </summary>
-		public static ITrace Create(string? name = null)
+		/// <param name="name">Name of trace</param>
+		/// <param name="traceFile">File to write trace to</param>
+		/// <param name="makeGlobal">Set this trace as the one to be used by all sessions</param>
+		/// <param name="channel">Set which channel this trace should be visible to</param>
+		public static ITrace Create(string? name = null, string? traceFile = null, bool makeGlobal = false, string? channel = "Default")
 		{
-			return new TraceImpl(name);
+			return new TraceImpl(name, traceFile, makeGlobal, channel);
 		}
 	}
 
diff --git a/Engine/Source/Programs/Shared/EpicGames.UBA/Impl/SessionServerImpl.cs b/Engine/Source/Programs/Shared/EpicGames.UBA/Impl/SessionServerImpl.cs
--- a/Engine/Source/Programs/Shared/EpicGames.UBA/Impl/SessionServerImpl.cs
+++ b/Engine/Source/Programs/Shared/EpicGames.UBA/Impl/SessionServerImpl.cs
@@ -163,7 +163,14 @@ namespace EpicGames.UBA.Impl
 			SessionServer_SetRemoteProcessAvailable(_handle, _remoteProcessSlotAvailableCallbackDelegate, 0);
 			SessionServer_SetRemoteProcessReturned(_handle, _remoteProcessReturnedCallbackDelegate);
 
-			_trace = new TraceImpl(SessionServer_GetTrace(_handle));
+			if (TraceImpl.GlobalTrace != null)
+			{
+				_trace = TraceImpl.GlobalTrace;
+			}
+			else
+			{
+				_trace = new TraceImpl(SessionServer_GetTrace(_handle));
+			}
 
 			foreach (KeyValuePair<string, string> kv in Utils.CrossArchitecturePaths)
 			{
@@ -186,7 +193,10 @@ namespace EpicGames.UBA.Impl
 			{
 			}
 
-			_trace.Dispose();
+			if (TraceImpl.GlobalTrace == null)
+			{
+				_trace.Dispose();
+			}
 			
 			if (_handle != IntPtr.Zero)
 			{
diff --git a/Engine/Source/Programs/Shared/EpicGames.UBA/Impl/TraceImpl.cs b/Engine/Source/Programs/Shared/EpicGames.UBA/Impl/TraceImpl.cs
--- a/Engine/Source/Programs/Shared/EpicGames.UBA/Impl/TraceImpl.cs
+++ b/Engine/Source/Programs/Shared/EpicGames.UBA/Impl/TraceImpl.cs
@@ -8,6 +8,7 @@ namespace EpicGames.UBA
 	internal class TraceImpl : ITrace
 	{
 		nint _handle = IntPtr.Zero;
+		readonly string? _traceFile;
 		readonly bool _ownsHandle = false;
 
 		#region DllImport
@@ -24,23 +25,40 @@ namespace EpicGames.UBA
 		static extern void Trace_UpdateStatus(nint trace, uint statusRow, uint statusColumn, string statusText, byte statusType, string? statusLink);
 
 		[DllImport("UbaHost", CharSet = CharSet.Auto)]
-		static extern nint Trace_Create(string? name);
+		static extern nint Trace_Create(string? name, string? channel);
 
 		[DllImport("UbaHost", CharSet = CharSet.Auto)]
-		static extern void Trace_Destroy(nint server, string? writeFile);
+		static extern void Trace_Destroy(nint trace, string? writeFile);
+
+		[DllImport("UbaHost", CharSet = CharSet.Auto)]
+		static extern void Trace_SetGlobal(nint trace);
 		#endregion
 
 		internal TraceImpl(nint handle)
 		{
 			_handle = handle;
-			_ownsHandle = false;
 		}
 
-		internal TraceImpl(string? name)
+		internal TraceImpl(string? name, string? traceFile, bool makeGlobal, string? channel)
 		{
-			_handle = Trace_Create(name);
+			_traceFile = traceFile;
+			_handle = Trace_Create(name, channel);
+			_ownsHandle = true;
+			if (makeGlobal)
+			{
+				try
+				{
+					Trace_SetGlobal(_handle);
+					GlobalTrace = this;
+				} catch (EntryPointNotFoundException)
+				{
+				}
+			}
+
 		}
 
+		internal static TraceImpl? GlobalTrace { set; get; }
+
 		#region IDisposable
 		~TraceImpl() => Dispose(false);
 
@@ -58,9 +76,14 @@ namespace EpicGames.UBA
 
 			if (_handle != IntPtr.Zero)
 			{
+				if (GlobalTrace == this)
+				{
+					GlobalTrace = null;
+				}
+
 				if (_ownsHandle)
 				{
-					Trace_Destroy(_handle, null);
+					Trace_Destroy(_handle, _traceFile);
 				}
 				_handle = IntPtr.Zero;
 			}
@@ -76,14 +99,5 @@ namespace EpicGames.UBA
 		public void TaskEnd(uint taskId) => Trace_TaskEnd(_handle, taskId, true);
 
 		public void UpdateStatus(uint statusRow, uint statusColumn, string statusText, LogEntryType statusType, string? statusLink) => Trace_UpdateStatus(_handle, statusRow, statusColumn, statusText, (byte)statusType, statusLink);
-
-		public void CloseAndWrite(string traceFileName)
-		{
-			if (_ownsHandle && _handle != IntPtr.Zero)
-			{
-				Trace_Destroy(_handle, traceFileName);
-				_handle = IntPtr.Zero;
-			}
-		}
 	}
 }
\ No newline at end of file
diff --git a/Engine/Source/Programs/UnrealBuildTool/Modes/BuildMode.cs b/Engine/Source/Programs/UnrealBuildTool/Modes/BuildMode.cs
--- a/Engine/Source/Programs/UnrealBuildTool/Modes/BuildMode.cs
+++ b/Engine/Source/Programs/UnrealBuildTool/Modes/BuildMode.cs
@@ -817,7 +817,7 @@ namespace UnrealBuildTool
 
 						ActionExecutor.SetMemoryPerActionOverride(Makefiles.Select(x => x.MemoryPerActionGB).Max());
 
-						using ITimelineEvent ExcuteTimer = Timeline.ScopeEvent("Executing ActionGraph");
+						//using ITimelineEvent ExcuteTimer = Timeline.ScopeEvent("Executing ActionGraph");
 						using (GlobalTracer.Instance.BuildSpan("ActionGraph.ExecuteActions()").StartActive())
 						{
 
diff --git a/Engine/Source/Programs/UnrealBuildTool/System/Timeline.cs b/Engine/Source/Programs/UnrealBuildTool/System/Timeline.cs
--- a/Engine/Source/Programs/UnrealBuildTool/System/Timeline.cs
+++ b/Engine/Source/Programs/UnrealBuildTool/System/Timeline.cs
@@ -59,6 +59,8 @@ namespace UnrealBuildTool
 			/// </summary>
 			LogLevel? Verbosity;
 
+			uint Id = 0;
+
 			/// <summary>
 			/// Constructor
 			/// </summary>
@@ -75,6 +77,11 @@ namespace UnrealBuildTool
 				this.Logger = Logger;
 				this.Verbosity = Verbosity;
 				Span = TraceSpan.Create(Name);
+
+				if (OnEventBegin != null)
+				{
+					Id = OnEventBegin.Invoke(Name);
+				}
 			}
 
 			/// <summary>
@@ -84,6 +91,11 @@ namespace UnrealBuildTool
 			{
 				if (!FinishTime.HasValue)
 				{
+					if (Id != 0 && OnEventEnd != null)
+					{
+						OnEventEnd.Invoke(Id);
+					}
+
 					FinishTime = Stopwatch.Elapsed;
 					if (Logger != null && Verbosity != null)
 					{
@@ -103,6 +115,10 @@ namespace UnrealBuildTool
 			}
 		}
 
+		public static Func<string, uint>? OnEventBegin;
+		public static Action<uint>? OnEventEnd;
+
+
 		/// <summary>
 		/// The stopwatch used for timing
 		/// </summary>
@@ -321,4 +337,35 @@ namespace UnrealBuildTool
 			return String.Format("{0}.{1:000}", TotalMilliseconds / 1000, TotalMilliseconds % 1000);
 		}
 	}
+
+	/// <summary>
+	/// Will add and remove to OnEvent
+	/// </summary>
+	public readonly struct TimelineEventScope : IDisposable
+	{
+		readonly Func<string, uint>? _begin;
+		readonly Action<uint>? _end;
+
+		/// <summary>
+		/// Ctor
+		/// </summary>
+		/// <param name="begin">OnEventBegin</param>
+		/// <param name="end">OnEventEnd</param>
+		public TimelineEventScope(Func<string, uint>? begin, Action<uint>? end)
+		{
+			_begin = begin;
+			_end = end;
+			Timeline.OnEventBegin += _begin;
+			Timeline.OnEventEnd += _end;
+		}
+
+		/// <summary>
+		/// Dispose
+		/// </summary>
+		public void Dispose()
+		{
+			Timeline.OnEventBegin -= _begin;
+			Timeline.OnEventEnd -= _end;
+		}
+	}
 }
diff --git a/Engine/Source/Programs/UnrealBuildTool/UnrealBuildTool.cs b/Engine/Source/Programs/UnrealBuildTool/UnrealBuildTool.cs
--- a/Engine/Source/Programs/UnrealBuildTool/UnrealBuildTool.cs
+++ b/Engine/Source/Programs/UnrealBuildTool/UnrealBuildTool.cs
@@ -411,6 +411,32 @@ namespace UnrealBuildTool
                 Console.CancelKeyPress += CancelKeyPressAsync;
             }
 
+			// If command line contains -Session it is a recursive action and we don't want to report this trace to the default trace channel
+			string? traceChannel = "Default";
+			if (ArgumentsArray.FirstOrDefault(s => s.Contains("-Session=")) != null)
+				traceChannel = null;
+
+			// Setup trace file and backup old one
+			FileReference ubaTraceFile;
+			if (Unreal.IsBuildMachine())
+			{
+				ubaTraceFile = FileReference.Combine(Unreal.EngineProgramSavedDirectory, "AutomationTool", "Saved", "Logs", "Trace.uba");
+			}
+			else if (Log.OutputFile != null)
+			{
+				ubaTraceFile = Log.OutputFile.ChangeExtension(".uba");
+			}
+			else
+			{
+				ubaTraceFile = FileReference.Combine(Unreal.EngineProgramSavedDirectory, "UnrealBuildTool", "Trace.uba");
+			}
+			DirectoryReference.CreateDirectory(ubaTraceFile.Directory);
+			Log.BackupLogFile(ubaTraceFile);
+
+			// Attach timeline to uba trace
+			using EpicGames.UBA.ITrace trace = EpicGames.UBA.ITrace.Create("UnrealBuildTool", ubaTraceFile.FullName, true, traceChannel);
+			using TimelineEventScope _ = new((name) => trace.TaskBegin(name, ""), (id) => trace.TaskEnd(id));
+
 			bool bShowTimeline = false;
 			try
 			{
@@ -656,7 +682,7 @@ namespace UnrealBuildTool
 
 				// Execute the mode
 				MainScope.Finish();
-				using ITimelineEvent ExecuteScope = Timeline.ScopeEvent(ModeType.Name);
+				//using ITimelineEvent ExecuteScope = Timeline.ScopeEvent(ModeType.Name);
 				int Result = Mode.ExecuteAsync(Arguments, Logger).GetAwaiter().GetResult();
 				ApplicationResult = (CompilationResult)Result;
 				return Result;
