[EpicGames.Build]
* Minor opt when creating FileItems from DirectoryItem.Scan.. Set DirectoryItem directly on FileItem in ctor

diff --git a/Engine/Source/Programs/Shared/EpicGames.Build/System/DirectoryItem.cs b/Engine/Source/Programs/Shared/EpicGames.Build/System/DirectoryItem.cs
--- a/Engine/Source/Programs/Shared/EpicGames.Build/System/DirectoryItem.cs
+++ b/Engine/Source/Programs/Shared/EpicGames.Build/System/DirectoryItem.cs
@@ -254,14 +254,11 @@ namespace UnrealBuildBase
 					{
 						FileItem.Info info = new() { Attributes = entry.Attributes, Length = entry.Size, LastWriteTimeUtc = entry.LastWriteTimeUtc, CreationTimeUtc = entry.CreationTimeUtc };
 						FileReference loc = FileReference.Combine(Location, entry.Name);
-						FileItem file = FileItem.UniqueSourceFileMap.GetOrAdd(loc, loc => new FileItem(loc, entry.Name, info));
+						FileItem file = FileItem.UniqueSourceFileMap.GetOrAdd(loc, loc => new FileItem(loc, entry.Name, info, this));
 
 						// There are folders in linux sdk that has files with same name but different casing.
 						// Ideally FileReference.Comparer should be case sensitive on linux/mac but I don't dare changing that right now
-						if (newFiles.TryAdd(entry.Name, file))
-						{
-							file.UpdateCachedDirectory(this);
-						}
+						newFiles.TryAdd(entry.Name, file);
 					}
 				}
 
diff --git a/Engine/Source/Programs/Shared/EpicGames.Build/System/FileItem.cs b/Engine/Source/Programs/Shared/EpicGames.Build/System/FileItem.cs
--- a/Engine/Source/Programs/Shared/EpicGames.Build/System/FileItem.cs
+++ b/Engine/Source/Programs/Shared/EpicGames.Build/System/FileItem.cs
@@ -116,7 +116,7 @@ namespace UnrealBuildBase
 		/// <param name="Location">Location of the file</param>
 		/// <returns>The FileItem that represents the given a full file path.</returns>
 		public static FileItem GetItemByFileReference(FileReference Location)
-			=> UniqueSourceFileMap.GetOrAdd(Location, static loc => new FileItem(loc, loc.GetFileName(), null));
+			=> UniqueSourceFileMap.GetOrAdd(Location, static loc => new FileItem(loc, loc.GetFileName(), null, null));
 
 		/// <summary>
 		/// Deletes the file.
@@ -261,13 +261,20 @@ namespace UnrealBuildBase
 		#endregion
 
 
-		internal FileItem(FileReference Location, string name, Info? info)
+		internal FileItem(FileReference Location, string name, Info? info, DirectoryItem? directory)
 		{
 			this.Location = Location;
 			Name = name;
 			_info = info;
 
-			CachedDirectory = new(() => DirectoryItem.GetItemByDirectoryReference(Location.Directory)); // Take this in
+			if (directory != null)
+			{
+				CachedDirectory = new(directory);
+			}
+			else
+			{
+				CachedDirectory = new(() => DirectoryItem.GetItemByDirectoryReference(Location.Directory)); // Take this in
+			}
 		}
 
 		private Info Update()
