[EpicGames.Core]
* Fixed code formatting to satisfy static analysis

diff --git a/Engine/Source/Programs/Shared/EpicGames.Core/LegacyEventLogger.cs b/Engine/Source/Programs/Shared/EpicGames.Core/LegacyEventLogger.cs
--- a/Engine/Source/Programs/Shared/EpicGames.Core/LegacyEventLogger.cs
+++ b/Engine/Source/Programs/Shared/EpicGames.Core/LegacyEventLogger.cs
@@ -73,7 +73,7 @@ namespace EpicGames.Core
 		{
 			void DoLog()
 			{
-				var logger = _inner;
+				ILogger logger = _inner;
 				lock (_parser)
 				{
 					_parser.Flush();
@@ -81,7 +81,7 @@ namespace EpicGames.Core
 				logger.Log(logLevel, eventId, state, exception, formatter);
 			}
 
-			var s = Volatile.Read(ref _state);
+			ThreadingState? s = Volatile.Read(ref _state);
 			if (s == null)
 			{
 				DoLog();
@@ -99,8 +99,14 @@ namespace EpicGames.Core
 			using ManualResetEventSlim done = new(false, spinCount: 1);
 			s.Writer.TryWrite(() =>
 			{
-				try { DoLog(); }
-				finally { done.Set(); }
+				try
+				{
+					DoLog();
+				}
+				finally
+				{
+					done.Set();
+				}
 			});
 			done.Wait();
 		}
