[EpicGames.Build]
* Removed unused code

diff --git a/Engine/Source/Programs/Shared/EpicGames.Build/UEBuildPlatformSDK.cs b/Engine/Source/Programs/Shared/EpicGames.Build/UEBuildPlatformSDK.cs
--- a/Engine/Source/Programs/Shared/EpicGames.Build/UEBuildPlatformSDK.cs
+++ b/Engine/Source/Programs/Shared/EpicGames.Build/UEBuildPlatformSDK.cs
@@ -1142,10 +1142,6 @@ namespace EpicGames.Core
 			}
 		}
 
-		static Lazy<JsonSerializerOptions>ProcessJsonFileOptions = new(() => new() {
-			AllowTrailingCommas = true,
-			ReadCommentHandling = JsonCommentHandling.Skip,
-		});
 		private void ProcessJsonFile(FileReference SDKConfigFile, Dictionary<string, string> VersionMap, Dictionary<string, string[]> VersionArrayMap)
 		{
 			// load the json
