Fixed race condition when attempting to add more CB converter in EpicGames.AspNet

We need to be respecting the lock around the type converters dict, thus made the it private and added method for adding instead.

Fixes build Issue 1037071

diff --git a/Engine/Source/Programs/Shared/EpicGames.AspNet/CbConvertersAspnet.cs b/Engine/Source/Programs/Shared/EpicGames.AspNet/CbConvertersAspnet.cs
--- a/Engine/Source/Programs/Shared/EpicGames.AspNet/CbConvertersAspnet.cs
+++ b/Engine/Source/Programs/Shared/EpicGames.AspNet/CbConvertersAspnet.cs
@@ -17,7 +17,7 @@ namespace EpicGames.AspNet
 		/// </summary>
 		public static void AddAspnetConverters()
 		{
-			CbConverter.TypeToConverter[typeof(ProblemDetails)] = new CbProblemDetailsConverter();
+			CbConverter.AddConverter(typeof(ProblemDetails), new CbProblemDetailsConverter());
 		}
 	}
 
diff --git a/Engine/Source/Programs/Shared/EpicGames.Serialization/CbConverter.cs b/Engine/Source/Programs/Shared/EpicGames.Serialization/CbConverter.cs
--- a/Engine/Source/Programs/Shared/EpicGames.Serialization/CbConverter.cs
+++ b/Engine/Source/Programs/Shared/EpicGames.Serialization/CbConverter.cs
@@ -211,7 +211,7 @@ namespace EpicGames.Serialization
 		/// <summary>
 		/// Cache of type to converter
 		/// </summary>
-		public static Dictionary<Type, CbConverter> TypeToConverter { get; } = new Dictionary<Type, CbConverter>()
+		private static Dictionary<Type, CbConverter> TypeToConverter { get; } = new Dictionary<Type, CbConverter>()
 		{
 			[typeof(bool)] = new CbPrimitiveConverter<bool>(x => x.AsBool(), (w, v) => w.WriteBoolValue(v), (w, n, v) => w.WriteBool(n, v)),
 			[typeof(int)] = new CbPrimitiveConverter<int>(x => x.AsInt32(), (w, v) => w.WriteIntegerValue(v), (w, n, v) => w.WriteInteger(n, v)),
@@ -308,6 +308,20 @@ namespace EpicGames.Serialization
 			return GetConverter(property.PropertyType);
 		}
 
+		
+		/// <summary>
+		/// Adds a new converter to the type to converter mapping
+		/// </summary>
+		/// <param name="type">The type to convert to/from</param>
+		/// <param name="converter">The converter that converts to/from the type</param>
+		public static void AddConverter(Type type, CbConverter converter)
+		{
+			lock (s_lockObject)
+			{
+				TypeToConverter.Add(type, converter);
+			}
+		}
+
 		/// <summary>
 		/// Gets the converter for a particular type
 		/// </summary>
