Fixed invalid usage of _getUObject, and deprecated it so that user code doesn't use it.

diff --git a/Engine/Plugins/VirtualProduction/DynamicMaterial/Source/DynamicMaterial/Public/IDMParameterContainer.h b/Engine/Plugins/VirtualProduction/DynamicMaterial/Source/DynamicMaterial/Public/IDMParameterContainer.h
--- a/Engine/Plugins/VirtualProduction/DynamicMaterial/Source/DynamicMaterial/Public/IDMParameterContainer.h
+++ b/Engine/Plugins/VirtualProduction/DynamicMaterial/Source/DynamicMaterial/Public/IDMParameterContainer.h
@@ -41,12 +41,12 @@ public:
 
 	void CopyParametersTo(UObject* InOther)
 	{
-		CopyParametersBetween(_getUObject(), InOther);
+		CopyParametersBetween(Cast<UObject>(this), InOther);
 	}
 
 	void CopyParametersFromWrapper(UObject* InOther)
 	{
-		CopyParametersBetween(InOther, _getUObject());
+		CopyParametersBetween(InOther, Cast<UObject>(this));
 	}
 
 	/** Copies the parameter-based value of the given value to this value, if possible. */
diff --git a/Engine/Plugins/VirtualProduction/SceneState/Source/SceneStateBinding/Private/SceneStateBindingUtils.cpp b/Engine/Plugins/VirtualProduction/SceneState/Source/SceneStateBinding/Private/SceneStateBindingUtils.cpp
--- a/Engine/Plugins/VirtualProduction/SceneState/Source/SceneStateBinding/Private/SceneStateBindingUtils.cpp
+++ b/Engine/Plugins/VirtualProduction/SceneState/Source/SceneStateBinding/Private/SceneStateBindingUtils.cpp
@@ -57,7 +57,7 @@ void HandleStructIdChanged(UObject& InObject, const FGuid& InOldStructId, const
 		});
 
 	FStructIdChange Change;
-	Change.BindingOwner = BindingCollectionOwner->_getUObject();
+	Change.BindingOwner = Cast<UObject>(BindingCollectionOwner);
 	Change.OldToNewStructIdMap.Add(InOldStructId, InNewStructId);
 	OnStructIdChanged.Broadcast(Change);
 }
diff --git a/Engine/Source/Programs/Shared/EpicGames.UHT/Exporters/CodeGen/UhtHeaderCodeGeneratorCppFile.cs b/Engine/Source/Programs/Shared/EpicGames.UHT/Exporters/CodeGen/UhtHeaderCodeGeneratorCppFile.cs
--- a/Engine/Source/Programs/Shared/EpicGames.UHT/Exporters/CodeGen/UhtHeaderCodeGeneratorCppFile.cs
+++ b/Engine/Source/Programs/Shared/EpicGames.UHT/Exporters/CodeGen/UhtHeaderCodeGeneratorCppFile.cs
@@ -1337,7 +1337,7 @@ namespace EpicGames.UHT.Exporters.CodeGen
 			// Callee declaration
 			if (classObj.ClassType == UhtClassType.Interface)
 			{
-				builder.Append("\t\tCalleeType Callee{_V_EXEC_CONTEXT_PARAM_NAME, _getUObject(), CalleePath};\r\n");
+				builder.Append("\t\tCalleeType Callee{_V_EXEC_CONTEXT_PARAM_NAME, const_cast<UObject*>(Cast<UObject>(this)), CalleePath};\r\n");
 			}
 			else if (classObj.ClassType == UhtClassType.VModule)
 			{
diff --git a/Engine/Source/Programs/Shared/EpicGames.UHT/Exporters/CodeGen/UhtHeaderCodeGeneratorHFile.cs b/Engine/Source/Programs/Shared/EpicGames.UHT/Exporters/CodeGen/UhtHeaderCodeGeneratorHFile.cs
--- a/Engine/Source/Programs/Shared/EpicGames.UHT/Exporters/CodeGen/UhtHeaderCodeGeneratorHFile.cs
+++ b/Engine/Source/Programs/Shared/EpicGames.UHT/Exporters/CodeGen/UhtHeaderCodeGeneratorHFile.cs
@@ -1011,7 +1011,7 @@ namespace EpicGames.UHT.Exporters.CodeGen
 			}
 			AppendNativeInterfaceVerseProxyFunctions(builder, classObj);
 			builder.Append('\t').AppendClassSourceNameOrInterfaceProxyName(classObj).Append("(const TNonNullPtr<UObject>& InSelf) : Self(InSelf) {}\r\n");
-			builder.Append("\tUObject* _getUObject() const { return Self; }\r\n");
+			builder.Append("\t[[deprecated(\"Do not call _getUObject(), use Cast.\")]] UObject* _getUObject() const { return Self; }\r\n");
 			builder.Append('\t').AppendClassSourceNameOrInterfaceProxyName(classObj).Append("* operator->() { return this; }\r\n");
 			builder.Append("private:\r\n");
 			builder.Append("\tTNonNullPtr<UObject> Self;\r\n");
@@ -1583,7 +1583,7 @@ namespace EpicGames.UHT.Exporters.CodeGen
 				{
 					if (baseClass.ClassFlags.HasAnyFlags(EClassFlags.Interface))
 					{
-						builder.Append("\tvirtual UObject* _getUObject() const override { return const_cast<").Append(classObj.SourceName).Append("*>(this); } \\\r\n");
+						builder.Append("\t[[deprecated(\"Do not call _getUObject(), use Cast.\")]] virtual UObject* _getUObject() const override { return const_cast<").Append(classObj.SourceName).Append("*>(this); } \\\r\n");
 						break;
 					}
 				}
@@ -1772,7 +1772,7 @@ namespace EpicGames.UHT.Exporters.CodeGen
 			{
 				// Note: This used to be declared as a pure virtual function, but it was changed here in order to allow the Blueprint nativization process
 				// to detect C++ interface classes that explicitly declare pure virtual functions via type traits. This code will no longer trigger that check.
-				builder.Append("\tvirtual UObject* _getUObject() const { return nullptr; } \\\r\n");
+				builder.Append("\t[[deprecated(\"Do not call _getUObject(), use Cast.\")]] virtual UObject* _getUObject() const { return nullptr; } \\\r\n");
 			}
 
 			if (classObj.ClassExportFlags.HasAnyFlags(UhtClassExportFlags.SelfHasReplicatedProperties))
diff --git a/Engine/Source/Runtime/CoreUObject/Public/Templates/Casts.h b/Engine/Source/Runtime/CoreUObject/Public/Templates/Casts.h
--- a/Engine/Source/Runtime/CoreUObject/Public/Templates/Casts.h
+++ b/Engine/Source/Runtime/CoreUObject/Public/Templates/Casts.h
@@ -100,7 +100,9 @@ inline TCopyQualifiersFromTo_T<From, To>* Cast(From* Src)
 	{
 		if constexpr (TIsIInterface<From>::Value)
 		{
+PRAGMA_DISABLE_DEPRECATION_WARNINGS
 			if (UObject* Obj = Src->_getUObject())
+PRAGMA_ENABLE_DEPRECATION_WARNINGS
 			{
 				if constexpr (TIsIInterface<To>::Value)
 				{
@@ -253,7 +255,9 @@ UE_FORCEINLINE_HINT TCopyQualifiersFromTo_T<From, To>* ExactCast(From* Src)
 
 		if constexpr (TIsIInterface<From>::Value)
 		{
+PRAGMA_DISABLE_DEPRECATION_WARNINGS
 			UObject* Obj = Src->_getUObject();
+PRAGMA_ENABLE_DEPRECATION_WARNINGS
 			if constexpr (TIsIInterface<To>::Value)
 			{
 				return (To*)Obj->GetInterfaceAddress(To::UClassType::StaticClass());
diff --git a/Engine/Source/Runtime/Engine/Private/Interfaces/Interface_AssetUserData.cpp b/Engine/Source/Runtime/Engine/Private/Interfaces/Interface_AssetUserData.cpp
--- a/Engine/Source/Runtime/Engine/Private/Interfaces/Interface_AssetUserData.cpp
+++ b/Engine/Source/Runtime/Engine/Private/Interfaces/Interface_AssetUserData.cpp
@@ -22,7 +22,7 @@ bool IInterface_AssetUserData::HasAssetUserDataOfClass(TSubclassOf<UAssetUserDat
 bool IInterface_AssetUserData::AddAssetUserDataOfClass(TSubclassOf<UAssetUserData> InUserDataClass)
 {
 	check(InUserDataClass.Get() != nullptr);
-	AddAssetUserData(NewObject<UAssetUserData>(_getUObject(), InUserDataClass.Get()));
+	AddAssetUserData(NewObject<UAssetUserData>(Cast<UObject>(this), InUserDataClass.Get()));
 	return GetAssetUserDataOfClass(InUserDataClass.Get()) != nullptr;
 }
 
@@ -33,4 +33,4 @@ const TArray<UAssetUserData*>* IInterface_AssetUserData::GetAssetUserDataArray()
 
 void IInterface_AssetUserData::RemoveUserDataOfClass(TSubclassOf<UAssetUserData> InUserDataClass)
 {
-}
+}
\ No newline at end of file
diff --git a/Engine/Source/Runtime/GameplayTags/Private/GameplayTagAssetInterface.cpp b/Engine/Source/Runtime/GameplayTags/Private/GameplayTagAssetInterface.cpp
--- a/Engine/Source/Runtime/GameplayTags/Private/GameplayTagAssetInterface.cpp
+++ b/Engine/Source/Runtime/GameplayTags/Private/GameplayTagAssetInterface.cpp
@@ -40,8 +40,9 @@ FGameplayTagContainer IGameplayTagAssetInterface::BP_GetOwnedGameplayTags() cons
 	// we can redirect this member function to the non-member function in the future once we allow the CoreRedirect
 	// of the 'self' pin.
 	TScriptInterface<IGameplayTagAssetInterface> GameplayTagAssetInterface;
-	GameplayTagAssetInterface.SetObject(_getUObject());
-	GameplayTagAssetInterface.SetInterface(const_cast<IGameplayTagAssetInterface*>(this));
+	IGameplayTagAssetInterface* NonConstThis = const_cast<IGameplayTagAssetInterface*>(this);
+	GameplayTagAssetInterface.SetObject(Cast<UObject>(NonConstThis));
+	GameplayTagAssetInterface.SetInterface(NonConstThis);
 
 	return UBlueprintGameplayTagLibrary::GetOwnedGameplayTags(GameplayTagAssetInterface);
 }
