Cloud DDC - Only append converts to the converter mapping if the type does not have a converter registered already.
Fixed build issue 1044569

diff --git a/Engine/Source/Programs/Shared/EpicGames.Serialization/CbConverter.cs b/Engine/Source/Programs/Shared/EpicGames.Serialization/CbConverter.cs
--- a/Engine/Source/Programs/Shared/EpicGames.Serialization/CbConverter.cs
+++ b/Engine/Source/Programs/Shared/EpicGames.Serialization/CbConverter.cs
@@ -310,7 +310,7 @@ namespace EpicGames.Serialization
 
 		
 		/// <summary>
-		/// Adds a new converter to the type to converter mapping
+		/// Adds a new converter to the type to converter mapping if it does not exist
 		/// </summary>
 		/// <param name="type">The type to convert to/from</param>
 		/// <param name="converter">The converter that converts to/from the type</param>
@@ -318,7 +318,7 @@ namespace EpicGames.Serialization
 		{
 			lock (s_lockObject)
 			{
-				TypeToConverter.Add(type, converter);
+				TypeToConverter.TryAdd(type, converter);
 			}
 		}
 
