Fixed VNI deprecation warnings.

diff --git a/Engine/Source/Programs/Shared/EpicGames.UHT/Exporters/CodeGen/UhtHeaderCodeGeneratorCppFile.cs b/Engine/Source/Programs/Shared/EpicGames.UHT/Exporters/CodeGen/UhtHeaderCodeGeneratorCppFile.cs
--- a/Engine/Source/Programs/Shared/EpicGames.UHT/Exporters/CodeGen/UhtHeaderCodeGeneratorCppFile.cs
+++ b/Engine/Source/Programs/Shared/EpicGames.UHT/Exporters/CodeGen/UhtHeaderCodeGeneratorCppFile.cs
@@ -1337,7 +1337,9 @@ namespace EpicGames.UHT.Exporters.CodeGen
 			// Callee declaration
 			if (classObj.ClassType == UhtClassType.Interface)
 			{
-				builder.Append("\t\tCalleeType Callee{_V_EXEC_CONTEXT_PARAM_NAME, const_cast<UObject*>(Cast<UObject>(this)), CalleePath};\r\n");
+				builder.Append("\t\tPRAGMA_DISABLE_DEPRECATION_WARNINGS\r\n");
+				builder.Append("\t\tCalleeType Callee{_V_EXEC_CONTEXT_PARAM_NAME, _getUObject(), CalleePath};\r\n");
+				builder.Append("\t\tPRAGMA_ENABLE_DEPRECATION_WARNINGS\r\n");
 			}
 			else if (classObj.ClassType == UhtClassType.VModule)
 			{
diff --git a/Engine/Source/Programs/Shared/EpicGames.UHT/Exporters/CodeGen/UhtHeaderCodeGeneratorHFile.cs b/Engine/Source/Programs/Shared/EpicGames.UHT/Exporters/CodeGen/UhtHeaderCodeGeneratorHFile.cs
--- a/Engine/Source/Programs/Shared/EpicGames.UHT/Exporters/CodeGen/UhtHeaderCodeGeneratorHFile.cs
+++ b/Engine/Source/Programs/Shared/EpicGames.UHT/Exporters/CodeGen/UhtHeaderCodeGeneratorHFile.cs
@@ -1011,7 +1011,7 @@ namespace EpicGames.UHT.Exporters.CodeGen
 			}
 			AppendNativeInterfaceVerseProxyFunctions(builder, classObj);
 			builder.Append('\t').AppendClassSourceNameOrInterfaceProxyName(classObj).Append("(const TNonNullPtr<UObject>& InSelf) : Self(InSelf) {}\r\n");
-			builder.Append("\t[[deprecated(\"Do not call _getUObject(), use Cast.\")]] UObject* _getUObject() const { return Self; }\r\n");
+			builder.Append("\t[[deprecated(\"Do not call _getUObject().\")]] UObject* _getUObject() const { return Self; }\r\n");
 			builder.Append('\t').AppendClassSourceNameOrInterfaceProxyName(classObj).Append("* operator->() { return this; }\r\n");
 			builder.Append("private:\r\n");
 			builder.Append("\tTNonNullPtr<UObject> Self;\r\n");
