[EpicGames.Build]
* Minor opt where DirectoryItem scan updates containing FileItem's data if not set already

diff --git a/Engine/Source/Programs/Shared/EpicGames.Build/System/DirectoryItem.cs b/Engine/Source/Programs/Shared/EpicGames.Build/System/DirectoryItem.cs
--- a/Engine/Source/Programs/Shared/EpicGames.Build/System/DirectoryItem.cs
+++ b/Engine/Source/Programs/Shared/EpicGames.Build/System/DirectoryItem.cs
@@ -269,9 +269,9 @@ namespace UnrealBuildBase
 					}
 					else
 					{
-						FileItem.Info info = new() { Attributes = entry.Attributes, Length = entry.Size, LastWriteTimeUtc = entry.LastWriteTimeUtc, CreationTimeUtc = entry.CreationTimeUtc };
 						FileReference loc = FileReference.Combine(Location, entry.Name);
-						FileItem file = FileItem.UniqueSourceFileMap.GetOrAdd(loc, loc => new FileItem(loc, entry.Name, info, this));
+						FileItem file = FileItem.UniqueSourceFileMap.GetOrAdd(loc, loc => new FileItem(loc, entry.Name, new() { Attributes = entry.Attributes, Length = entry.Size, LastWriteTimeUtc = entry.LastWriteTimeUtc, CreationTimeUtc = entry.CreationTimeUtc }, this));
+						file.Update(entry);
 
 						// There are folders in linux sdk that has files with same name but different casing.
 						// Ideally FileReference.Comparer should be case sensitive on linux/mac but I don't dare changing that right now
diff --git a/Engine/Source/Programs/Shared/EpicGames.Build/System/FileItem.cs b/Engine/Source/Programs/Shared/EpicGames.Build/System/FileItem.cs
--- a/Engine/Source/Programs/Shared/EpicGames.Build/System/FileItem.cs
+++ b/Engine/Source/Programs/Shared/EpicGames.Build/System/FileItem.cs
@@ -296,6 +296,14 @@ namespace UnrealBuildBase
 			return newInfo;
 		}
 
+		internal void Update(DirectoryEnumerator.Entry entry)
+		{
+			if (_info == null)
+			{
+				_info = new() { Attributes = entry.Attributes, Length = entry.Size, LastWriteTimeUtc = entry.LastWriteTimeUtc, CreationTimeUtc = entry.CreationTimeUtc };
+			}
+		}
+
 		/// <summary>
 		/// A case-insensitive dictionary that's used to map each unique file name to a single FileItem object.
 		/// </summary>
