[EpicGames.Build]
* Added so code traversing directories to find rules (.build.cs and .target.cs) filters out more custom directories

diff --git a/Engine/Source/Programs/Shared/EpicGames.Build/System/Rules.cs b/Engine/Source/Programs/Shared/EpicGames.Build/System/Rules.cs
--- a/Engine/Source/Programs/Shared/EpicGames.Build/System/Rules.cs
+++ b/Engine/Source/Programs/Shared/EpicGames.Build/System/Rules.cs
@@ -259,6 +259,7 @@ namespace UnrealBuildBase
 			// Scan all the files in this directory
 			bool bSearchSubFolders = true;
 			bool bIsPlugin = false;
+			bool bIsCsProj = false;
 
 			foreach (FileItem File in Directory.EnumerateFiles())
 			{
@@ -298,6 +299,10 @@ namespace UnrealBuildBase
 					bIsPlugin = true;
 					bSearchSubFolders = false;
 				}
+				else if (File.HasExtension(".csproj"))
+				{
+					bIsCsProj = true;
+				}
 				else if (File.Name == ".ubtignore")
 				{
 					bSearchSubFolders = false;
@@ -325,6 +330,23 @@ namespace UnrealBuildBase
 			{
 				foreach (DirectoryItem SubDirectory in Directory.EnumerateDirectories())
 				{
+					if (bIsCsProj)
+					{
+						if (SubDirectory.Name == "bin" || SubDirectory.Name == "obj") // Intermediate files.. we don't want to enter here
+						{
+							continue;
+						}
+					}
+					else if (SubDirectory.Name.StartsWith("__External")) // This is a bit annoying but version control likes to leave around these huuuge hierarchies of empty directories
+					{
+						//Console.WriteLine("Most likely leftovers: " + Directory.Location);
+						continue;
+					}
+					//else if (SubDirectory.Name == "Intermediate")
+					//{
+					//	Console.WriteLine("Most likely leftovers: " + Directory.Location);
+					//}
+
 					Queue.Enqueue(() => FindAllRulesFilesRecursively(SubDirectory, Cache, Queue));
 				}
 			}
