[EpicGames.Build]
* Changed so cached directories use .TryAdd since there could be subdirectories with different casing

diff --git a/Engine/Source/Programs/Shared/EpicGames.Build/System/DirectoryItem.cs b/Engine/Source/Programs/Shared/EpicGames.Build/System/DirectoryItem.cs
--- a/Engine/Source/Programs/Shared/EpicGames.Build/System/DirectoryItem.cs
+++ b/Engine/Source/Programs/Shared/EpicGames.Build/System/DirectoryItem.cs
@@ -265,7 +265,10 @@ namespace UnrealBuildBase
 						Info info = new() { Exists = true, LastWriteTimeUtc = entry.LastWriteTimeUtc, CreationTimeUtc = entry.CreationTimeUtc };
 						DirectoryReference loc = DirectoryReference.Combine(Location, entry.Name);
 						DirectoryItem dir = LocationToItem.GetOrAdd(loc, loc => new DirectoryItem(loc, entry.Name, info));
-						newDirs.Add(entry.Name, dir);
+
+						// There are folders in linux sdk that has files with same name but different casing.
+						// Ideally FileReference.Comparer should be case sensitive on linux/mac but I don't dare changing that right now
+						newDirs.TryAdd(entry.Name, dir);
 					}
 					else
 					{
