UHT changes to support compiled in UObjects being emitted as constinit variable definitions rather than constructed on the heap at runtime from param structs.
Use of params structs for heap allocated objects vs constinit variables is controlled by the preprocessor definition UE_WITH_CONSTINIT_UOBJECT.
In constinit UObject variables, pointers to other UObjects/FProperties are assigned at compile time where possible.

New target rules setting: bEnableConstInitUObject.
Both sets of code are always emitted by UHT, but whether the target defines UE_WITH_CONSTINIT_UOBJECT is controlled by this setting.
Currently only supported in monolithic targets with live coding disabled.
Added command line arguments to enable/disable live coding.

Singleton names:
- Replace many uses of bool for registered/unregistered factory function with enum UhtSingletonType to expand to cover constinit var and statics struct.
- Add IUhtSingletonNameProvider to provide access to these names from code generators, property member context, etc

UhtPackageCodeGeneratorCppFile:
- Emit a constinit variable definition for each UhtPackage.
- Emit a FRegisterCompiledInObjects definition for all singletons in each package that will not be otherwise initialized (replaces array of singleton factory functions from package factory function).

UhtHeaderCodeGeneratorHFile:
- Emit forward declaration for class constinit singleton outside of class body macro
- Emit forward declarations of _Statics struct for each UClass function with a custom thunk and make those classes a friend of the UClass
- Make StaticConfigName constexpr

UhtHeaderCodeGeneratorCppFile:
- Move RigVM function registration out of StaticStruct factory method and register with a global variable like other UObject registrations in constinit case. In params case, use same lists of method info rather emitting code to build an array in StaticStruct method.
- Separate lists of enum names and values in generated code for constinit case.
- Define UScriptStruct::TCppStructOps as a constinit variable in constinit case rather than using factory method.
- Assign native function pointer for UFunction in definition of constinit variable rather than storing names and looking them up at runtime.
- Initialize FStructBaseChain with constinit array of pointers to supers of UStructs

UhtType:
- Change UhtType.EngineClassName to return concrete type name runtime class i.e. VerseClass for verse classes rather than just Class for all classes.
- Add UhtType.EngineLinkClassName to return more general type name for factory function declarations, i.e. just Class even if it is a verse class.
- Add UhtType.EngineClass to return parsed runtime type matching EngineClassName.

UhtStruct:
- Added define-scope aware linked list of all properties within the struct.

UhtClass:
- Added a define-scope aware linked list of all UhtFunctions within the class.
- Add UhtClass.NativeFunctionCallName lazily-initialized property containing correct identifier for function call scoping i.e. UMyClass or IMyInterface.

UhtProperty:
- Added UhtDefineScopeLink and give properties a linked list of other properties within the same struct.
- Added methods AppendConstInitMemberDecl/Def to emit declarations/definitions for property objects as variables.

UhtEnum:
- Generate max name/value in UHT rather than leaving it until runtime.

Misc changes:
- Add support for emitting TConstArrayView expressions rather than separate pointer/size arguments.
- Add more properties to access parsed types from CoreUObject to UhtSession e.g. UFunction.
- Many types add more references to IUhtReferenceCollector that are necessary to access constinit var definitions to link objects together.

diff --git a/Engine/Source/Programs/Shared/EpicGames.Core/UHTTypes.cs b/Engine/Source/Programs/Shared/EpicGames.Core/UHTTypes.cs
--- a/Engine/Source/Programs/Shared/EpicGames.Core/UHTTypes.cs
+++ b/Engine/Source/Programs/Shared/EpicGames.Core/UHTTypes.cs
@@ -2,6 +2,7 @@
 
 using System;
 using System.Collections.Generic;
+using System.Diagnostics.CodeAnalysis;
 using System.Text.Json.Serialization;
 
 #pragma warning disable CS1591 // Missing documentation
@@ -99,6 +100,73 @@ namespace EpicGames.Core
 		InternalUser
 	}
 
+	/// <summary>
+	/// Types of code to generate to initialize compiled-in UObject instances. 
+	/// Flags to allow multiple formats to be generated at once for development/debugging.
+	/// </summary>
+	[Flags]
+	public enum UhtCompiledInObjectFormat
+	{
+		/// <summary>
+		/// No flags, should not be configured as a desired output type.
+		/// </summary>
+		None = 0x0,
+
+		/// <summary>
+		/// Generate params structs which will be used to initialize UObjects on the heap at runtime.
+		/// </summary>
+		Params = 0x1,
+
+		/// <summary>
+		/// Generate declarations of constinit variables of UObject-derived types allowing these
+		/// objects to be stored in the binary's data segment.
+		/// </summary>
+		ConstInit = 0x2,
+
+		/// <summary>
+		/// Default setting to use
+		/// </summary>
+		Default = Params,
+	
+	 	/// <summary>
+		/// Generate all possible types at once for comparison
+		/// </summary>
+		All = Params | ConstInit,
+	}
+
+	/// <summary>
+	/// Per-target settings for UHT to control features that need to be the same across all modules. 
+	/// </summary>
+	public record struct UHTTargetSettings
+	{
+		public UHTTargetSettings()
+		{
+		}
+
+		[SetsRequiredMembers]
+		public UHTTargetSettings(BinaryArchiveReader reader) : this()
+		{
+			CompiledInObjectFormat = Enum.Parse<UhtCompiledInObjectFormat>(reader.ReadString()!);
+		}
+
+		public void Write(BinaryArchiveWriter writer)
+		{
+			writer.WriteString(CompiledInObjectFormat.ToString());
+		}
+
+		/// <summary>
+		/// What type of code to generate for compiled-on UObject instances such as UClass instances.
+		/// </summary>
+		[JsonConverter(typeof(JsonStringEnumConverter))]
+		public UhtCompiledInObjectFormat CompiledInObjectFormat { get; set; } = UhtCompiledInObjectFormat.Default;
+
+		/// <summary>
+		/// Whether to generate a name for the maximum value of UENUMs where one was not found in the authored code.
+		/// If this is false, the UEnum code will generate it at runtime if it can do so without creating a collision. 
+		/// </summary>
+		public bool GenerateEnumMaxValues { get; set; } = false;
+	}
+
 	public class UHTManifest
 	{
 		public class Module
@@ -194,6 +262,11 @@ namespace EpicGames.Core
 		/// </summary>
 		public string ExternalDependenciesFile { get; set; } = String.Empty;
 
+		/// <summary>
+		/// Target-wide settings for UHT/code generation. 
+		/// </summary>
+		public UHTTargetSettings TargetSettings { set; get; } = new();
+
 		/// <summary>
 		/// List of modules
 		/// </summary>
diff --git a/Engine/Source/Programs/Shared/EpicGames.UHT/Exporters/CodeGen/UhtCodeGenerator.cs b/Engine/Source/Programs/Shared/EpicGames.UHT/Exporters/CodeGen/UhtCodeGenerator.cs
--- a/Engine/Source/Programs/Shared/EpicGames.UHT/Exporters/CodeGen/UhtCodeGenerator.cs
+++ b/Engine/Source/Programs/Shared/EpicGames.UHT/Exporters/CodeGen/UhtCodeGenerator.cs
@@ -12,8 +12,67 @@ using EpicGames.UHT.Utils;
 
 namespace EpicGames.UHT.Exporters.CodeGen
 {
+	/// <summary>
+	/// Type of singleton to retrieve a name or declaration for
+	/// </summary>
+	public enum UhtSingletonType
+	{
+		/// <summary>
+		/// Function that returns a fully constructed object - the distinction between this and Unregistered is only important for classes.
+		/// </summary>
+		Registered = 0,
+
+		/// <summary>
+		/// Function that returns a partially constructed object
+		/// </summary>
+		Unregistered = 1,
+
+		/// <summary>
+		/// Constinit variable (C++20 feature)
+		/// </summary>
+		ConstInit = 2,
+
+		/// <summary>
+		/// Struct containing initialization data for the object 
+		/// </summary>
+		Statics = 3,
+	}
+	
+	/// <summary>
+	/// Interface to provide singleton names during code generation. 
+	/// </summary>
+	public interface IUhtSingletonNameProvider
+	{
+		/// <summary>
+		/// Return the singleton name for the given object
+		/// </summary>
+		/// <param name="obj">Object in question - may be a struct, function, class etc or a package</param>
+		/// <param name="type">Type of declaration to get</param>
+		/// <returns>Singleton function name</returns>
+		public string GetSingletonName(UhtObject? obj, UhtSingletonType type);
+	}
+
+	/// <summary>
+	/// Interface to provide global settings during code generation. 
+	/// </summary>
+	public interface IUhtSessionSettingsProvider
+	{
+		/// <summary>
+		/// Returns whether UHT should generate a certain type of code for compiled-in object instances
+		/// </summary>
+		/// <param name="format">Format to test for</param>
+		/// <returns>Whether the given format is active - multiple formats may be active in certain development modes</returns>
+		public bool IsUsingCompiledInObjectFormat(UhtCompiledInObjectFormat format);
+
+		/// <summary>
+		/// Returns whether UHT is emitting code for multiple compiled-in object formats, and therefore such
+		/// code should be surrounded by preprocessor checks
+		/// </summary>
+		public bool IsUsingMultipleCompiledInObjectFormats { get; }
+	}
+
 	[UnrealHeaderTool]
-	class UhtCodeGenerator
+	class UhtCodeGenerator : IUhtSingletonNameProvider, IUhtSessionSettingsProvider
 	{
 		[UhtExporter(Name = "CodeGen", Description = "Standard UnrealEngine code generation", Options = UhtExporterOptions.Default,
 			CppFilters = new string[] { "*.generated.cpp", "*.generated.*.cpp", "*.gen.cpp", "*.gen.*.cpp" },
@@ -50,10 +109,23 @@ namespace EpicGames.UHT.Exporters.CodeGen
 
 		public struct ObjectInfo
 		{
+			// e.g. Z_Construct_UClass_UObject
 			public string RegisteredSingletonName { get; set; }
+			// e.g. Z_Construct_UClass_UObject_NoRegister
 			public string UnregisteredSingletonName { get; set; }
+			// e.g. COREUOBJECT_API* UClass* Z_Construct_UClass_UObject();
 			public string RegisteredExternalDecl { get; set; }
+			// e.g. COREUOBJECT_API* UClass* Z_Construct_UClass_UObject_NoRegister();
 			public string UnregisteredExternalDecl { get; set; }
+			// Fully qualified name of variable containing const-initialized object 
+			// e.g. Z_ConstInit_UClass_UObject
+			public string ConstInitSingletonName { get; set; }
+			// Full declaration to import variable containing const-initialized object
+			// e.g. COREUOBJECT_API extern UE_CONSTINIT_UOBJECT_DECL UClass Z_ConstInit_UClass_UObject
+			public string ConstInitExternalDecl { get; set; }
+			// Name of statics struct in generated cpp containing construction parameters
+			// e.g. Z_Construct_UClass_UObject_Statics
+			public string StaticsName { get; set; }
 			public UhtClass? NativeInterface { get; set; }
 			public UhtProperty? FastArrayProperty { get; set; }
 			public uint Hash { get; set; }
@@ -182,39 +254,69 @@ namespace EpicGames.UHT.Exporters.CodeGen
 		/// <summary>
 		/// Return the singleton name for an object
 		/// </summary>
-		/// <param name="obj">The object in question.</param>
-		/// <param name="registered">If true, return the registered singleton name.  Otherwise return the unregistered.</param>
+		/// <param name="obj">The type or package in question.</param>
+		/// <param name="type">Type of declaration to get</param>
 		/// <returns>Singleton name or "nullptr" if Object is null</returns>
-		public string GetSingletonName(UhtObject? obj, bool registered)
+		public string GetSingletonName(UhtObject? obj, UhtSingletonType type)
 		{
 			if (obj == null)
 			{
 				return "nullptr";
 			}
-			return registered ? ObjectInfos[obj.ObjectTypeIndex].RegisteredSingletonName : ObjectInfos[obj.ObjectTypeIndex].UnregisteredSingletonName;
+			switch (type)
+			{
+				case UhtSingletonType.Statics:
+					return ObjectInfos[obj.ObjectTypeIndex].StaticsName;
+				case UhtSingletonType.Registered:
+					return ObjectInfos[obj.ObjectTypeIndex].RegisteredSingletonName;
+				case UhtSingletonType.Unregistered:
+					return ObjectInfos[obj.ObjectTypeIndex].UnregisteredSingletonName;
+				case UhtSingletonType.ConstInit:
+					return ObjectInfos[obj.ObjectTypeIndex].ConstInitSingletonName;
+			}
+			throw new UhtIceException($"Invalid UhtSingletonType value {type}");
 		}
 
 		/// <summary>
 		/// Return the external declaration for an object
 		/// </summary>
 		/// <param name="obj">The object in question.</param>
-		/// <param name="registered">If true, return the registered external declaration.  Otherwise return the unregistered.</param>
+		/// <param name="type">Type of declaration to get</param>
 		/// <returns>External declaration</returns>
-		public string GetExternalDecl(UhtObject obj, bool registered)
+		public string GetExternalDecl(UhtObject obj, UhtSingletonType type)
 		{
-			return GetExternalDecl(obj.ObjectTypeIndex, registered);
+			return GetExternalDecl(obj.ObjectTypeIndex, type);
 		}
 
 		/// <summary>
 		/// Return the external declaration for an object
 		/// </summary>
 		/// <param name="objectIndex">The object in question.</param>
-		/// <param name="registered">If true, return the registered external declaration.  Otherwise return the unregistered.</param>
+		/// <param name="type">Type of declaration to get</param>
 		/// <returns>External declaration</returns>
-		public string GetExternalDecl(int objectIndex, bool registered)
+		public string GetExternalDecl(int objectIndex, UhtSingletonType type)
 		{
-			return registered ? ObjectInfos[objectIndex].RegisteredExternalDecl : ObjectInfos[objectIndex].UnregisteredExternalDecl;
+			switch (type)
+			{
+				case UhtSingletonType.Registered:
+					return ObjectInfos[objectIndex].RegisteredExternalDecl;
+				case UhtSingletonType.Unregistered:
+					return ObjectInfos[objectIndex].UnregisteredExternalDecl;
+				case UhtSingletonType.ConstInit:
+					return ObjectInfos[objectIndex].ConstInitExternalDecl;
+				case UhtSingletonType.Statics:
+					throw new UhtIceException($"Statics struct for {ObjectInfos[objectIndex].StaticsName} cannot be imported across translation units");
+			}
+			throw new UhtIceException($"Invalid UhtSingletonType value {type}");
 		}
+
+		public bool IsUsingCompiledInObjectFormat(UhtCompiledInObjectFormat format) => Session.IsUsingCompiledInObjectFormat(format);
+
+		/// <summary>
+		/// Returns whether UHT is emitting code for multiple compiled-in object formats, and therefore such
+		/// code should be surrounded by preprocessor checks
+		/// </summary>
+		public bool IsUsingMultipleCompiledInObjectFormats => Session.IsUsingMultipleCompiledInObjectFormats;
 		#endregion
 
 		#region Information initialization
@@ -243,7 +345,10 @@ namespace EpicGames.UHT.Exporters.CodeGen
 			builder.Append("Z_Construct_UPackage_");
 			builder.Append(packageInfo.StrippedName);
 			objectInfo.UnregisteredSingletonName = objectInfo.RegisteredSingletonName = builder.ToString();
+			objectInfo.StaticsName = $"{objectInfo.RegisteredSingletonName}_Statics";
 			objectInfo.UnregisteredExternalDecl = objectInfo.RegisteredExternalDecl = $"\tUPackage* {objectInfo.RegisteredSingletonName}();\r\n";
+			objectInfo.ConstInitSingletonName = $"Z_ConstInit_UPackage_{packageInfo.StrippedName}";
+			objectInfo.ConstInitExternalDecl = $"\textern UE_CONSTINIT_UOBJECT_DECL UPackage {objectInfo.ConstInitSingletonName};\r\n";
 		}
 
 		private void InitHeaderInfo(StringBuilder builder, UhtHeaderFile headerFile)
@@ -310,7 +415,6 @@ namespace EpicGames.UHT.Exporters.CodeGen
 
 		private void InitObjectInfo(StringBuilder builder, UhtPackage package, ref HeaderInfo headerInfo, UhtObject obj)
 		{
-			UhtModule module = package.Module;
 			ref ObjectInfo objectInfo = ref ObjectInfos[obj.ObjectTypeIndex];
 
 			builder.Clear();
@@ -319,7 +423,7 @@ namespace EpicGames.UHT.Exporters.CodeGen
 			bool isNonIntrinsicClass = false;
 			builder.Append("Z_Construct_U").Append(obj.EngineClassName).AppendOuterNames(obj);
 
-			string engineClassName = obj.EngineClassName;
+			string api = obj.Module.Api;
 			if (obj is UhtClass classObj)
 			{
 				if (!classObj.ClassFlags.HasAnyFlags(EClassFlags.Intrinsic))
@@ -375,11 +479,6 @@ namespace EpicGames.UHT.Exporters.CodeGen
 					}
 				}
 			}
-			else if (obj is UhtFunction)
-			{
-				// The method for EngineClassName returns type specific where in this case we need just the simple return type
-				engineClassName = "Function";
-			}
 			else if (obj is UhtEnum enumObj)
 			{
 				if (enumObj.IsVerseField)
@@ -393,15 +492,16 @@ namespace EpicGames.UHT.Exporters.CodeGen
 				objectInfo.RegisteredSingletonName = builder.ToString();
 				builder.Append("_NoRegister");
 				objectInfo.UnregisteredSingletonName = builder.ToString();
-
-				objectInfo.UnregisteredExternalDecl = $"\t{module.Api}U{engineClassName}* {objectInfo.UnregisteredSingletonName}();\r\n";
-				objectInfo.RegisteredExternalDecl = $"\t{module.Api}U{engineClassName}* {objectInfo.RegisteredSingletonName}();\r\n";
 			}
 			else
 			{
 				objectInfo.UnregisteredSingletonName = objectInfo.RegisteredSingletonName = builder.ToString();
-				objectInfo.UnregisteredExternalDecl = objectInfo.RegisteredExternalDecl = $"\t{module.Api}U{engineClassName}* {objectInfo.RegisteredSingletonName}();\r\n";
 			}
+			objectInfo.StaticsName = $"{objectInfo.RegisteredSingletonName}_Statics";
+			objectInfo.UnregisteredExternalDecl = $"\t{api}U{obj.EngineLinkClassName}* {objectInfo.UnregisteredSingletonName}();\r\n";
+			objectInfo.RegisteredExternalDecl = $"\t{api}U{obj.EngineLinkClassName}* {objectInfo.RegisteredSingletonName}();\r\n";
+			objectInfo.ConstInitSingletonName = objectInfo.RegisteredSingletonName.Replace("Z_Construct_", "Z_ConstInit_", StringComparison.InvariantCultureIgnoreCase);
+			objectInfo.ConstInitExternalDecl = $"\t{api}extern UE_CONSTINIT_UOBJECT_DECL TNoDestroy<U{obj.EngineClassName}> {objectInfo.ConstInitSingletonName};\r\n";
 
 			// Init the children
 			foreach (UhtType child in obj.Children)
diff --git a/Engine/Source/Programs/Shared/EpicGames.UHT/Exporters/CodeGen/UhtHeaderCodeGenerator.cs b/Engine/Source/Programs/Shared/EpicGames.UHT/Exporters/CodeGen/UhtHeaderCodeGenerator.cs
--- a/Engine/Source/Programs/Shared/EpicGames.UHT/Exporters/CodeGen/UhtHeaderCodeGenerator.cs
+++ b/Engine/Source/Programs/Shared/EpicGames.UHT/Exporters/CodeGen/UhtHeaderCodeGenerator.cs
@@ -225,7 +225,7 @@ namespace EpicGames.UHT.Exporters.CodeGen
 			builder.Append(' ');
 			if (!isDeclaration && outerClass != null)
 			{
-				builder.Append(outerClass.Namespace.FullSourceName).AppendClassSourceNameOrInterfaceName(outerClass).Append("::");
+				builder.AppendClassSourceNameOrInterfaceName(outerClass).Append("::");
 			}
 
 			if (alternateFunctionName != null)
