Fix for MSVC compiler bug - provide all designated initializers for FRigVMCompiledInFunction including for predicates.

diff --git a/Engine/Source/Programs/Shared/EpicGames.UHT/Exporters/CodeGen/UhtHeaderCodeGeneratorCppFile.cs b/Engine/Source/Programs/Shared/EpicGames.UHT/Exporters/CodeGen/UhtHeaderCodeGeneratorCppFile.cs
--- a/Engine/Source/Programs/Shared/EpicGames.UHT/Exporters/CodeGen/UhtHeaderCodeGeneratorCppFile.cs
+++ b/Engine/Source/Programs/Shared/EpicGames.UHT/Exporters/CodeGen/UhtHeaderCodeGeneratorCppFile.cs
@@ -871,7 +871,7 @@ namespace EpicGames.UHT.Exporters.CodeGen
 					{
 						if (methodInfo.IsPredicate)
 						{
-							builder.Append($"\t\t{{ .MethodName = TEXT(\"{methodInfo.Name}\"), \r\n");
+							builder.Append($"\t\t{{ .MethodName = TEXT(\"{methodInfo.Name}\"), .Function = nullptr, \r\n");
 						}
 						else
 						{
