[UBT/EpicGames.Build]
* Optimized logic finding .build.cs and .target.cs files. It is now creating a tree mirrored from file system with only the interesting files. Previous caching did much more I/O interaction than what was needed. Cut time times quite a lot

diff --git a/Engine/Source/Programs/Shared/EpicGames.Build/System/EnumeratePlugins.cs b/Engine/Source/Programs/Shared/EpicGames.Build/System/EnumeratePlugins.cs
--- a/Engine/Source/Programs/Shared/EpicGames.Build/System/EnumeratePlugins.cs
+++ b/Engine/Source/Programs/Shared/EpicGames.Build/System/EnumeratePlugins.cs
@@ -127,8 +127,8 @@ namespace UnrealBuildBase
 				PluginDirectories.AddRange(Unreal.GetExtensionDirs(ProjectFile.Directory, "Plugins"));
 				PluginDirectories.AddRange(Unreal.GetExtensionDirs(ProjectFile.Directory, "Mods"));
 			}
-			List<FileReference> Plugins = UnrealBuildBase.Rules.FindAllRulesFiles(PluginDirectories, UnrealBuildBase.Rules.RulesFileType.UbtPlugin);
-			Plugins.SortBy(P => P.FullName);
+			Rules.Populate(PluginDirectories);
+			List<FileReference> Plugins = Rules.Get(PluginDirectories, UnrealBuildBase.Rules.RulesFileType.UbtPlugin);
 			return Plugins;
 		}
 
diff --git a/Engine/Source/Programs/Shared/EpicGames.Build/System/Rules.cs b/Engine/Source/Programs/Shared/EpicGames.Build/System/Rules.cs
--- a/Engine/Source/Programs/Shared/EpicGames.Build/System/Rules.cs
+++ b/Engine/Source/Programs/Shared/EpicGames.Build/System/Rules.cs
@@ -38,21 +38,6 @@ namespace UnrealBuildBase
 			UbtPlugin,
 		}
 
-		/// <summary>
-		/// Cached list of rules files in each directory of each type
-		/// </summary>
-		class RulesFileCache
-		{
-			public List<FileReference> ModuleRules = [];
-			public List<FileReference> TargetRules = [];
-			public List<FileReference> AutomationModules = [];
-			public List<FileReference> UbtPlugins = [];
-		}
-
-		/// Map of root folders to a cached list of all UBT-related source files in that folder or any of its sub-folders.
-		/// We cache these file names so we can avoid searching for them later on.
-		static Dictionary<DirectoryReference, RulesFileCache> RootFolderToRulesFileCache = [];
-
 		/// <summary>
 		/// 
 		/// </summary>
@@ -131,7 +116,10 @@ namespace UnrealBuildBase
 				Folders.AddRange(AdditionalSearchPaths.Where(x => x != null && DirectoryReference.Exists(x)));
 			}
 
-			return FindAllRulesFiles(Folders, RulesFileType);
+			Populate(Folders);
+			List<FileReference> files = Get(Folders, RulesFileType);
+			DistinctSortedList(files); // Seems like same path is added multiple times and Get is not removing dups, just sorting
+			return files;
 		}
 
 		/// <summary>
@@ -140,158 +128,196 @@ namespace UnrealBuildBase
 		/// <param name="DirectoryPath">Directory to invalidate</param>
         public static void InvalidateRulesFileCache(string DirectoryPath)
         {
-            DirectoryReference Directory = new DirectoryReference(DirectoryPath);
-			lock(RootFolderToRulesFileCache)
+			DirectoryReference dirRef = new DirectoryReference(DirectoryPath);
+			lock (Root)
 			{
-				RootFolderToRulesFileCache.Remove(Directory);
+				Directory? dir = FindDirectory(dirRef, false);
+				if (dir != null)
+				{
+					dir.IsScanned = false;
+					dir.Files = null;
+					dir.SubDirectories = null;
+				}
 			}
-            DirectoryLookupCache.InvalidateCachedDirectory(Directory);
-        }
+			DirectoryLookupCache.InvalidateCachedDirectory(dirRef);
+		}
 
-		/// <summary>
-		/// Prefetch multiple directories in parallel
-		/// </summary>
-		/// <param name="Directories">The directories to cache</param>
-		public static void PrefetchRulesFiles(IEnumerable<DirectoryReference> Directories)
+		class Files
 		{
-			PrefetchRulesFilesInternal(Directories);
+			public List<FileReference> ModuleRules = [];
+			public List<FileReference> TargetRules = [];
+			public List<FileReference> AutomationModules = [];
+			public List<FileReference> UbtPlugins = [];
 		}
 
-		/// <summary>
-		/// Prefetch multiple directories in parallel
-		/// </summary>
-		/// <param name="Directories">The directories to cache</param>
-		/// <returns>The list of caches for each directory</returns>
-		private static List<(DirectoryReference, RulesFileCache)> PrefetchRulesFilesInternal(IEnumerable<DirectoryReference> Directories)
+
+		class Directory
+		{
+			internal Files? Files;
+			internal Dictionary<string, Directory>? SubDirectories;
+			internal bool IsScanned;
+		}
+		static Directory Root = new();
+
+		static Directory? FindDirectory(DirectoryReference dirRef, bool createIfNotFound)
 		{
-			List<(DirectoryReference, RulesFileCache)> Caches = new List<(DirectoryReference, RulesFileCache)>(Directories.Count());
-			lock (RootFolderToRulesFileCache)
+			string[] parts = dirRef.FullName.Split(Path.DirectorySeparatorChar);
+
+			Directory parentDir = Root;
+			for (int i = 0; i != parts.Length; ++i)
 			{
-				ParallelQueue queue = new();
-				foreach (DirectoryReference Directory in Directories)
+				parentDir.SubDirectories ??= new(DirectoryReference.Comparer);
+				if (!parentDir.SubDirectories.TryGetValue(parts[i], out Directory? dir))
 				{
-					RulesFileCache? Cache;
-					if (!RootFolderToRulesFileCache.TryGetValue(Directory, out Cache))
+					if (!createIfNotFound)
 					{
-						Cache = new RulesFileCache();
-						queue.Enqueue(() => FindAllRulesFilesRecursively(DirectoryItem.GetItemByDirectoryReference(Directory), Cache, queue));
+						return null;
 					}
-					Caches.Add((Directory, Cache));
+					dir = new Directory();
+					parentDir.SubDirectories[parts[i]] = dir;
 				}
-				queue.Drain();
+				parentDir = dir;
+			}
+
+			return parentDir;
+		}
 
-				ParallelQueue sortQueue = new();
-				foreach ((DirectoryReference Directory, RulesFileCache Cache) in Caches)
+		/// <summary>
+		/// Populate rule files from provided directories
+		/// </summary>
+		/// <param name="dirs">Directories to traverse</param>
+		public static void Populate(IEnumerable<DirectoryReference> dirs)
+		{
+			lock (Root)
+			{
+				ParallelQueue? queue = null;
+				foreach (DirectoryReference dirRef in dirs)
 				{
-					if (!RootFolderToRulesFileCache.ContainsKey(Directory))
+					Directory dir = FindDirectory(dirRef, true)!;
+
+					if (dir.IsScanned)
 					{
-						sortQueue.Enqueue(() =>
-						{
-							Cache.ModuleRules.Sort((A, B) => A.FullName.CompareTo(B.FullName));
-							Cache.TargetRules.Sort((A, B) => A.FullName.CompareTo(B.FullName));
-							Cache.AutomationModules.Sort((A, B) => A.FullName.CompareTo(B.FullName));
-							Cache.UbtPlugins.Sort((A, B) => A.FullName.CompareTo(B.FullName));
-						});
-						RootFolderToRulesFileCache[Directory] = Cache;
+						continue;
 					}
+
+					dir.IsScanned = true;
+					queue ??= new ();
+
+					queue.Enqueue(() =>
+					{
+						DirectoryItem dirItem = DirectoryItem.GetItemByDirectoryReference(dirRef);
+						PopulateRecursively(dirItem, dir, queue);
+					});
 				}
-				sortQueue.Drain();
+				queue?.Drain();
 			}
-			return Caches;
 		}
 
 		/// <summary>
-		/// Finds all the rules of the given type under a given directory
+		/// Get all rule files of type. Requires populate to have been run for this to work
 		/// </summary>
-		/// <param name="Directory">Directory to search</param>
-		/// <param name="Type">Type of rules to return</param>
-		/// <returns>List of rules files of the given type</returns>
-		public static IReadOnlyList<FileReference> FindAllRulesFiles(DirectoryReference Directory, RulesFileType Type)
+		/// <param name="dirs">Directories to traverse</param>
+		/// <param name="type">Type to return</param>
+		/// <returns></returns>
+		public static List<FileReference> Get(IEnumerable<DirectoryReference> dirs, RulesFileType type)
 		{
-			return FindAllRulesFiles([Directory], Type);
+			List<FileReference> files = [];
+			foreach (DirectoryReference dirRef in dirs)
+			{
+				Directory? dir = FindDirectory(dirRef, false);
+				if (dir != null)
+				{
+					GetRecursively(files, dir, type);
+				}
+			}
+			files.Sort(static (A, B) => A.FullName.CompareTo(B.FullName));
+			return files;
 		}
 
-		/// <summary>
-		/// Finds all the rules of the given type under a given directories
-		/// </summary>
-		/// <param name="Directories">Directories to search</param>
-		/// <param name="Type">Type of rules to return</param>
-		/// <returns>List of rules files of the given type</returns>
-		public static List<FileReference> FindAllRulesFiles(IEnumerable<DirectoryReference> Directories, RulesFileType Type)
+		private static void DistinctSortedList(List<FileReference> files)
+		{
+			int write = 1;
+			FileReference prev = files[0];
+
+			for (int i = 1; i < files.Count; i++)
+			{
+				FileReference x = files[i];
+				if (x != prev)
+				{
+					files[write++] = x;
+					prev = x;
+				}
+			}
+
+			if (write < files.Count)
+			{
+				files.RemoveRange(write, files.Count - write);
+			}
+		}
+
+		private static void GetRecursively(List<FileReference> outFiles, Directory dir, RulesFileType type)
 		{
-			List<(DirectoryReference, RulesFileCache)> Caches = PrefetchRulesFilesInternal(Directories);
-			List<FileReference> Files = [];			
-			foreach ((DirectoryReference Directory, RulesFileCache Cache) in Caches)
+			if (dir.Files != null)
 			{
-				// Get the list of files of the type we're looking for
-				if (Type == RulesFileType.Module)
+				if (type == RulesFileType.Module)
 				{
-					Files.AddRange(Cache.ModuleRules);
+					outFiles.AddRange(dir.Files.ModuleRules);
 				}
-				else if (Type == RulesFileType.Target)
+				else if (type == RulesFileType.Target)
 				{
-					Files.AddRange(Cache.TargetRules);
+					outFiles.AddRange(dir.Files.TargetRules);
 				}
-				else if (Type == RulesFileType.AutomationModule)
+				else if (type == RulesFileType.AutomationModule)
 				{
-					Files.AddRange(Cache.AutomationModules);
+					outFiles.AddRange(dir.Files.AutomationModules);
 				}
-				else if (Type == RulesFileType.UbtPlugin)
+				else if (type == RulesFileType.UbtPlugin)
 				{
-					Files.AddRange(Cache.UbtPlugins);
+					outFiles.AddRange(dir.Files.UbtPlugins);
 				}
 				else
 				{
-					throw new Exception($"Unhandled rules type: {Type}");
+					throw new Exception($"Unhandled rules type: {type}");
+				}
+			}
+			if (dir.SubDirectories != null)
+			{
+				foreach (KeyValuePair<string, Directory> kv in dir.SubDirectories)
+				{
+					GetRecursively(outFiles, kv.Value, type);
 				}
 			}
-			return Files;
 		}
 
-		/// <summary>
-		/// Search through a directory tree for any rules files
-		/// </summary>
-		/// <param name="Directory">The root directory to search from</param>
-		/// <param name="Cache">Receives all the discovered rules files</param>
-		/// <param name="Queue">Queue for adding additional tasks to</param>
-		private static void FindAllRulesFilesRecursively(DirectoryItem Directory, RulesFileCache Cache, ParallelQueue Queue)
+		private static void PopulateRecursively(DirectoryItem dirItem, Directory dir, ParallelQueue queue)
 		{
-			// Scan all the files in this directory
 			bool bSearchSubFolders = true;
 			bool bIsPlugin = false;
 			bool bIsCsProj = false;
 
-			foreach (FileItem File in Directory.EnumerateFiles())
+			foreach (FileItem File in dirItem.EnumerateFiles())
 			{
 				if (File.HasExtension(".build.cs"))
 				{
-					lock(Cache.ModuleRules)
-					{
-						Cache.ModuleRules.Add(File.Location);
-					}
+					dir.Files ??= new();
+					dir.Files.ModuleRules.Add(File.Location);
 					bSearchSubFolders = false;
 				}
 				else if (File.HasExtension(".target.cs"))
 				{
-					lock(Cache.TargetRules)
-					{
-						Cache.TargetRules.Add(File.Location);
-					}
+					dir.Files ??= new();
+					dir.Files.TargetRules.Add(File.Location);
 				}
 				else if (File.HasExtension(".automation.csproj"))
 				{
-					lock(Cache.AutomationModules)
-					{
-						Cache.AutomationModules.Add(File.Location);
-					}
+					dir.Files ??= new();
+					dir.Files.AutomationModules.Add(File.Location);
 					bSearchSubFolders = false;
 				}
 				else if (File.HasExtension(".ubtplugin.csproj"))
 				{
-					lock(Cache.UbtPlugins)
-					{
-						Cache.UbtPlugins.Add(File.Location);
-					}
+					dir.Files ??= new();
+					dir.Files.UbtPlugins.Add(File.Location);
 					bSearchSubFolders = false;
 				}
 				else if (File.HasExtension(".uplugin"))
@@ -309,46 +335,66 @@ namespace UnrealBuildBase
 				}
 			}
 
-			if (bIsPlugin)
+			bool checkExisting = dir.SubDirectories != null;
+
+			void QueueSubDir(DirectoryItem subDirItem)
 			{
-				DirectoryItem? SourceDir;
-				if (Directory.TryGetDirectory("Source", out SourceDir))
+				if (checkExisting && dir.SubDirectories!.TryGetValue(subDirItem.Name, out Directory? subDir))
+				{
+					if (subDir.IsScanned)
+					{
+						return;
+					}
+				}
+				else
 				{
-					Queue.Enqueue(() => FindAllRulesFilesRecursively(SourceDir, Cache, Queue));
+					subDir = new();
+					dir.SubDirectories ??= new(DirectoryReference.Comparer);
+					dir.SubDirectories.Add(subDirItem.Name, subDir);
 				}
+				subDir.IsScanned = true;
+
+				queue.Enqueue(() => PopulateRecursively(subDirItem, subDir, queue));
+			}
 
-				// search Build directory as well as it can contain automation projects
-				DirectoryItem? BuildDir;
-				if (Directory.TryGetDirectory("Build", out BuildDir))
+			if (bIsPlugin)
+			{
+				if (dirItem.TryGetDirectory("Source", out DirectoryItem? sourceDirItem))
 				{
-					Queue.Enqueue(() => FindAllRulesFilesRecursively(BuildDir, Cache, Queue));
+					QueueSubDir(sourceDirItem);
+				}
+				if (dirItem.TryGetDirectory("Build", out DirectoryItem? buildDirItem))
+				{
+					QueueSubDir(buildDirItem);
 				}
 			}
 
 			// If we didn't find anything to stop the search, search all the subdirectories too
-			if (bSearchSubFolders)
+			if (!bSearchSubFolders)
+			{
+				return;
+			}
+
+			foreach (DirectoryItem subDirectory in dirItem.EnumerateDirectories())
 			{
-				foreach (DirectoryItem SubDirectory in Directory.EnumerateDirectories())
+				if (bIsCsProj)
 				{
-					if (bIsCsProj)
+					if (subDirectory.Name == "bin" || subDirectory.Name == "obj") // Intermediate files.. we don't want to enter here
 					{
-						if (SubDirectory.Name == "bin" || SubDirectory.Name == "obj") // Intermediate files.. we don't want to enter here
-						{
-							continue;
-						}
-					}
-					else if (SubDirectory.Name.StartsWith("__External")) // This is a bit annoying but version control likes to leave around these huuuge hierarchies of empty directories
-					{
-						//Console.WriteLine("Most likely leftovers: " + Directory.Location);
 						continue;
 					}
-					//else if (SubDirectory.Name == "Intermediate")
-					//{
-					//	Console.WriteLine("Most likely leftovers: " + Directory.Location);
-					//}
-
-					Queue.Enqueue(() => FindAllRulesFilesRecursively(SubDirectory, Cache, Queue));
 				}
+				else if (subDirectory.Name.StartsWith("__External")) // This is a bit annoying but version control likes to leave around these huuuge hierarchies of empty directories
+				{
+					//Console.WriteLine("Most likely leftovers: " + Directory.Location);
+					continue;
+				}
+				//else if (subDirectory.Name == "Intermediate")
+				//{
+				//	Console.WriteLine("Most likely leftovers: " + Directory.Location);
+				//}
+
+				QueueSubDir(subDirectory);
 			}
 		}
 	}
diff --git a/Engine/Source/Programs/UnrealBuildTool/System/RulesAssembly.cs b/Engine/Source/Programs/UnrealBuildTool/System/RulesAssembly.cs
--- a/Engine/Source/Programs/UnrealBuildTool/System/RulesAssembly.cs
+++ b/Engine/Source/Programs/UnrealBuildTool/System/RulesAssembly.cs
@@ -8,6 +8,7 @@ using System.Linq;
 using System.Reflection;
 using System.Runtime.CompilerServices;
 using System.Text;
+using System.Threading.Tasks;
 using EpicGames.Core;
 using Microsoft.Extensions.Logging;
 using UnrealBuildBase;
@@ -221,35 +222,61 @@ namespace UnrealBuildTool
 			// is given because the method is implemented, not because it's called.
 			if (CompiledAssembly != null)
 			{
-				foreach (Type CompiledType in CompiledAssembly.GetTypes())
+				Task.Run(() =>
 				{
-					foreach (MethodInfo Method in CompiledType.GetMethods(BindingFlags.Instance | BindingFlags.Public | BindingFlags.DeclaredOnly))
+					//using var timelineScope = Timeline.ScopeEvent("PrejitRules");
+					ParallelQueue queue = new();
+					foreach (Type CompiledType in CompiledAssembly.GetTypes())
 					{
-						ObsoleteOverrideAttribute? Attribute = Method.GetCustomAttribute<ObsoleteOverrideAttribute>(true);
-						if (Attribute != null)
+						if (CompiledType.BaseType == typeof(ModuleRules))
 						{
-							FileReference? Location;
-							if (!TryGetFileNameFromType(CompiledType, out Location))
+							foreach (ConstructorInfo Constructor in CompiledType.GetConstructors())
 							{
-								Location = new FileReference(CompiledAssembly.Location);
+								ParameterInfo[] parameters = Constructor.GetParameters();
+								if (parameters.Length > 0)
+								{
+									if (parameters[0].ParameterType == typeof(ReadOnlyTargetRules))
+									{
+										queue.Enqueue(() => RuntimeHelpers.PrepareMethod(Constructor.MethodHandle));
+									}
+								}
 							}
-							Logger.LogWarning("{Location}: warning: {AttributeMessage}", Location, Attribute.Message);
 						}
 					}
-					if (CompiledType.BaseType == typeof(ModuleRules))
+					queue.Drain();
+					//timelineScope.Finish();
+
+					// Write out warnings as well
+					foreach (Type CompiledType in CompiledAssembly.GetTypes())
 					{
-						ConstructorInfo? Constructor = CompiledType.GetConstructor(new Type[] { typeof(TargetInfo) });
-						if (Constructor != null)
+						foreach (MethodInfo Method in CompiledType.GetMethods(BindingFlags.Instance | BindingFlags.Public | BindingFlags.DeclaredOnly))
+						{
+							ObsoleteOverrideAttribute? Attribute = Method.GetCustomAttribute<ObsoleteOverrideAttribute>(true);
+							if (Attribute != null)
+							{
+								FileReference? Location;
+								if (!TryGetFileNameFromType(CompiledType, out Location))
+								{
+									Location = new FileReference(CompiledAssembly.Location);
+								}
+								Logger.LogWarning("{Location}: warning: {AttributeMessage}", Location, Attribute.Message);
+							}
+						}
+						if (CompiledType.BaseType == typeof(ModuleRules))
 						{
-							FileReference? Location;
-							if (!TryGetFileNameFromType(CompiledType, out Location))
+							ConstructorInfo? Constructor = CompiledType.GetConstructor([typeof(TargetInfo)]);
+							if (Constructor != null)
 							{
-								Location = new FileReference(CompiledAssembly.Location);
+								FileReference? Location;
+								if (!TryGetFileNameFromType(CompiledType, out Location))
+								{
+									Location = new FileReference(CompiledAssembly.Location);
+								}
+								Logger.LogWarning("{Location}: warning: Module constructors should take a ReadOnlyTargetRules argument (rather than a TargetInfo argument) and pass it to the base class constructor from 4.15 onwards. Please update the method signature.", Location);
 							}
-							Logger.LogWarning("{Location}: warning: Module constructors should take a ReadOnlyTargetRules argument (rather than a TargetInfo argument) and pass it to the base class constructor from 4.15 onwards. Please update the method signature.", Location);
 						}
 					}
-				}
+				});
 			}
 		}
 
diff --git a/Engine/Source/Programs/UnrealBuildTool/System/RulesCompiler.cs b/Engine/Source/Programs/UnrealBuildTool/System/RulesCompiler.cs
--- a/Engine/Source/Programs/UnrealBuildTool/System/RulesCompiler.cs
+++ b/Engine/Source/Programs/UnrealBuildTool/System/RulesCompiler.cs
@@ -36,31 +36,13 @@ namespace UnrealBuildTool
 		/// <param name="ModuleFileToContext">Map of module files to their context</param>
 		private static void AddModuleRulesWithContext(DirectoryReference BaseDirectory, ModuleRulesContext ModuleContext, Dictionary<FileReference, ModuleRulesContext> ModuleFileToContext)
 		{
-			IReadOnlyList<FileReference> RulesFiles = Rules.FindAllRulesFiles(BaseDirectory, Rules.RulesFileType.Module);
+			IReadOnlyList<FileReference> RulesFiles = Rules.Get([BaseDirectory], Rules.RulesFileType.Module);
 			foreach (FileReference RulesFile in RulesFiles)
 			{
 				ModuleFileToContext[RulesFile] = ModuleContext;
 			}
 		}
 
-		/// <summary>
-		/// Find all the module rules files under a given directory
-		/// </summary>
-		/// <param name="BaseDirectory">The directory to search under</param>
-		/// <param name="SubDirectoryName">Name of the subdirectory to look under</param>
-		/// <param name="BaseModuleContext">The module context for each found rules instance</param>
-		/// <param name="DefaultUHTModuleType">The UHT module type</param>
-		/// <param name="ModuleFileToContext">Map of module files to their context</param>
-		private static void AddEngineModuleRulesWithContext(DirectoryReference BaseDirectory, string SubDirectoryName, ModuleRulesContext BaseModuleContext, UHTModuleType? DefaultUHTModuleType, Dictionary<FileReference, ModuleRulesContext> ModuleFileToContext)
-		{
-			DirectoryReference Directory = DirectoryReference.Combine(BaseDirectory, SubDirectoryName);
-			if (DirectoryLookupCache.DirectoryExists(Directory))
-			{
-				ModuleRulesContext ModuleContext = new ModuleRulesContext(BaseModuleContext) { DefaultUHTModuleType = DefaultUHTModuleType, bCanHotReload = true };
-				AddModuleRulesWithContext(Directory, ModuleContext, ModuleFileToContext);
-			}
-		}
-
 		/// <summary>
 		/// Lock for creating engine rule assembly
 		/// </summary>
@@ -148,10 +130,47 @@ namespace UnrealBuildTool
 			ModuleRulesContext DefaultModuleContext = new ModuleRulesContext(Scope, RootDirectories[0]);
 			List<FileReference> ProgramTargetFiles = new List<FileReference>();
 
+			// We want to prefetch the directories in one call to make it go as wide as possible against filesystem
+			{
+				List<DirectoryReference> populateDirs = [];
+				foreach (PluginInfo info in Plugins)
+				{
+					populateDirs.Add(DirectoryReference.Combine(info.Directory, "Source"));
+					populateDirs.Add(DirectoryReference.Combine(info.Directory, "Tests"));
+					foreach (FileReference ChildFile in info.ChildFiles)
+					{
+						populateDirs.Add(DirectoryReference.Combine(ChildFile.Directory, "Source"));
+						populateDirs.Add(DirectoryReference.Combine(ChildFile.Directory, "Tests"));
+					}
+				}
+				foreach (DirectoryReference RootDirectory in RootDirectories)
+				{
+					populateDirs.Add(DirectoryReference.Combine(RootDirectory, "Source"));
+					populateDirs.Add(DirectoryReference.Combine(RootDirectory, "Shaders"));
+				}
+				Rules.Populate(populateDirs);
+			}
+
+
 			foreach (DirectoryReference RootDirectory in RootDirectories)
 			{
 				using (GlobalTracer.Instance.BuildSpan("Finding engine modules").StartActive())
 				{
+					void AddEngineModuleRulesWithContext(DirectoryReference BaseDirectory, string SubDirectoryName, ModuleRulesContext BaseModuleContext, UHTModuleType? DefaultUHTModuleType, Dictionary<FileReference, ModuleRulesContext> ModuleFileToContext)
+					{
+						DirectoryReference Directory = DirectoryReference.Combine(BaseDirectory, SubDirectoryName);
+						if (DirectoryLookupCache.DirectoryExists(Directory))
+						{
+							ModuleRulesContext ModuleContext = new ModuleRulesContext(BaseModuleContext) { DefaultUHTModuleType = DefaultUHTModuleType, bCanHotReload = true };
+
+							IReadOnlyList<FileReference> RulesFiles = Rules.Get([Directory], Rules.RulesFileType.Module);
+							foreach (FileReference RulesFile in RulesFiles)
+							{
+								ModuleFileToContext[RulesFile] = ModuleContext;
+							}
+						}
+					}
+
 					DirectoryReference SourceDirectory = DirectoryReference.Combine(RootDirectory, "Source");
 
 					AddEngineModuleRulesWithContext(SourceDirectory, "Runtime", DefaultModuleContext, UHTModuleType.EngineRuntime, ModuleFileToContext);
@@ -195,7 +214,7 @@ namespace UnrealBuildTool
 
 				using (GlobalTracer.Instance.BuildSpan("Finding program targets").StartActive())
 				{
-					ProgramTargetFiles.AddRange(Rules.FindAllRulesFiles(SourceDirectory, Rules.RulesFileType.Target));
+					ProgramTargetFiles.AddRange(Rules.Get([SourceDirectory], Rules.RulesFileType.Target));
 				}
 			}
 
@@ -287,15 +306,6 @@ namespace UnrealBuildTool
 					ProjectDirectories.AddRange(Project.AdditionalRootDirectories);
 				}
 
-				// Find all the rules/plugins under the project source directories
-				foreach (DirectoryReference ProjectDirectory in ProjectDirectories)
-				{
-					DirectoryReference ProjectSourceDirectory = DirectoryReference.Combine(ProjectDirectory, "Source");
-
-					AddModuleRulesWithContext(ProjectSourceDirectory, DefaultModuleContext, ModuleFiles);
-					TargetFiles.AddRange(Rules.FindAllRulesFiles(ProjectSourceDirectory, Rules.RulesFileType.Target));
-				}
-
 				// Find all the project plugins
 				List<PluginInfo> ProjectPlugins = new List<PluginInfo>();
 				ProjectPlugins.AddRange(Plugins.ReadProjectPlugins(MainProjectDirectory));
@@ -309,6 +319,35 @@ namespace UnrealBuildTool
 					}
 				}
 
+				// We want to prefetch the directories in one call to make it go as wide as possible against filesystem
+				{
+					List<DirectoryReference> populateDirs = [];
+					foreach (PluginInfo info in ProjectPlugins)
+					{
+						populateDirs.Add(DirectoryReference.Combine(info.Directory, "Source"));
+						populateDirs.Add(DirectoryReference.Combine(info.Directory, "Tests"));
+						foreach (FileReference ChildFile in info.ChildFiles)
+						{
+							populateDirs.Add(DirectoryReference.Combine(ChildFile.Directory, "Source"));
+							populateDirs.Add(DirectoryReference.Combine(ChildFile.Directory, "Tests"));
+						}
+					}
+					foreach (DirectoryReference RootDirectory in ProjectDirectories)
+					{
+						populateDirs.Add(DirectoryReference.Combine(RootDirectory, "Source"));
+					}
+					Rules.Populate(populateDirs);
+				}
+
+				// Find all the rules/plugins under the project source directories
+				foreach (DirectoryReference ProjectDirectory in ProjectDirectories)
+				{
+					DirectoryReference ProjectSourceDirectory = DirectoryReference.Combine(ProjectDirectory, "Source");
+					AddModuleRulesWithContext(ProjectSourceDirectory, DefaultModuleContext, ModuleFiles);
+					TargetFiles.AddRange(Rules.Get([ProjectSourceDirectory], Rules.RulesFileType.Target));
+				}
+
+
 				Logger.LogTrace(" Found {Count} Plugins:", ProjectPlugins.Count);
 				ProjectPlugins.ForEach(x => { Logger.LogTrace("  {Plugin}", x.File); });
 
@@ -327,7 +366,7 @@ namespace UnrealBuildTool
 				if (DirectoryReference.Exists(ProjectIntermediateSourceDirectory))
 				{
 					AddModuleRulesWithContext(ProjectIntermediateSourceDirectory, DefaultModuleContext, ModuleFiles);
-					TargetFiles.AddRange(Rules.FindAllRulesFiles(ProjectIntermediateSourceDirectory, Rules.RulesFileType.Target));
+					TargetFiles.AddRange(Rules.Get([ProjectIntermediateSourceDirectory], Rules.RulesFileType.Target));
 				}
 
 				RulesAssembly ProjectRulesAssembly;
@@ -451,8 +490,6 @@ namespace UnrealBuildTool
 		/// <param name="ModuleFileToContext">Dictionary which is filled with mappings from the module file to its corresponding context</param>
 		private static void FindModuleRulesForPlugins(IReadOnlyList<PluginInfo> Plugins, ModuleRulesContext DefaultContext, Dictionary<FileReference, ModuleRulesContext> ModuleFileToContext)
 		{
-			Rules.PrefetchRulesFiles(Plugins.Select(x => DirectoryReference.Combine(x.Directory, "Source")));
-
 			foreach (PluginInfo Plugin in Plugins)
 			{
 				FindModuleRulesForPluginInFolder(Plugin, "Source", DefaultContext, ModuleFileToContext);
@@ -468,12 +505,10 @@ namespace UnrealBuildTool
 		/// <param name="TargetFiles">List of target files to add test target rules to</param>
 		private static void FindTestRulesForPlugins(IReadOnlyList<PluginInfo> Plugins, ModuleRulesContext DefaultContext, Dictionary<FileReference, ModuleRulesContext> ModuleFileToContext, List<FileReference> TargetFiles)
 		{
-			Rules.PrefetchRulesFiles(Plugins.Select(x => DirectoryReference.Combine(x.Directory, "Tests")));
-
 			foreach (PluginInfo Plugin in Plugins)
 			{
 				FindModuleRulesForPluginInFolder(Plugin, "Tests", DefaultContext, ModuleFileToContext);
-				TargetFiles.AddRange(Rules.FindAllRulesFiles(DirectoryReference.Combine(Plugin.Directory, "Tests"), Rules.RulesFileType.Target));
+				TargetFiles.AddRange(Rules.Get([DirectoryReference.Combine(Plugin.Directory, "Tests")], Rules.RulesFileType.Target));
 			}
 		}
 
@@ -486,10 +521,11 @@ namespace UnrealBuildTool
 		/// <param name="ModuleFileToContext">Dictionary which is filled with mappings from the module file to its corresponding context</param>
 		private static void FindModuleRulesForPluginInFolder(PluginInfo Plugin, string Folder, ModuleRulesContext DefaultContext, Dictionary<FileReference, ModuleRulesContext> ModuleFileToContext)
 		{
-			List<FileReference> PluginModuleFiles = Rules.FindAllRulesFiles(DirectoryReference.Combine(Plugin.Directory, Folder), Rules.RulesFileType.Module).ToList();
+			List<FileReference> PluginModuleFiles = [];
+			PluginModuleFiles.AddRange(Rules.Get([DirectoryReference.Combine(Plugin.Directory, Folder)], Rules.RulesFileType.Module));
 			foreach (FileReference ChildFile in Plugin.ChildFiles)
 			{
-				PluginModuleFiles.AddRange(Rules.FindAllRulesFiles(DirectoryReference.Combine(ChildFile.Directory, Folder), Rules.RulesFileType.Module));
+				PluginModuleFiles.AddRange(Rules.Get([DirectoryReference.Combine(ChildFile.Directory, Folder)], Rules.RulesFileType.Module));
 			}
 			foreach (FileReference ModuleFile in PluginModuleFiles)
 			{
