Silence static analysis complaint about large function

diff --git a/Engine/Source/Programs/Shared/EpicGames.UHT/Exporters/CodeGen/UhtHeaderCodeGeneratorCppFile.cs b/Engine/Source/Programs/Shared/EpicGames.UHT/Exporters/CodeGen/UhtHeaderCodeGeneratorCppFile.cs
--- a/Engine/Source/Programs/Shared/EpicGames.UHT/Exporters/CodeGen/UhtHeaderCodeGeneratorCppFile.cs
+++ b/Engine/Source/Programs/Shared/EpicGames.UHT/Exporters/CodeGen/UhtHeaderCodeGeneratorCppFile.cs
@@ -2532,6 +2532,8 @@ namespace EpicGames.UHT.Exporters.CodeGen
 			return builder;
 		}
 
+#pragma warning disable CA1505 //  'AppendNativeGeneratedInitCode' has a maintainability index of '5'. Rewrite or refactor the code to increase its maintainability index (MI) above '9'. 
+
 		private StringBuilder AppendNativeGeneratedInitCode(StringBuilder builder, UhtClass classObj, UhtUsedDefineScopes<UhtFunction> functions)
 		{
 			const string MetaDataParamsName = "Class_MetaDataParams";
@@ -2693,7 +2695,7 @@ namespace EpicGames.UHT.Exporters.CodeGen
 
 				// Functions
 				if (Session.IsUsingCompiledInObjectFormat(UhtCompiledInObjectFormat.ConstInit))
-				{ 
+				{
 					using UhtConditionalMacroBlock block = new(builder, "UE_WITH_CONSTINIT_UOBJECT", Session.IsUsingMultipleCompiledInObjectFormats);
 					builder.Append("\tstatic consteval UFunction* GetFirstFunction()\r\n");
 					builder.Append("\t{\r\n");
@@ -2701,7 +2703,7 @@ namespace EpicGames.UHT.Exporters.CodeGen
 					builder.Append("\t}\r\n");
 				}
 				if (Session.IsUsingCompiledInObjectFormat(UhtCompiledInObjectFormat.Params))
-				{ 
+				{
 					using UhtConditionalMacroBlock block = new(builder, "!UE_WITH_CONSTINIT_UOBJECT", Session.IsUsingMultipleCompiledInObjectFormats);
 					builder.AppendInstances(functions, UhtDefineScopeNames.WithEditor,
 						builder =>
@@ -2745,10 +2747,10 @@ namespace EpicGames.UHT.Exporters.CodeGen
 				builder.Append("\t};\r\n");
 
 				if (Session.IsUsingCompiledInObjectFormat(UhtCompiledInObjectFormat.ConstInit))
-				{ 
+				{
 					using UhtConditionalMacroBlock block = new(builder, "UE_WITH_CONSTINIT_UOBJECT", Session.IsUsingMultipleCompiledInObjectFormats);
 					builder.Append("\tstatic UE_CONSTINIT_UOBJECT_DECL inline const FStructBaseChain* StructBases[] = {\r\n");
-					if (classObj.SuperClass is not null) 
+					if (classObj.SuperClass is not null)
 					{
 						List<UhtClass> supers = new();
 						for (UhtClass? super = classObj.SuperClass; super is not null; super = super.SuperClass)
@@ -2765,7 +2767,7 @@ namespace EpicGames.UHT.Exporters.CodeGen
 				}
 
 				if (Session.IsUsingCompiledInObjectFormat(UhtCompiledInObjectFormat.Params))
-				{ 
+				{
 					using UhtConditionalMacroBlock block = new(builder, "UE_WITH_CONSTINIT_UOBJECT", Session.IsUsingMultipleCompiledInObjectFormats);
 					builder.Append("\tstatic const UECodeGen_Private::").AppendClassParamsType(classObj).Append(" ClassParams;\r\n");
 				}
@@ -2779,7 +2781,7 @@ namespace EpicGames.UHT.Exporters.CodeGen
 
 				// Dependent singletons
 				if (Session.IsUsingCompiledInObjectFormat(UhtCompiledInObjectFormat.Params))
-				{ 
+				{
 					using UhtConditionalMacroBlock block = new(builder, "!UE_WITH_CONSTINIT_UOBJECT", Session.IsUsingMultipleCompiledInObjectFormats);
 					builder.Append("UObject* (*const ").Append(staticsName).Append("::DependentSingletons[])() = {\r\n");
 					if (classObj.SuperClass != null && classObj.SuperClass != classObj)
@@ -2798,16 +2800,17 @@ namespace EpicGames.UHT.Exporters.CodeGen
 					{
 						using UhtConditionalMacroBlock block = new(builder, "UE_WITH_CONSTINIT_UOBJECT", Session.IsUsingMultipleCompiledInObjectFormats);
 						builder.Append($"const UE::CodeGen::ConstInit::FClassImplementedInterface {staticsName}::Interfaces[] = {{\r\n");
-						Action<UhtClass> appendConstInitInterfaceParam = (interfaceObj) => {
+						Action<UhtClass> appendConstInitInterfaceParam = (interfaceObj) =>
+						{
 							if (interfaceObj.ClassFlags.HasAnyFlags(EClassFlags.Interface))
 							{
 								string isVerseDirectInterface = interfaceObj.IsVerseField && classObj.VerseInterfaces.Contains(interfaceObj) ? "true" : "false";
 								builder.Append("\t{ ")
 									.Append(".Class = ")
 									.AppendSingletonRef(context, interfaceObj.AlternateObject, UhtSingletonType.ConstInit).Append(", ")
-									.Append(interfaceObj.IsVerseField 
-										? ".PointerOffset= 0, .bImplementedByK2 = true, " 
-										:$".PointerOffset = (int32)VTABLE_OFFSET({classObj.FullyQualifiedSourceName}, {interfaceObj.FullyQualifiedSourceName}), .bImplementedByK2 = false, ")
+									.Append(interfaceObj.IsVerseField
+										? ".PointerOffset= 0, .bImplementedByK2 = true, "
+										: $".PointerOffset = (int32)VTABLE_OFFSET({classObj.FullyQualifiedSourceName}, {interfaceObj.FullyQualifiedSourceName}), .bImplementedByK2 = false, ")
 									.Append($".bVerseDirectInterface = {isVerseDirectInterface}, ")
 									.Append("},")
 									.AppendObjectHash(classObj, context, interfaceObj.AlternateObject)
@@ -2848,12 +2851,12 @@ namespace EpicGames.UHT.Exporters.CodeGen
 
 				EClassFlags classFlags = classObj.ClassFlags & EClassFlags.SaveInCompiledInClasses;
 				// Propagate class flags without allowing removal for UClass::ClassFlags
-				for (UhtClass? superClass = classObj.SuperClass; superClass is not null; superClass = superClass.SuperClass) 
+				for (UhtClass? superClass = classObj.SuperClass; superClass is not null; superClass = superClass.SuperClass)
 				{
 					classFlags |= superClass.ClassFlags & EClassFlags.SaveInCompiledInClasses & EClassFlags.Inherit;
 				}
 				if (Session.IsUsingCompiledInObjectFormat(UhtCompiledInObjectFormat.ConstInit))
-				{ 
+				{
 					// constinit class definition
 					using UhtConditionalMacroBlock block = new(builder, "UE_WITH_CONSTINIT_UOBJECT", Session.IsUsingMultipleCompiledInObjectFormats);
 					builder.Append($"UE_CONSTINIT_UOBJECT_DECL TNoDestroy<U{classObj.EngineClassName}> {constinitName}{{NoDestroyConstEval,\r\n")
@@ -2908,7 +2911,7 @@ namespace EpicGames.UHT.Exporters.CodeGen
 				}
 
 				if (Session.IsUsingCompiledInObjectFormat(UhtCompiledInObjectFormat.Params))
-				{ 
+				{
 					// Class parameters
 					using UhtConditionalMacroBlock block = new(builder, "!UE_WITH_CONSTINIT_UOBJECT", Session.IsUsingMultipleCompiledInObjectFormats);
 					builder.Append("const UECodeGen_Private::").AppendClassParamsType(classObj).Append($" {staticsName}::ClassParams = {{\r\n");
@@ -3097,6 +3100,7 @@ namespace EpicGames.UHT.Exporters.CodeGen
 			}
 			return builder;
 		}
+#pragma warning restore CA1505 //  'AppendNativeGeneratedInitCode' has a maintainability index of '5'. Rewrite or refactor the code to increase its maintainability index (MI) above '9'. 
 
 		private static StringBuilder AppendImplementedInterfaceParam(StringBuilder builder, PropertyMemberContextImpl context, UhtClass classObj, UhtClass interfaceObj, UhtSingletonType singletonType)
 		{
