[EpicGames.UHT]
* Fixed so .generated.h is always written regardless if it is empty or not

diff --git a/Engine/Source/Programs/Shared/EpicGames.UHT/Exporters/CodeGen/UhtCodeGenerator.cs b/Engine/Source/Programs/Shared/EpicGames.UHT/Exporters/CodeGen/UhtCodeGenerator.cs
--- a/Engine/Source/Programs/Shared/EpicGames.UHT/Exporters/CodeGen/UhtCodeGenerator.cs
+++ b/Engine/Source/Programs/Shared/EpicGames.UHT/Exporters/CodeGen/UhtCodeGenerator.cs
@@ -194,6 +194,7 @@ namespace EpicGames.UHT.Exporters.CodeGen
 					HeaderInfos[headerFile.HeaderFileTypeIndex].Task = Factory.CreateTask(prereqs,
 						(IUhtExportFactory factory) =>
 						{
+							new UhtHeaderCodeGeneratorHFile(this, headerFile).Generate(factory);
 							new UhtHeaderCodeGeneratorCppFile(this, headerFile).Generate(factory);
 						});
 				}
