[EpicGames.Core]
* Fixed so AssemblyUtils.AddFileToAssemblyCache is not checking timestamp of already registered file if path is the same as before

diff --git a/Engine/Source/Programs/Shared/EpicGames.Core/AssemblyUtils.cs b/Engine/Source/Programs/Shared/EpicGames.Core/AssemblyUtils.cs
--- a/Engine/Source/Programs/Shared/EpicGames.Core/AssemblyUtils.cs
+++ b/Engine/Source/Programs/Shared/EpicGames.Core/AssemblyUtils.cs
@@ -102,54 +102,57 @@ namespace EpicGames.Core
 			}
 
 			string assemblyName = Path.GetFileNameWithoutExtension(assemblyPath);
-			DateTime assemblyLastWriteTime = File.GetLastWriteTimeUtc(assemblyPath);
 			lock(s_assemblyLocationCache)
 			{
-				if (s_assemblyLocationCache.ContainsKey(assemblyName))
+				if (s_assemblyLocationCache.TryGetValue(assemblyName, out string? oldAssemblyPath))
 				{
-					// We already have this assembly in our cache. Only replace it if the discovered file is newer (to avoid stale assemblies breaking stuff).
-					if (assemblyLastWriteTime > s_assemblyWriteTimes[assemblyName])
+					if (oldAssemblyPath != assemblyPath)
 					{
-						s_assemblyLocationCache[assemblyName] = assemblyPath;
-						s_assemblyWriteTimes[assemblyName] = assemblyLastWriteTime;
+						// We already have this assembly in our cache. Only replace it if the discovered file is newer (to avoid stale assemblies breaking stuff).
+						DateTime assemblyLastWriteTime = File.GetLastWriteTimeUtc(assemblyPath);
+						if (assemblyLastWriteTime > s_assemblyWriteTimes[assemblyName])
+						{
+							s_assemblyLocationCache[assemblyName] = assemblyPath;
+							s_assemblyWriteTimes[assemblyName] = assemblyLastWriteTime;
+						}
 					}
 				}
 				else
 				{
 					// This is the first copy of this assembly ... add it to our cache.
 					s_assemblyLocationCache.Add(assemblyName, assemblyPath);
-					s_assemblyWriteTimes.Add(assemblyName, assemblyLastWriteTime);
+					s_assemblyWriteTimes.Add(assemblyName, File.GetLastWriteTimeUtc(assemblyPath));
 				}
-			}
 
-			if (!s_addedToAssemblyResolver)
-			{
-				AppDomain.CurrentDomain.AssemblyResolve += (sender, args) =>
+				if (!s_addedToAssemblyResolver)
 				{
-					// Name is fully qualified assembly definition - e.g. "p4dn, Version=1.0.0.0, Culture=neutral, PublicKeyToken=ff968dc1933aba6f"
-					string assemblyName = args.Name!.Split(',')[0];
-					lock (s_assemblyLocationCache)
+					AppDomain.CurrentDomain.AssemblyResolve += (sender, args) =>
 					{
-						if (s_assemblyLocationCache.TryGetValue(assemblyName, out string? assemblyLocation))
+						// Name is fully qualified assembly definition - e.g. "p4dn, Version=1.0.0.0, Culture=neutral, PublicKeyToken=ff968dc1933aba6f"
+						string assemblyName = args.Name!.Split(',')[0];
+						lock (s_assemblyLocationCache)
 						{
-							// We have this assembly in our folder.					
-							if (File.Exists(assemblyLocation))
+							if (s_assemblyLocationCache.TryGetValue(assemblyName, out string? assemblyLocation))
 							{
-								// The assembly still exists, so load it.
-								return Assembly.LoadFile(assemblyLocation);
-							}
-							else
-							{
-								// The assembly no longer exists on disk, so remove it from our cache.
-								s_assemblyLocationCache.Remove(assemblyName);
+								// We have this assembly in our folder.					
+								if (File.Exists(assemblyLocation))
+								{
+									// The assembly still exists, so load it.
+									return Assembly.LoadFile(assemblyLocation);
+								}
+								else
+								{
+									// The assembly no longer exists on disk, so remove it from our cache.
+									s_assemblyLocationCache.Remove(assemblyName);
+								}
 							}
 						}
-					}
 
-					return null;
-				};
+						return null;
+					};
 
-				s_addedToAssemblyResolver = true;
+					s_addedToAssemblyResolver = true;
+				}
 			}
 		}
 
