Horde: Improve retry mechanism for Slack API client

Slack is enforcing stricter rate limits which the client must respect to avoid failing requests or fighting too much with other requests from the same Slack app.

diff --git a/Engine/Source/Programs/Horde/Plugins/Build/HordeServer.Build/Notifications/Sinks/SlackNotificationSink.cs b/Engine/Source/Programs/Horde/Plugins/Build/HordeServer.Build/Notifications/Sinks/SlackNotificationSink.cs
--- a/Engine/Source/Programs/Horde/Plugins/Build/HordeServer.Build/Notifications/Sinks/SlackNotificationSink.cs
+++ b/Engine/Source/Programs/Horde/Plugins/Build/HordeServer.Build/Notifications/Sinks/SlackNotificationSink.cs
@@ -236,10 +236,7 @@ namespace HordeServer.Notifications.Sinks
 		static readonly RedisListKey<int> s_redisIssueQueue = "slack/issue-queue";
 		readonly string _redisIssueLockPrefix;
 
-		readonly HttpClient _httpClient;
 		readonly SlackClient _slackClient;
-
-		readonly HttpClient? _adminHttpClient;
 		readonly SlackClient? _adminSlackClient;
 
 		/// <summary>
@@ -272,15 +269,11 @@ namespace HordeServer.Notifications.Sinks
 			_issueQueueTicker = clock.AddSharedTicker("slack-issues", TimeSpan.FromMinutes(1.0), ProcessIssueQueueAsync, _logger);
 			_redisIssueLockPrefix = "slack/issues/";
 
-			_httpClient = new HttpClient();
-			_httpClient.DefaultRequestHeaders.Add("Authorization", $"Bearer {_settings.SlackToken ?? ""}");
-			_slackClient = new SlackClient(_httpClient, _logger);
+			_slackClient = SlackClient.Create(_settings.SlackToken, _logger);
 
 			if (!String.IsNullOrEmpty(_settings.SlackAdminToken))
 			{
-				_adminHttpClient = new HttpClient();
-				_adminHttpClient.DefaultRequestHeaders.Add("Authorization", $"Bearer {_settings.SlackAdminToken}");
-				_adminSlackClient = new SlackClient(_adminHttpClient, _logger);
+				_adminSlackClient = SlackClient.Create(_settings.SlackAdminToken, _logger);
 			}
 
 			_jsonSerializerOptions = new JsonSerializerOptions();
@@ -297,8 +290,8 @@ namespace HordeServer.Notifications.Sinks
 		{
 			await _backgroundTask.DisposeAsync();
 			_userCache.Dispose();
-			_httpClient.Dispose();
-			_adminHttpClient?.Dispose();
+			_slackClient.Dispose();
+			_adminSlackClient?.Dispose();
 			await _issueQueueTicker.DisposeAsync();
 			await _escalateTicker.DisposeAsync();
 		}
diff --git a/Engine/Source/Programs/Horde/Plugins/Experimental/HordeServer.Experimental/Notifications/SlackNotificationProcessor.cs b/Engine/Source/Programs/Horde/Plugins/Experimental/HordeServer.Experimental/Notifications/SlackNotificationProcessor.cs
--- a/Engine/Source/Programs/Horde/Plugins/Experimental/HordeServer.Experimental/Notifications/SlackNotificationProcessor.cs
+++ b/Engine/Source/Programs/Horde/Plugins/Experimental/HordeServer.Experimental/Notifications/SlackNotificationProcessor.cs
@@ -45,8 +45,6 @@ namespace HordeServer.Experimental.Notifications
 		readonly IUserCollection _userCollection;
 		readonly IServerInfo _serverInfo;
 		readonly ILogger _logger;
-
-		readonly HttpClient? _httpClient;
 		readonly SlackClient? _slackClient;
 
 		/// <summary>
@@ -110,15 +108,12 @@ namespace HordeServer.Experimental.Notifications
 
 			if (!String.IsNullOrEmpty(buildServerConfig.Value.SlackToken))
 			{
-				_httpClient = new HttpClient();
-				_httpClient.DefaultRequestHeaders.Add("Authorization", $"Bearer {buildServerConfig.Value.SlackToken}");
-				_slackClient = new SlackClient(_httpClient, _logger);
+				_slackClient = SlackClient.Create(buildServerConfig.Value.SlackToken, _logger);
 			}
 		}
 
 		/// <inheritdoc/>
-		public void Dispose()
-			=> _httpClient?.Dispose();
+		public void Dispose() => _slackClient?.Dispose();
 
 		/// <summary>
 		/// Processes a Job complete notification
diff --git a/Engine/Source/Programs/Shared/EpicGames.Slack/EpicGames.Slack.csproj b/Engine/Source/Programs/Shared/EpicGames.Slack/EpicGames.Slack.csproj
--- a/Engine/Source/Programs/Shared/EpicGames.Slack/EpicGames.Slack.csproj
+++ b/Engine/Source/Programs/Shared/EpicGames.Slack/EpicGames.Slack.csproj
@@ -38,6 +38,8 @@
     <Folder Include="Blocks\" />
   </ItemGroup>
   <ItemGroup>
+    <PackageReference Include="Microsoft.Extensions.Http" Version="8.0.1" />
+    <PackageReference Include="Microsoft.Extensions.Http.Polly" Version="8.0.10" />
     <PackageReference Include="Microsoft.Extensions.Logging.Abstractions" Version="8.0.2" />
   </ItemGroup>
 </Project>
\ No newline at end of file
diff --git a/Engine/Source/Programs/Shared/EpicGames.Slack/SlackClient.cs b/Engine/Source/Programs/Shared/EpicGames.Slack/SlackClient.cs
--- a/Engine/Source/Programs/Shared/EpicGames.Slack/SlackClient.cs
+++ b/Engine/Source/Programs/Shared/EpicGames.Slack/SlackClient.cs
@@ -2,13 +2,18 @@
 
 using System;
 using System.Collections.Generic;
+using System.Diagnostics.CodeAnalysis;
+using System.Net;
 using System.Net.Http;
 using System.Text;
 using System.Text.Json;
 using System.Text.Json.Serialization;
 using System.Threading;
 using System.Threading.Tasks;
+using Microsoft.Extensions.Http;
 using Microsoft.Extensions.Logging;
+using Polly;
+using Polly.Extensions.Http;
 
 namespace EpicGames.Slack
 {
@@ -68,7 +73,7 @@ namespace EpicGames.Slack
 	/// <summary>
 	/// Wrapper around Slack client functionality
 	/// </summary>
-	public class SlackClient
+	public sealed class SlackClient : IDisposable
 	{
 		class SlackResponse
 		{
@@ -88,15 +93,54 @@ namespace EpicGames.Slack
 		/// </summary>
 		/// <param name="httpClient">Http client for connecting to slack. Should have the necessary authorization headers.</param>
 		/// <param name="logger">Logger interface</param>
-		public SlackClient(HttpClient httpClient, ILogger logger)
+		private SlackClient(HttpClient httpClient, ILogger logger)
 		{
 			_httpClient = httpClient;
-
-			_serializerOptions = new JsonSerializerOptions();
-			_serializerOptions.DefaultIgnoreCondition = JsonIgnoreCondition.WhenWritingNull;
-
+			_serializerOptions = new JsonSerializerOptions { DefaultIgnoreCondition = JsonIgnoreCondition.WhenWritingNull };
 			_logger = logger;
 		}
+		
+		/// <inheritdoc/>
+		public void Dispose()
+		{
+			_httpClient.Dispose();
+		}
+		
+		/// <summary>
+		/// Create a Slack client with a retry policy
+		/// </summary>
+		/// <param name="slackToken">Slack token to use</param>
+		/// <param name="logger">Logger</param>
+		[SuppressMessage("Reliability", "CA2000:Dispose objects before losing scope", Justification = "Polly handler is disposed by HttpClient")]
+		public static SlackClient Create(string? slackToken, ILogger logger)
+		{
+			// Slack has tight rate limiting so respecting their returned retried intervals is useful to avoid TooManyRequests errors
+			Polly.Retry.AsyncRetryPolicy<HttpResponseMessage> retryPolicy = HttpPolicyExtensions
+				.HandleTransientHttpError()
+				.OrResult(r => r.StatusCode == HttpStatusCode.TooManyRequests)
+				.WaitAndRetryAsync(
+					retryCount: 5,
+					sleepDurationProvider: (retryAttempt, outcome, context) =>
+					{
+						if (outcome.Result?.Headers.RetryAfter?.Delta.HasValue == true)
+						{
+							return outcome.Result.Headers.RetryAfter.Delta.Value;
+						}
+						return TimeSpan.FromSeconds(15);
+					},
+					onRetryAsync: (outcome, timespan, retryCount, context) =>
+					{
+						logger.LogInformation("Slack HTTP request to {Uri} was rate limited. Retry {RetryCount}/5. Waiting {WaitTime}s...",
+							outcome.Result.RequestMessage?.RequestUri, retryCount, (int)timespan.TotalSeconds);
+						
+						return Task.CompletedTask;
+					}
+				);
+			PolicyHttpMessageHandler pollyHandler = new (retryPolicy) { InnerHandler = new HttpClientHandler() };
+			HttpClient httpClient = new (pollyHandler, true); // SlackClient will dispose this
+			httpClient.DefaultRequestHeaders.Add("Authorization", $"Bearer {slackToken ?? ""}");
+			return new SlackClient(httpClient, logger);
+		}
 
 		static bool ShouldLogError<TResponse>(TResponse response) where TResponse : SlackResponse => !response.Ok;
 
@@ -107,15 +151,11 @@ namespace EpicGames.Slack
 
 		private async Task<TResponse> SendRequestAsync<TResponse>(string requestUrl, object request, Func<TResponse, bool> shouldLogError, CancellationToken cancellationToken) where TResponse : SlackResponse
 		{
-			using (HttpRequestMessage sendMessageRequest = new HttpRequestMessage(HttpMethod.Post, requestUrl))
-			{
-				string requestJson = JsonSerializer.Serialize(request, _serializerOptions);
-				using (StringContent messageContent = new StringContent(requestJson, Encoding.UTF8, "application/json"))
-				{
-					sendMessageRequest.Content = messageContent;
-					return await SendRequestAsync<TResponse>(sendMessageRequest, requestJson, shouldLogError, cancellationToken);
-				}
-			}
+			using HttpRequestMessage sendMessageRequest = new (HttpMethod.Post, requestUrl);
+			string requestJson = JsonSerializer.Serialize(request, _serializerOptions);
+			using StringContent messageContent = new (requestJson, Encoding.UTF8, "application/json");
+			sendMessageRequest.Content = messageContent;
+			return await SendRequestAsync<TResponse>(sendMessageRequest, requestJson, shouldLogError, cancellationToken);
 		}
 
 		private Task<TResponse> SendRequestAsync<TResponse>(HttpRequestMessage request, string requestJson, CancellationToken cancellationToken) where TResponse : SlackResponse
