EpicGames.Slack: Make constructor public again

Originally changed in 45782122.

diff --git a/Engine/Source/Programs/Shared/EpicGames.Slack/SlackClient.cs b/Engine/Source/Programs/Shared/EpicGames.Slack/SlackClient.cs
--- a/Engine/Source/Programs/Shared/EpicGames.Slack/SlackClient.cs
+++ b/Engine/Source/Programs/Shared/EpicGames.Slack/SlackClient.cs
@@ -93,7 +93,7 @@ namespace EpicGames.Slack
 		/// </summary>
 		/// <param name="httpClient">Http client for connecting to slack. Should have the necessary authorization headers.</param>
 		/// <param name="logger">Logger interface</param>
-		private SlackClient(HttpClient httpClient, ILogger logger)
+		public SlackClient(HttpClient httpClient, ILogger logger)
 		{
 			_httpClient = httpClient;
 			_serializerOptions = new JsonSerializerOptions { DefaultIgnoreCondition = JsonIgnoreCondition.WhenWritingNull };
