[UBT]
* Changed so FileItem and DirectoryItem does not have custom GetHashCode etc. We know the instances of these only exist once so one path equals one instance.
* Added FileItem.SortByPath, DirectoryItem.SortByPath comparers that should be used if we want to track these items sorted. (Usually smart to wait until the end to sort things)

diff --git a/Engine/Source/Programs/Shared/EpicGames.Build/System/DirectoryItem.cs b/Engine/Source/Programs/Shared/EpicGames.Build/System/DirectoryItem.cs
--- a/Engine/Source/Programs/Shared/EpicGames.Build/System/DirectoryItem.cs
+++ b/Engine/Source/Programs/Shared/EpicGames.Build/System/DirectoryItem.cs
@@ -14,7 +14,7 @@ namespace UnrealBuildBase
 	/// <summary>
 	/// Stores the state of a directory. May or may not exist.
 	/// </summary>
-	public class DirectoryItem : IComparable<DirectoryItem>, IEquatable<DirectoryItem>
+	public class DirectoryItem
 	{
 		/// <summary>
 		/// Full path to the directory on disk
@@ -129,7 +129,9 @@ namespace UnrealBuildBase
 			foreach (string file in files)
 			{
 				FileReference fileRef = FileReference.Combine(Location, file);
-				cacheFiles.Add(file, new FileItem(fileRef, file, new() { Attributes = FileAttributes.Normal }, this));
+				FileItem fileItem = FileItem.GetItemByFileReference(fileRef);
+				fileItem.Update(new() { Attributes = FileAttributes.Normal }); 
+				cacheFiles.Add(file, fileItem);
 			}
 
 			_cache = new(([], cacheFiles));
@@ -341,6 +343,11 @@ namespace UnrealBuildBase
 		/// <returns>Location of the directory</returns>
 		public override string ToString() => Location.ToString();
 
+		/// <summary>
+		/// Sort file items by full path
+		/// </summary>
+		public static readonly IComparer<DirectoryItem?> SortByPath = new SortByPathComparer();
+
 		/// <summary>
 		/// Writes out all the enumerated files full names sorted to OutFile
 		/// </summary>
@@ -360,68 +367,6 @@ namespace UnrealBuildBase
 			File.WriteAllLines(OutFile, AllFiles);
 		}
 
-		#region IComparable, IEquatbale
-		/// <inheritdoc/>
-		public int CompareTo(DirectoryItem? other) => Location.CompareTo(other?.Location);
-
-		/// <inheritdoc/>
-		public bool Equals(DirectoryItem? other) => Location.Equals(other?.Location);
-
-		/// <inheritdoc/>
-		public override bool Equals(object? obj)
-		{
-			if (ReferenceEquals(this, obj))
-			{
-				return true;
-			}
-
-			if (obj is null)
-			{
-				return false;
-			}
-
-			return Equals(obj as DirectoryItem);
-		}
-
-		/// <inheritdoc/>
-		public override int GetHashCode() => Location.GetHashCode();
-
-		public static bool operator ==(DirectoryItem? left, DirectoryItem? right)
-		{
-			if (left is null)
-			{
-				return right is null;
-			}
-
-			return left.Equals(right);
-		}
-
-		public static bool operator !=(DirectoryItem? left, DirectoryItem? right)
-		{
-			return !(left == right);
-		}
-
-		public static bool operator <(DirectoryItem? left, DirectoryItem? right)
-		{
-			return left is null ? right is not null : left.CompareTo(right) < 0;
-		}
-
-		public static bool operator <=(DirectoryItem? left, DirectoryItem? right)
-		{
-			return left is null || left.CompareTo(right) <= 0;
-		}
-
-		public static bool operator >(DirectoryItem? left, DirectoryItem? right)
-		{
-			return left is not null && left.CompareTo(right) > 0;
-		}
-
-		public static bool operator >=(DirectoryItem? left, DirectoryItem? right)
-		{
-			return left is null ? right is null : left.CompareTo(right) >= 0;
-		}
-		#endregion
-
 		class Info
 		{
 			internal bool Exists;
@@ -432,6 +377,26 @@ namespace UnrealBuildBase
 		Info? _info;
 		Lazy<(Dictionary<string, DirectoryItem>, Dictionary<string, FileItem>)> _cache;
 		static ConcurrentDictionary<DirectoryReference, DirectoryItem> LocationToItem = [];
+
+		class SortByPathComparer : IComparer<DirectoryItem?>
+		{
+			public int Compare(DirectoryItem? a, DirectoryItem? b)
+			{
+				if (a != null && b != null)
+				{
+					return a.Location.CompareTo(b.Location);
+				}
+				if (a == b)
+				{
+					return 0;
+				}
+				if (a is null)
+				{
+					return -1;
+				}
+				return 1;
+			}
+		}
 	}
 
 	/// <summary>
diff --git a/Engine/Source/Programs/Shared/EpicGames.Build/System/FileItem.cs b/Engine/Source/Programs/Shared/EpicGames.Build/System/FileItem.cs
--- a/Engine/Source/Programs/Shared/EpicGames.Build/System/FileItem.cs
+++ b/Engine/Source/Programs/Shared/EpicGames.Build/System/FileItem.cs
@@ -2,6 +2,7 @@
 
 using System;
 using System.Collections.Concurrent;
+using System.Collections.Generic;
 using System.Diagnostics;
 using System.IO;
 using System.Linq;
@@ -16,7 +17,7 @@ namespace UnrealBuildBase
 	/// Represents a file on disk that is used as an input or output of a build action. FileItem instances are unique for a given path. Use FileItem.GetItemByFileReference 
 	/// to get the FileItem for a specific path.
 	/// </summary>
-	public class FileItem : IComparable<FileItem>, IEquatable<FileItem>
+	public class FileItem
 	{
 		/// <summary>
 		/// The directory containing this file
@@ -93,7 +94,10 @@ namespace UnrealBuildBase
 		public void UpdateCachedDirectory(DirectoryItem Directory)
 		{
 			Debug.Assert(Directory.Location == Location.Directory);
-			CachedDirectory = new(Directory);
+			if (!CachedDirectory.IsValueCreated)
+			{
+				CachedDirectory = new(Directory);
+			}
 		}
 
 		/// <summary>
@@ -198,68 +202,15 @@ namespace UnrealBuildBase
 		/// <returns>Absolute path to this file item</returns>
 		public override string ToString() => Location.ToString();
 
-		#region IComparable, IEquatbale
-		/// <inheritdoc/>
-		public int CompareTo(FileItem? other) => Location.CompareTo(other?.Location);
-
-		/// <inheritdoc/>
-		public bool Equals(FileItem? other) => Location.Equals(other?.Location);
-
-		/// <inheritdoc/>
-		public override bool Equals(object? obj)
-		{
-			if (ReferenceEquals(this, obj))
-			{
-				return true;
-			}
-
-			if (obj is null)
-			{
-				return false;
-			}
-
-			return Equals(obj as FileItem);
-		}
-
-		/// <inheritdoc/>
-		public override int GetHashCode() => Location.GetHashCode();
-
-		public static bool operator ==(FileItem? left, FileItem? right)
-		{
-			if (left is null)
-			{
-				return right is null;
-			}
-
-			return left.Equals(right);
-		}
-
-		public static bool operator !=(FileItem? left, FileItem? right)
-		{
-			return !(left == right);
-		}
-
-		public static bool operator <(FileItem? left, FileItem? right)
-		{
-			return left is null ? right is not null : left.CompareTo(right) < 0;
-		}
-
-		public static bool operator <=(FileItem? left, FileItem? right)
-		{
-			return left is null || left.CompareTo(right) <= 0;
-		}
-
-		public static bool operator >(FileItem? left, FileItem? right)
-		{
-			return left is not null && left.CompareTo(right) > 0;
-		}
-
-		public static bool operator >=(FileItem? left, FileItem? right)
-		{
-			return left is null ? right is null : left.CompareTo(right) >= 0;
-		}
-		#endregion
+		/// <summary>
+		/// Sort file items by full path
+		/// </summary>
+		public static readonly IComparer<FileItem?> SortByPath = new SortByPathComparer();
 
+		/// <summary>
+		/// For debugging purposes, to know if info has been populated
+		/// </summary>
+		public bool HasInfo => _info != null;
 
 		internal FileItem(FileReference Location, string name, Info? info, DirectoryItem? directory)
 		{
@@ -318,6 +269,26 @@ namespace UnrealBuildBase
 		}
 
 		Info? _info;
+
+		class SortByPathComparer : IComparer<FileItem?>
+		{
+			public int Compare(FileItem? a, FileItem? b)
+			{
+				if (a != null && b != null)
+				{
+					return a.Location.CompareTo(b.Location);
+				}
+				if (a == b)
+				{
+					return 0;
+				}
+				if (a is null)
+				{
+					return -1;
+				}
+				return 1;
+			}
+		}
 	}
 
 	/// <summary>
@@ -356,8 +327,8 @@ namespace UnrealBuildBase
 		{
 			DirectoryItem Directory = Reader.ReadDirectoryItem()!;
 			string Name = Reader.ReadString()!;
-
-			FileItem FileItem = FileItem.GetItemByFileReference(FileReference.Combine(Directory.Location, Name));
+			FileReference fileRef = FileReference.Combine(Directory.Location, Name);
+			FileItem FileItem = FileItem.UniqueSourceFileMap.GetOrAdd(fileRef, loc => new FileItem(loc, Name, null, Directory));
 			FileItem.UpdateCachedDirectory(Directory);
 			return FileItem;
 		}
diff --git a/Engine/Source/Programs/UnrealBuildTool/Configuration/UEBuildTarget.cs b/Engine/Source/Programs/UnrealBuildTool/Configuration/UEBuildTarget.cs
--- a/Engine/Source/Programs/UnrealBuildTool/Configuration/UEBuildTarget.cs
+++ b/Engine/Source/Programs/UnrealBuildTool/Configuration/UEBuildTarget.cs
@@ -3169,12 +3169,12 @@ namespace UnrealBuildTool
 			}
 
 			HashSet<FileItem> allObjects = new(Binaries.SelectMany(x => x.InputObjects));
-			SortedSet<FileItem> allExis = new(allObjects.Select(x => FileItem.GetItemByFileReference(x.Location.ChangeExtension(".exi"))));
+			SortedSet<FileItem> allExis = new(allObjects.Select(x => FileItem.GetItemByFileReference(x.Location.ChangeExtension(".exi"))), FileItem.SortByPath);
 
 			void CreateStripAction(string Name, UEBuildBinaryType binaryType, IEnumerable< FileItem> toStrip, DirectoryReference IntermediateDirectory)
 			{
-				SortedSet<FileItem> toStripExis = new(toStrip.Select(x => FileItem.GetItemByFileReference(x.Location.ChangeExtension(".exi"))));
-				SortedSet<FileItem> otherExis = new(allExis);
+				SortedSet<FileItem> toStripExis = new(toStrip.Select(x => FileItem.GetItemByFileReference(x.Location.ChangeExtension(".exi"))), FileItem.SortByPath);
+				SortedSet<FileItem> otherExis = new(allExis, FileItem.SortByPath);
 				otherExis.ExceptWith(toStripExis);
 
 				FileItem exportsRsp = FileItem.GetItemByFileReference(FileReference.Combine(IntermediateDirectory, $"{Name}.exports.rsp"));
diff --git a/Engine/Source/Programs/UnrealBuildTool/System/Action.cs b/Engine/Source/Programs/UnrealBuildTool/System/Action.cs
--- a/Engine/Source/Programs/UnrealBuildTool/System/Action.cs
+++ b/Engine/Source/Programs/UnrealBuildTool/System/Action.cs
@@ -241,21 +241,21 @@ namespace UnrealBuildTool
 		/// <summary>
 		/// Every file this action depends on.  These files need to exist and be up to date in order for this action to even be considered
 		/// </summary>
-		public SortedSet<FileItem> PrerequisiteItems { get; set; } = [];
+		public SortedSet<FileItem> PrerequisiteItems { get; set; } = new(FileItem.SortByPath);
 
 		/// <summary>
 		/// The files that this action produces after completing
 		/// </summary>
-		public SortedSet<FileItem> ProducedItems { get; set; } = [];
+		public SortedSet<FileItem> ProducedItems { get; set; } = new(FileItem.SortByPath);
 
 		/// <summary>
 		/// Items that should be deleted before running this action
 		/// </summary>
-		public SortedSet<FileItem> DeleteItems { get; set; } = [];
+		public SortedSet<FileItem> DeleteItems { get; set; } = new(FileItem.SortByPath);
 
 		/// <summary>
 		/// Root paths for this action (generally engine root project root, toolchain root, sdk root)
-        /// Order needs to be maintained, so a list is used
+		/// Order needs to be maintained, so a list is used
 		/// </summary>
 		public CppRootPaths RootPaths { get; set; } = new();
 
@@ -381,9 +381,9 @@ namespace UnrealBuildTool
 		{
 			ActionType = InOther.ActionType;
 			ArtifactMode = InOther.ArtifactMode;
-			PrerequisiteItems = [.. InOther.PrerequisiteItems];
-			ProducedItems = [.. InOther.ProducedItems];
-			DeleteItems = [.. InOther.DeleteItems];
+			PrerequisiteItems = new (InOther.PrerequisiteItems, FileItem.SortByPath);
+			ProducedItems = new(InOther.ProducedItems, FileItem.SortByPath);
+			DeleteItems = new(InOther.DeleteItems, FileItem.SortByPath);
 			RootPaths = new(InOther.RootPaths);
 			DependencyListFile = InOther.DependencyListFile;
 			WorkingDirectory = InOther.WorkingDirectory;
@@ -430,9 +430,9 @@ namespace UnrealBuildTool
 			bShouldOutputStatusDescription = Reader.ReadBool();
 			bShouldOutputLog = Reader.ReadBool();
 			bProducesImportLibrary = Reader.ReadBool();
-			PrerequisiteItems = Reader.ReadSortedSet(() => Reader.ReadFileItem())!;
-			ProducedItems = Reader.ReadSortedSet(() => Reader.ReadFileItem())!;
-			DeleteItems = Reader.ReadSortedSet(() => Reader.ReadFileItem())!;
+			PrerequisiteItems = Reader.ReadSortedSet(() => Reader.ReadFileItem(), FileItem.SortByPath)!;
+			ProducedItems = Reader.ReadSortedSet(() => Reader.ReadFileItem(), FileItem.SortByPath)!;
+			DeleteItems = Reader.ReadSortedSet(() => Reader.ReadFileItem(), FileItem.SortByPath)!;
 			RootPaths = new(Reader);
 			DependencyListFile = Reader.ReadFileItem();
 			bUseActionHistory = Reader.ReadBool();
diff --git a/Engine/Source/Programs/UnrealBuildTool/System/ActionGraph.cs b/Engine/Source/Programs/UnrealBuildTool/System/ActionGraph.cs
--- a/Engine/Source/Programs/UnrealBuildTool/System/ActionGraph.cs
+++ b/Engine/Source/Programs/UnrealBuildTool/System/ActionGraph.cs
@@ -739,6 +739,13 @@ namespace UnrealBuildTool
 			string NewProducingAttributes = $"{RootAction.CommandPath.FullName} {RootAction.CommandArguments} (ver {RootAction.CommandVersion})";
 			foreach (FileItem ProducedItem in RootAction.ProducedItems)
 			{
+				ProducedItem.Directory.CacheFiles();
+
+				if (!ProducedItem.HasInfo)
+				{
+					//Console.WriteLine("HASINFO " + ProducedItem);
+				}
+
 				// Check if the command-line of the action previously used to produce the item is outdated.
 				if (ActionHistory.UpdateProducingAttributes(ProducedItem, NewProducingAttributes, Logger) && RootAction.bUseActionHistory)
 				{
@@ -751,6 +758,7 @@ namespace UnrealBuildTool
 					bIsOutdated = true;
 				}
 
+
 				// If the produced file doesn't exist or has zero size, consider it outdated.  The zero size check is to detect cases
 				// where aborting an earlier compile produced invalid zero-sized obj files, but that may cause actions where that's
 				// legitimate output to always be considered outdated.
@@ -784,6 +792,8 @@ namespace UnrealBuildTool
 						continue;
 					}
 
+					PrerequisiteItem.Directory.CacheFiles();
+
 					if (PrerequisiteItem.Exists)
 					{
 						// allow a 1 second slop for network copies
diff --git a/Engine/Source/Programs/UnrealBuildTool/System/HotReload.cs b/Engine/Source/Programs/UnrealBuildTool/System/HotReload.cs
--- a/Engine/Source/Programs/UnrealBuildTool/System/HotReload.cs
+++ b/Engine/Source/Programs/UnrealBuildTool/System/HotReload.cs
@@ -1127,7 +1127,7 @@ namespace UnrealBuildTool
 						}
 					}
 				}
-				NewAction.PrerequisiteItems = new SortedSet<FileItem>(UpdatePrerequisiteItems);
+				NewAction.PrerequisiteItems = new (UpdatePrerequisiteItems, FileItem.SortByPath);
 
 				// Update this action's list of produced items too
 				List<FileItem> UpdateProducedItems = new List<FileItem>(NewAction.ProducedItems);
@@ -1146,7 +1146,7 @@ namespace UnrealBuildTool
 						AffectedOriginalFileItemAndNewFileItemMap.Add(OriginalProducedItem, NewProducedItem);
 					}
 				}
-				NewAction.ProducedItems = new SortedSet<FileItem>(UpdateProducedItems);
+				NewAction.ProducedItems = new (UpdateProducedItems, FileItem.SortByPath);
 
 				// Fix up the list of items to delete too
 				List<FileItem> UpdateDeleteItems = new List<FileItem>(NewAction.DeleteItems);
@@ -1158,7 +1158,7 @@ namespace UnrealBuildTool
 						UpdateDeleteItems[Idx] = NewItem;
 					}
 				}
-				NewAction.DeleteItems = new SortedSet<FileItem>(UpdateDeleteItems);
+				NewAction.DeleteItems = new (UpdateDeleteItems, FileItem.SortByPath);
 
 				// The status description of the item has the file name, so we'll update it too
 				NewAction.StatusDescription = ReplaceBaseFileName(Action.StatusDescription, OriginalFileNameWithoutExtension, NewFileNameWithoutExtension);
@@ -1190,7 +1190,7 @@ namespace UnrealBuildTool
 						UpdatePrerequisiteItems[ItemIndex] = NewFileItem;
 					}
 				}
-				NewAction.PrerequisiteItems = new SortedSet<FileItem>(UpdatePrerequisiteItems);
+				NewAction.PrerequisiteItems = new (UpdatePrerequisiteItems, FileItem.SortByPath);
 				Action.Inner = NewAction;
 			}
 
