[EpicGames.Build]
* Fixed so populate dirs force-write the FileItem info

diff --git a/Engine/Source/Programs/Shared/EpicGames.Build/System/DirectoryItem.cs b/Engine/Source/Programs/Shared/EpicGames.Build/System/DirectoryItem.cs
--- a/Engine/Source/Programs/Shared/EpicGames.Build/System/DirectoryItem.cs
+++ b/Engine/Source/Programs/Shared/EpicGames.Build/System/DirectoryItem.cs
@@ -130,7 +130,7 @@ namespace UnrealBuildBase
 			{
 				FileReference fileRef = FileReference.Combine(Location, file);
 				FileItem fileItem = FileItem.GetItemByFileReference(fileRef);
-				fileItem.Update(new() { Attributes = FileAttributes.Normal, Size = long.MaxValue }); 
+				fileItem.Update(new() { Attributes = FileAttributes.Normal, Size = long.MaxValue }, true); 
 				cacheFiles.Add(file, fileItem);
 			}
 
@@ -276,7 +276,7 @@ namespace UnrealBuildBase
 					{
 						FileReference loc = FileReference.Combine(Location, entry.Name);
 						FileItem file = FileItem.UniqueSourceFileMap.GetOrAdd(loc, loc => new FileItem(loc, entry.Name, new() { Attributes = entry.Attributes, Length = entry.Size, LastWriteTimeUtc = entry.LastWriteTimeUtc, CreationTimeUtc = entry.CreationTimeUtc }, this));
-						file.Update(entry);
+						file.Update(entry, false);
 
 						// There are folders in linux sdk that has files with same name but different casing.
 						// Ideally FileReference.Comparer should be case sensitive on linux/mac but I don't dare changing that right now
diff --git a/Engine/Source/Programs/Shared/EpicGames.Build/System/FileItem.cs b/Engine/Source/Programs/Shared/EpicGames.Build/System/FileItem.cs
--- a/Engine/Source/Programs/Shared/EpicGames.Build/System/FileItem.cs
+++ b/Engine/Source/Programs/Shared/EpicGames.Build/System/FileItem.cs
@@ -262,9 +262,9 @@ namespace UnrealBuildBase
 			return info;
 		}
 
-		internal void Update(DirectoryEnumerator.Entry entry)
+		internal void Update(DirectoryEnumerator.Entry entry, bool force)
 		{
-			if (_info == null)
+			if (_info == null || force)
 			{
 				_info = new() { Attributes = entry.Attributes, Length = entry.Size, LastWriteTimeUtc = entry.LastWriteTimeUtc, CreationTimeUtc = entry.CreationTimeUtc };
 			}
