Added feature to Horde's issues to be grouped based on customizable rules

diff --git a/Engine/Source/Programs/Horde/Plugins/Build/HordeServer.Build.Tests/Issues/IssueGroupingRulesTests.cs b/Engine/Source/Programs/Horde/Plugins/Build/HordeServer.Build.Tests/Issues/IssueGroupingRulesTests.cs
new file mode 100644
--- /dev/null
+++ b/Engine/Source/Programs/Horde/Plugins/Build/HordeServer.Build.Tests/Issues/IssueGroupingRulesTests.cs
@@ -0,0 +1,102 @@
+// Copyright Epic Games, Inc. All Rights Reserved.
+
+using System.Text.Json;
+using EpicGames.Horde.Issues;
+
+using HordeServer.Issues;
+using MongoDB.Bson;
+using MongoDB.Bson.Serialization;
+
+namespace HordeServer.Tests.Issues
+{
+	[TestClass]
+	public class IssueGroupingRulesTests
+	{
+		[TestMethod]
+		public void GroupingRulesBsonDocumentSerialization()
+		{
+			// ------------------------------------------------------------------------------------------------------------
+			// Required Setup
+			// Force the static constructor of IssueCollection to run so that the setup for MongoDB to be able to serialize our classes is done
+			// This part is usually run when anyone creates an IssueCollection class which we won't do in this test.
+			System.Runtime.CompilerServices.RuntimeHelpers.RunClassConstructor(typeof(IssueCollection).TypeHandle);
+			// ------------------------------------------------------------------------------------------------------------
+
+			// Test DefaultGroupingRules
+			IIssueGroupingRules defaultGroupingRules = DefaultIssueGroupingRules.Instance;
+			IssueCollection.IssueFingerprintDocument defaultDocumentInstance = new IssueCollection.IssueFingerprintDocument(new NewIssueFingerprint("Test Default", "Test Default Template Summary", "...", defaultGroupingRules));
+			Assert.IsNotNull(defaultDocumentInstance);
+
+			// Test CustomizableGroupingRules
+			IIssueGroupingRules customGroupingRules = new CustomizableIssueGroupingRules();
+			IssueCollection.IssueFingerprintDocument customDocumentInstance = new IssueCollection.IssueFingerprintDocument(new NewIssueFingerprint("Test Custom", "Test Custom Template Summary", "...", customGroupingRules));
+			Assert.IsNotNull(customDocumentInstance);
+
+			byte[] defaultDocumentBson = defaultDocumentInstance.ToBson();
+			byte[] customDocumentBson = customDocumentInstance.ToBson();
+
+			string readableDefaultDocumentBson = BsonSerializer.Deserialize<BsonDocument>(defaultDocumentBson).ToString();
+			string readableCustomDocumentBson = BsonSerializer.Deserialize<BsonDocument>(customDocumentBson).ToString();
+
+			Console.WriteLine(readableDefaultDocumentBson); // Debug Output
+			Console.WriteLine(readableCustomDocumentBson); // Debug Output
+
+			Assert.AreEqual("{ \"Type\" : \"Test Default\", \"inc\" : [], \"exc\" : [], \"met\" : [], \"fl2\" : \"...\", \"summary\" : \"Test Default Template Summary\" }", readableDefaultDocumentBson);
+			Assert.AreEqual("{ \"Type\" : \"Test Custom\", \"inc\" : [], \"exc\" : [], \"met\" : [], \"fl2\" : \"...\", \"gr\" : { \"_t\" : \"c\", \"DefaultRule\" : 0, \"SpanRule\" : 6 }, \"summary\" : \"Test Custom Template Summary\" }", readableCustomDocumentBson);
+
+			// Test deserializing without the new GroupingRule field
+			string noGroupingRulesStr = "{ \"Type\" : \"Old Issue\", \"inc\" : [], \"exc\" : [], \"met\" : [], \"fl2\" : \"...\", \"summary\" : \"Old Template Summary\" }";
+			IssueCollection.IssueFingerprintDocument deserializedNoGroupingRules = BsonSerializer.Deserialize<IssueCollection.IssueFingerprintDocument>(noGroupingRulesStr);
+			Assert.IsNotNull(deserializedNoGroupingRules);
+
+			// Test re-serializing an old issue that did not have GroupingRules
+			byte[] noGroupingRulesBson = deserializedNoGroupingRules.ToBson();
+
+			string readableNoGroupingRulesBson = BsonSerializer.Deserialize<BsonDocument>(noGroupingRulesBson).ToString();
+
+			Console.WriteLine(readableNoGroupingRulesBson); // Debug Output
+
+			Assert.AreEqual("{ \"Type\" : \"Old Issue\", \"inc\" : [], \"exc\" : [], \"met\" : [], \"fl2\" : \"...\", \"summary\" : \"Old Template Summary\" }", readableNoGroupingRulesBson);
+		}
+
+		[TestMethod]
+		public void GroupingRulesJsonDocumentSerialization()
+		{
+			// Test DefaultGroupingRules
+			IIssueGroupingRules defaultGroupingRules = DefaultIssueGroupingRules.Instance;
+			IssueFingerprint defaultFingerprintInstance = new IssueFingerprint();
+			Assert.IsNotNull(defaultFingerprintInstance);
+			defaultFingerprintInstance.GroupingRules = defaultGroupingRules;
+
+			// Test CustomizableGroupingRules
+			IIssueGroupingRules customGroupingRules = new CustomizableIssueGroupingRules();
+			IssueFingerprint customFingerprintInstance = new IssueFingerprint();
+			Assert.IsNotNull(customFingerprintInstance);
+			customFingerprintInstance.GroupingRules = customGroupingRules;
+
+			string defaultFingerprintJson = JsonSerializer.Serialize(defaultFingerprintInstance, IssueFingerprint.JsonSerializerOptions);
+			string customFingerprintJson = JsonSerializer.Serialize(customFingerprintInstance, IssueFingerprint.JsonSerializerOptions);
+			Assert.AreNotEqual(defaultFingerprintJson, customFingerprintJson);
+
+			string expectedDefault = "{\"type\":\"None\",\"summary\":\"\",\"rejectKeys\":null,\"filter\":\"...\",\"groupingRules\":{\"$type\":\"d\"}}";
+			string expectedCustom = "{\"type\":\"None\",\"summary\":\"\",\"rejectKeys\":null,\"filter\":\"...\",\"groupingRules\":{\"$type\":\"c\",\"DefaultRule\":0,\"SpanRule\":6}}";
+
+			Console.WriteLine(defaultFingerprintJson); // Debug Output
+			Console.WriteLine(customFingerprintJson); // Debug Output
+
+			Assert.AreEqual(expectedDefault, defaultFingerprintJson);
+			Assert.AreEqual(expectedCustom, customFingerprintJson);
+
+			// Test deserializing without the new GroupingRule field
+			string noGroupingRulesStr = "{\"type\":\"None\",\"summary\":\"\",\"rejectKeys\":null,\"filter\":\"...\"}";
+
+			IssueFingerprint? deserializedNoGroupingRules = JsonSerializer.Deserialize<IssueFingerprint>(noGroupingRulesStr);
+			Assert.IsNotNull(deserializedNoGroupingRules);
+
+			// Test re-serializing an old issue that did not have GroupingRules
+			string noGroupingRulesJson = JsonSerializer.Serialize(deserializedNoGroupingRules, IssueFingerprint.JsonSerializerOptions);
+
+			Assert.AreEqual(expectedDefault, noGroupingRulesJson);
+		}
+	}
+}
diff --git a/Engine/Source/Programs/Horde/Plugins/Build/HordeServer.Build/Issues/IIssueCollection.cs b/Engine/Source/Programs/Horde/Plugins/Build/HordeServer.Build/Issues/IIssueCollection.cs
--- a/Engine/Source/Programs/Horde/Plugins/Build/HordeServer.Build/Issues/IIssueCollection.cs
+++ b/Engine/Source/Programs/Horde/Plugins/Build/HordeServer.Build/Issues/IIssueCollection.cs
@@ -363,17 +363,22 @@ namespace HordeServer.Issues
 		/// <inheritdoc/>
 		public string ChangeFilter { get; set; }
 
+		/// <inheritdoc/>
+		public IIssueGroupingRules GroupingRules { get; set; }
+
 		/// <summary>
 		/// Constructor
 		/// </summary>
 		/// <param name="type">The type of issue</param>
 		/// <param name="summaryTemplate">Template for the summary string to display for the issue</param>
 		/// <param name="changeFilter">Filter for changes covered by this issue</param>
-		public NewIssueFingerprint(string type, string summaryTemplate, string changeFilter)
+		/// <param name="groupingRules">Grouping rules used by this issue</param>
+		public NewIssueFingerprint(string type, string summaryTemplate, string changeFilter, IIssueGroupingRules groupingRules)
 		{
 			Type = type;
 			SummaryTemplate = summaryTemplate;
 			ChangeFilter = changeFilter;
+			GroupingRules = groupingRules;
 		}
 
 		/// <summary>
@@ -385,7 +390,8 @@ namespace HordeServer.Issues
 		/// <param name="rejectKeys">Keys which should not match with this issue</param>
 		/// <param name="metadata">Additional metadata added by the issue handler</param>
 		/// <param name="changeFilter">Filter for changes covered by this issue</param>
-		public NewIssueFingerprint(string type, string summaryTemplate, IEnumerable<IssueKey> keys, IEnumerable<IssueKey>? rejectKeys, IEnumerable<IssueMetadata>? metadata, string changeFilter)
+		/// <param name="groupingRules">Grouping rules used by this issue</param>
+		public NewIssueFingerprint(string type, string summaryTemplate, IEnumerable<IssueKey> keys, IEnumerable<IssueKey>? rejectKeys, IEnumerable<IssueMetadata>? metadata, string changeFilter, IIssueGroupingRules groupingRules)
 		{
 			Type = type;
 			SummaryTemplate = summaryTemplate;
@@ -401,6 +407,8 @@ namespace HordeServer.Issues
 			}
 
 			ChangeFilter = changeFilter;
+
+			GroupingRules = groupingRules;
 		}
 
 		/// <summary>
@@ -408,7 +416,7 @@ namespace HordeServer.Issues
 		/// </summary>
 		/// <param name="other">The fingerprint to copy from</param>
 		public NewIssueFingerprint(IIssueFingerprint other)
-			: this(other.Type, other.SummaryTemplate, other.Keys, other.RejectKeys, other.Metadata, other.ChangeFilter)
+			: this(other.Type, other.SummaryTemplate, other.Keys, other.RejectKeys, other.Metadata, other.ChangeFilter, other.GroupingRules)
 		{
 		}
 
diff --git a/Engine/Source/Programs/Horde/Plugins/Build/HordeServer.Build/Issues/IssueCollection.cs b/Engine/Source/Programs/Horde/Plugins/Build/HordeServer.Build/Issues/IssueCollection.cs
--- a/Engine/Source/Programs/Horde/Plugins/Build/HordeServer.Build/Issues/IssueCollection.cs
+++ b/Engine/Source/Programs/Horde/Plugins/Build/HordeServer.Build/Issues/IssueCollection.cs
@@ -29,6 +29,7 @@ using Microsoft.Extensions.Options;
 using MongoDB.Bson;
 using MongoDB.Bson.Serialization;
 using MongoDB.Bson.Serialization.Attributes;
+using MongoDB.Bson.Serialization.Serializers;
 using MongoDB.Driver;
 using OpenTelemetry.Trace;
 
@@ -261,7 +262,7 @@ namespace HordeServer.Issues
 			}
 		}
 
-		class IssueFingerprintDocument : IIssueFingerprint
+		internal class IssueFingerprintDocument : IIssueFingerprint
 		{
 			public string Type { get; set; }
 
@@ -289,6 +290,9 @@ namespace HordeServer.Issues
 			[BsonElement("fl2")]
 			public string? ChangeFilter { get; set; }
 
+			[BsonElement("gr")]
+			public IIssueGroupingRules GroupingRules { get; set; } = DefaultIssueGroupingRules.Instance;
+
 			[Obsolete("Use ChangeFilter instead")]
 			[BsonElement("flt"), BsonIgnoreIfNull]
 			public List<string>? LegacyChangeFilter
@@ -303,6 +307,7 @@ namespace HordeServer.Issues
 			private IssueFingerprintDocument()
 			{
 				Type = String.Empty;
+				GroupingRules = DefaultIssueGroupingRules.Instance;
 			}
 
 			public IssueFingerprintDocument(IIssueFingerprint fingerprint)
@@ -321,6 +326,8 @@ namespace HordeServer.Issues
 				}
 
 				ChangeFilter = fingerprint.ChangeFilter;
+
+				GroupingRules = fingerprint.GroupingRules;
 			}
 
 			#region Legacy
@@ -687,6 +694,28 @@ namespace HordeServer.Issues
 				cm.MapProperty(x => x.Key).SetElementName("k");
 				cm.MapProperty(x => x.Value).SetElementName("v");
 			});
+
+			Func<Type, bool> issueGroupingRulesSerializer = delegate (Type type)
+			{
+				if (type == typeof(DefaultIssueGroupingRules) || type == typeof(CustomizableIssueGroupingRules))
+				{
+					return true;
+				}
+				return false;
+			};
+			BsonSerializer.RegisterSerializer(new ObjectSerializer(issueGroupingRulesSerializer));
+
+			BsonClassMap.RegisterClassMap<IssueFingerprintDocument>(cm =>
+			{
+				cm.AutoMap();
+				cm.GetMemberMap(c => c.GroupingRules).SetIgnoreIfDefault(true).SetDefaultValue(DefaultIssueGroupingRules.Instance);
+			});
+
+			BsonClassMap.RegisterClassMap<CustomizableIssueGroupingRules>(cm =>
+			{
+				cm.AutoMap();
+				cm.SetDiscriminator("c");
+			});
 		}
 
 		public IssueCollection(IMongoService mongoService, IRedisService redisService, IUserCollection userCollection, ICommitService commitService, IAuditLogFactory<int> auditLogFactory, ITelemetryWriter telemetryWriter, IOptionsMonitor<BuildConfig> buildConfig, Tracer tracer, ILogger<IssueCollection> logger)
diff --git a/Engine/Source/Programs/Horde/Plugins/Build/HordeServer.Build/Issues/IssueService.cs b/Engine/Source/Programs/Horde/Plugins/Build/HordeServer.Build/Issues/IssueService.cs
--- a/Engine/Source/Programs/Horde/Plugins/Build/HordeServer.Build/Issues/IssueService.cs
+++ b/Engine/Source/Programs/Horde/Plugins/Build/HordeServer.Build/Issues/IssueService.cs
@@ -823,7 +823,7 @@ namespace HordeServer.Issues
 			HashSet<IssueEventGroupInternal> eventGroups = new HashSet<IssueEventGroupInternal>();
 			foreach (IssueEventGroup issue in issues)
 			{
-				NewIssueFingerprint fingerprint = new NewIssueFingerprint(issue.Type, issue.SummaryTemplate, issue.ChangeFilter);
+				NewIssueFingerprint fingerprint = new NewIssueFingerprint(issue.Type, issue.SummaryTemplate, issue.ChangeFilter, issue.GroupingRules);
 				fingerprint.Keys.UnionWith(issue.Keys);
 				fingerprint.Metadata.UnionWith(issue.Metadata);
 
diff --git a/Engine/Source/Programs/Shared/EpicGames.Horde/Issues/Handlers/ExternalIssueHandler.cs b/Engine/Source/Programs/Shared/EpicGames.Horde/Issues/Handlers/ExternalIssueHandler.cs
--- a/Engine/Source/Programs/Shared/EpicGames.Horde/Issues/Handlers/ExternalIssueHandler.cs
+++ b/Engine/Source/Programs/Shared/EpicGames.Horde/Issues/Handlers/ExternalIssueHandler.cs
@@ -17,7 +17,6 @@ namespace EpicGames.Horde.Issues.Handlers
 	public class ExternalIssueHandler : IssueHandler
 	{
 		static readonly JsonNodeOptions s_jsonNodeOptions = new JsonNodeOptions { PropertyNameCaseInsensitive = true };
-		static readonly JsonSerializerOptions? s_jsonSerializerOptions = new JsonSerializerOptions { PropertyNameCaseInsensitive = true };
 
 		readonly List<IssueEventGroup> _issues = new List<IssueEventGroup>();
 
@@ -66,7 +65,7 @@ namespace EpicGames.Horde.Issues.Handlers
 					JsonObject? issue = properties["$issue"] as JsonObject;
 					if (issue != null)
 					{
-						IssueFingerprint? fingerprint = JsonSerializer.Deserialize<IssueFingerprint>(issue, s_jsonSerializerOptions);
+						IssueFingerprint? fingerprint = JsonSerializer.Deserialize<IssueFingerprint>(issue, IssueFingerprint.JsonSerializerOptions);
 						if (fingerprint != null)
 						{
 							return fingerprint;
diff --git a/Engine/Source/Programs/Shared/EpicGames.Horde/Issues/Handlers/LocalizationIssueHandler.cs b/Engine/Source/Programs/Shared/EpicGames.Horde/Issues/Handlers/LocalizationIssueHandler.cs
--- a/Engine/Source/Programs/Shared/EpicGames.Horde/Issues/Handlers/LocalizationIssueHandler.cs
+++ b/Engine/Source/Programs/Shared/EpicGames.Horde/Issues/Handlers/LocalizationIssueHandler.cs
@@ -42,6 +42,21 @@ namespace EpicGames.Horde.Issues.Handlers
 		/// <inheritdoc/>
 		public override int Priority => 10;
 
+
+		private CustomizableIssueGroupingRules? _localizationIssueGroupingRules = null;
+		private IIssueGroupingRules LocalizationIssueGroupingRules
+		{
+			get
+			{
+				if(_localizationIssueGroupingRules == null)
+				{
+					_localizationIssueGroupingRules = new CustomizableIssueGroupingRules();
+					_localizationIssueGroupingRules.SpanRule = CustomizableIssueGroupingRules.FingerprintMatchingRule.FullFingerprint;
+				}
+				return _localizationIssueGroupingRules;
+			}
+		}
+
 		/// <summary>
 		/// Constructor
 		/// </summary>
@@ -267,6 +282,8 @@ namespace EpicGames.Horde.Issues.Handlers
 						issueGroup = _issuesWithoutClearSuspects;
 					}
 
+					issueGroup.GroupingRules = LocalizationIssueGroupingRules;
+
 					// Add current event to the chosen issueGroup
 					issueGroup.Events.Add(issueEvent);
 
diff --git a/Engine/Source/Programs/Shared/EpicGames.Horde/Issues/IIssueFingerprint.cs b/Engine/Source/Programs/Shared/EpicGames.Horde/Issues/IIssueFingerprint.cs
--- a/Engine/Source/Programs/Shared/EpicGames.Horde/Issues/IIssueFingerprint.cs
+++ b/Engine/Source/Programs/Shared/EpicGames.Horde/Issues/IIssueFingerprint.cs
@@ -40,6 +40,11 @@ namespace EpicGames.Horde.Issues
 		/// Filter for changes that should be included in this issue
 		/// </summary>
 		public string ChangeFilter { get; }
+
+		/// <summary>
+		/// Issue grouping rules
+		/// </summary>
+		public IIssueGroupingRules GroupingRules { get; }
 	}
 
 	/// <summary>
@@ -55,33 +60,11 @@ namespace EpicGames.Horde.Issues
 		/// <returns>True is the fingerprints match</returns>
 		public static bool IsMatch(this IIssueFingerprint fingerprint, IIssueFingerprint other)
 		{
-			if (!fingerprint.Type.Equals(other.Type, StringComparison.Ordinal))
-			{
-				return false;
-			}
-			if (!fingerprint.SummaryTemplate.Equals(other.SummaryTemplate, StringComparison.Ordinal))
-			{
-				return false;
-			}
-			if (fingerprint.Keys.Count == 0)
+			if (fingerprint.GroupingRules != null)
 			{
-				if (other.Keys.Count > 0)
-				{
-					return false;
-				}
+				return fingerprint.GroupingRules.IsMatch(fingerprint, other);
 			}
-			else
-			{
-				if (!fingerprint.Keys.Any(x => other.Keys.Contains(x)))
-				{
-					return false;
-				}
-			}
-			if (fingerprint.RejectKeys != null && fingerprint.RejectKeys.Any(x => other.Keys.Contains(x)))
-			{
-				return false;
-			}
-			return true;
+			return DefaultIssueGroupingRules.Instance.IsMatch(fingerprint, other);
 		}
 
 		/// <summary>
@@ -92,15 +75,11 @@ namespace EpicGames.Horde.Issues
 		/// <returns>True is the fingerprints match</returns>
 		public static bool IsMatchForNewSpan(this IIssueFingerprint fingerprint, IIssueFingerprint other)
 		{
-			if (!fingerprint.Type.Equals(other.Type, StringComparison.Ordinal))
-			{
-				return false;
-			}
-			if (fingerprint.RejectKeys != null && fingerprint.RejectKeys.Any(x => other.Keys.Contains(x)))
+			if(fingerprint.GroupingRules != null)
 			{
-				return false;
+				return fingerprint.GroupingRules.IsMatchForNewSpan(fingerprint, other);
 			}
-			return true;
+			return DefaultIssueGroupingRules.Instance.IsMatchForNewSpan(fingerprint, other);
 		}
 	}
 }
diff --git a/Engine/Source/Programs/Shared/EpicGames.Horde/Issues/IssueEvent.cs b/Engine/Source/Programs/Shared/EpicGames.Horde/Issues/IssueEvent.cs
--- a/Engine/Source/Programs/Shared/EpicGames.Horde/Issues/IssueEvent.cs
+++ b/Engine/Source/Programs/Shared/EpicGames.Horde/Issues/IssueEvent.cs
@@ -89,6 +89,11 @@ namespace EpicGames.Horde.Issues
 		/// </summary>
 		public string ChangeFilter { get; set; }
 
+		/// <summary>
+		/// Issue grouping rules
+		/// </summary>
+		public IIssueGroupingRules GroupingRules { get; set; }
+
 		/// <summary>
 		/// Individual log events
 		/// </summary>
@@ -105,6 +110,7 @@ namespace EpicGames.Horde.Issues
 			Type = type;
 			SummaryTemplate = summaryTemplate;
 			ChangeFilter = changeFilter;
+			GroupingRules = DefaultIssueGroupingRules.Instance;
 		}
 
 		/// <summary>
diff --git a/Engine/Source/Programs/Shared/EpicGames.Horde/Issues/IssueFingerprint.cs b/Engine/Source/Programs/Shared/EpicGames.Horde/Issues/IssueFingerprint.cs
--- a/Engine/Source/Programs/Shared/EpicGames.Horde/Issues/IssueFingerprint.cs
+++ b/Engine/Source/Programs/Shared/EpicGames.Horde/Issues/IssueFingerprint.cs
@@ -2,7 +2,9 @@
 
 using System;
 using System.Collections.Generic;
+using System.Text.Json;
 using System.Text.Json.Serialization;
+using System.Text.Json.Serialization.Metadata;
 
 #pragma warning disable CA2227 // Change 'x' to be read-only by removing the property setter
 
@@ -52,6 +54,12 @@ namespace EpicGames.Horde.Issues
 		[JsonPropertyName("filter")]
 		public string ChangeFilter { get; set; } = "...";
 
+		/// <summary>
+		/// Issue grouping rules
+		/// </summary>
+		[JsonPropertyName("groupingRules")]
+		public IIssueGroupingRules GroupingRules { get; set; } = DefaultIssueGroupingRules.Instance;
+
 		/// <summary>
 		/// Keys set which is null when empty
 		/// </summary>
@@ -89,5 +97,36 @@ namespace EpicGames.Horde.Issues
 			Type = type;
 			SummaryTemplate = summaryTemplate;
 		}
+
+		private static readonly JsonSerializerOptions? s_jsonSerializerOptions = new JsonSerializerOptions()
+		{
+			TypeInfoResolver = new DefaultJsonTypeInfoResolver
+			{
+				Modifiers =
+					{
+						typeInfo =>
+						{
+							if(typeInfo.Type == typeof(IIssueGroupingRules))
+							{
+								typeInfo.PolymorphismOptions = new JsonPolymorphismOptions
+								{
+									TypeDiscriminatorPropertyName = "$type",
+									IgnoreUnrecognizedTypeDiscriminators = false,
+									UnknownDerivedTypeHandling = JsonUnknownDerivedTypeHandling.FailSerialization
+								};
+
+								typeInfo.PolymorphismOptions.DerivedTypes.Add(new JsonDerivedType(typeof(DefaultIssueGroupingRules), "d"));
+								typeInfo.PolymorphismOptions.DerivedTypes.Add(new JsonDerivedType(typeof(CustomizableIssueGroupingRules), "c"));
+							}
+						}
+					}
+			},
+			PropertyNameCaseInsensitive = true
+		};
+
+		/// <summary>
+		/// Always use those JsonSerializerOptions when serializing/deserializing a JSON IssueFingerprint
+		/// </summary>
+		public static JsonSerializerOptions? JsonSerializerOptions => s_jsonSerializerOptions;
 	}
 }
diff --git a/Engine/Source/Programs/Shared/EpicGames.Horde/Issues/IssueGroupingRules.cs b/Engine/Source/Programs/Shared/EpicGames.Horde/Issues/IssueGroupingRules.cs
new file mode 100644
--- /dev/null
+++ b/Engine/Source/Programs/Shared/EpicGames.Horde/Issues/IssueGroupingRules.cs
@@ -0,0 +1,184 @@
+// Copyright Epic Games, Inc. All Rights Reserved.
+
+using System;
+using System.Linq;
+using System.Text.Json.Serialization;
+
+namespace EpicGames.Horde.Issues
+{
+	/// <summary>
+	/// A set of grouping rules for an <see cref="IssueEventGroup"/>.
+	/// </summary>
+	public interface IIssueGroupingRules
+	{
+		/// <summary>
+		/// Checks if a fingerprint matches another fingerprint
+		/// </summary>
+		/// <param name="fingerprint">The first fingerprint to compare</param>
+		/// <param name="other">The other fingerprint to compare to</param>
+		/// <returns>True is the fingerprints match</returns>
+		public bool IsMatch(IIssueFingerprint fingerprint, IIssueFingerprint other);
+
+		/// <summary>
+		/// Checks if a fingerprint matches another fingerprint for creating a new span
+		/// </summary>
+		/// <param name="fingerprint">The first fingerprint to compare</param>
+		/// <param name="other">The other fingerprint to compare to</param>
+		/// <returns>True is the fingerprints match</returns>
+		public bool IsMatchForNewSpan(IIssueFingerprint fingerprint, IIssueFingerprint other);
+	}
+
+	/// <summary>
+	/// The default issue grouping rules <see cref="IIssueGroupingRules"/>
+	/// </summary>
+	public class DefaultIssueGroupingRules : IIssueGroupingRules
+	{
+		private static readonly DefaultIssueGroupingRules s_instance = new DefaultIssueGroupingRules();
+
+		private DefaultIssueGroupingRules() { }
+
+		/// <summary>
+		/// Singleton instance
+		/// </summary>
+		public static IIssueGroupingRules Instance => s_instance;
+
+		/// <summary>
+		/// Default matching is done via CustomizableIssueGroupingRules default's constructor
+		/// </summary>
+		private static readonly CustomizableIssueGroupingRules s_impl = new CustomizableIssueGroupingRules();
+
+		/// <inheritdoc/>
+		public bool IsMatch(IIssueFingerprint fingerprint, IIssueFingerprint other)
+		{
+			return s_impl.IsMatch(fingerprint, other);
+		}
+
+		/// <inheritdoc/>
+		public bool IsMatchForNewSpan(IIssueFingerprint fingerprint, IIssueFingerprint other)
+		{
+			return s_impl.IsMatchForNewSpan(fingerprint, other);
+		}
+	}
+
+	/// <summary>
+	/// Customizable issue grouping rules <see cref="IIssueGroupingRules"/>
+	/// </summary>
+	public class CustomizableIssueGroupingRules : IIssueGroupingRules
+	{
+		/// <summary>
+		/// Based on which rule set should the issue fingerprints be matched
+		/// </summary>
+		public enum FingerprintMatchingRule
+		{
+			/// <summary>
+			/// Match using all fingerprint information available.
+			/// </summary>
+			FullFingerprint = 0,
+			/// <summary>
+			/// Match issues based on the issue type in the fingerprint but also consider the rejected keys.
+			/// </summary>
+			TypeAndRejectKeys = 6,
+			/// <summary>
+			/// Match issues only based on the issue type in the fingerprint.
+			/// </summary>
+			TypeOnly = 8,
+			/// <summary>
+			/// Never match this issue with any other issues. Be careful using this, as a large quantity of issues could be created and not matched with other related issues.
+			/// </summary>
+			NeverMatch = 10
+		};
+
+		/// <summary>
+		/// Based on which fingerprinting rule set should the issue fingerprints be matched by default?
+		/// </summary>
+		public FingerprintMatchingRule DefaultRule { get; set; } = FingerprintMatchingRule.FullFingerprint;
+
+		/// <summary>
+		/// Based on which fingerprinting rule set should the issue fingerprints be matched across spans?
+		/// </summary>
+		public FingerprintMatchingRule SpanRule { get; set; } = FingerprintMatchingRule.TypeAndRejectKeys;
+
+
+		/// <inheritdoc/>
+		public bool IsMatch(IIssueFingerprint fingerprint, IIssueFingerprint otherFingerprint)
+		{
+			return IsRuleBasedMatch(DefaultRule, fingerprint, otherFingerprint);
+		}
+
+
+		/// <inheritdoc/>
+		public bool IsMatchForNewSpan(IIssueFingerprint fingerprint, IIssueFingerprint otherFingerprint)
+		{
+			return IsRuleBasedMatch(SpanRule, fingerprint, otherFingerprint);
+		}
+
+		/// <summary>
+		/// Checks if a fingerprint matches another fingerprint based on a specific rule
+		/// </summary>
+		/// <param name="rule">The ruleset to use for matching fingerprints</param>
+		/// <param name="fingerprint">The first fingerprint to compare</param>
+		/// <param name="otherFingerprint">The other fingerprint to compare to</param>
+		/// <returns>True is the fingerprints match</returns>
+		private static bool IsRuleBasedMatch(FingerprintMatchingRule rule, IIssueFingerprint fingerprint, IIssueFingerprint otherFingerprint)
+		{
+			switch (rule)
+			{
+				case FingerprintMatchingRule.TypeOnly:
+					return IsTypeOnlyMatch(fingerprint, otherFingerprint);
+
+				case FingerprintMatchingRule.TypeAndRejectKeys:
+					return IsTypeAndRejectKeysMatch(fingerprint, otherFingerprint);
+
+				case FingerprintMatchingRule.FullFingerprint:
+					return IsFullFingerprintMatch(fingerprint, otherFingerprint);
+
+				case FingerprintMatchingRule.NeverMatch:
+					return false;
+
+				default:
+					return true;
+			}
+		}
+
+		/// <summary>
+		/// Check to match a fingerprint to another fingerprint using only their types
+		/// </summary>
+		/// <param name="fingerprint">The first fingerprint to compare</param>
+		/// <param name="otherFingerprint">The other fingerprint to compare to</param>
+		/// <returns>True is the fingerprints match</returns>
+		private static bool IsTypeOnlyMatch(IIssueFingerprint fingerprint, IIssueFingerprint otherFingerprint)
+		{
+			return fingerprint.GroupingRules.Equals(otherFingerprint.GroupingRules)
+				&& fingerprint.Type.Equals(otherFingerprint.Type, StringComparison.Ordinal);
+		}
+
+		/// <summary>
+		/// Check to match a fingerprint to another fingerprint using their types and reject keys
+		/// </summary>
+		/// <param name="fingerprint">The first fingerprint to compare</param>
+		/// <param name="otherFingerprint">The other fingerprint to compare to</param>
+		/// <returns>True is the fingerprints match</returns>
+		private static bool IsTypeAndRejectKeysMatch(IIssueFingerprint fingerprint, IIssueFingerprint otherFingerprint)
+		{
+			return fingerprint.GroupingRules.Equals(otherFingerprint.GroupingRules)
+				&& fingerprint.Type.Equals(otherFingerprint.Type, StringComparison.Ordinal)
+				&& (fingerprint.RejectKeys == null || !fingerprint.RejectKeys.Any(x => otherFingerprint.Keys.Contains(x)));
+		}
+
+		/// <summary>
+		/// Full check to match a fingerprint to another fingerprint
+		/// </summary>
+		/// <param name="fingerprint">The first fingerprint to compare</param>
+		/// <param name="otherFingerprint">The other fingerprint to compare to</param>
+		/// <returns>True is the fingerprints match</returns>
+		private static bool IsFullFingerprintMatch(IIssueFingerprint fingerprint, IIssueFingerprint otherFingerprint)
+		{
+			return fingerprint.GroupingRules.Equals(otherFingerprint.GroupingRules)
+				&& fingerprint.Type.Equals(otherFingerprint.Type, StringComparison.Ordinal)
+				&& fingerprint.SummaryTemplate.Equals(otherFingerprint.SummaryTemplate, StringComparison.Ordinal)
+				&& (fingerprint.Keys.Count != 0 || otherFingerprint.Keys.Count == 0)
+				&& (fingerprint.Keys.Count == 0 || fingerprint.Keys.Any(x => otherFingerprint.Keys.Contains(x)))
+				&& (fingerprint.RejectKeys == null || !fingerprint.RejectKeys.Any(x => otherFingerprint.Keys.Contains(x)));
+		}
+	}
+}
