Fixed coding standard violations from CL 45881300

diff --git a/Engine/Source/Programs/Shared/EpicGames.Horde/Issues/IIssueFingerprint.cs b/Engine/Source/Programs/Shared/EpicGames.Horde/Issues/IIssueFingerprint.cs
--- a/Engine/Source/Programs/Shared/EpicGames.Horde/Issues/IIssueFingerprint.cs
+++ b/Engine/Source/Programs/Shared/EpicGames.Horde/Issues/IIssueFingerprint.cs
@@ -1,8 +1,6 @@
 // Copyright Epic Games, Inc. All Rights Reserved.
 
-using System;
 using System.Collections.Generic;
-using System.Linq;
 
 namespace EpicGames.Horde.Issues
 {
diff --git a/Engine/Source/Programs/Shared/EpicGames.Horde/Issues/IssueGroupingRules.cs b/Engine/Source/Programs/Shared/EpicGames.Horde/Issues/IssueGroupingRules.cs
--- a/Engine/Source/Programs/Shared/EpicGames.Horde/Issues/IssueGroupingRules.cs
+++ b/Engine/Source/Programs/Shared/EpicGames.Horde/Issues/IssueGroupingRules.cs
@@ -2,7 +2,6 @@
 
 using System;
 using System.Linq;
-using System.Text.Json.Serialization;
 
 namespace EpicGames.Horde.Issues
 {
@@ -98,14 +97,12 @@ namespace EpicGames.Horde.Issues
 		/// </summary>
 		public FingerprintMatchingRule SpanRule { get; set; } = FingerprintMatchingRule.TypeAndRejectKeys;
 
-
 		/// <inheritdoc/>
 		public bool IsMatch(IIssueFingerprint fingerprint, IIssueFingerprint otherFingerprint)
 		{
 			return IsRuleBasedMatch(DefaultRule, fingerprint, otherFingerprint);
 		}
 
-
 		/// <inheritdoc/>
 		public bool IsMatchForNewSpan(IIssueFingerprint fingerprint, IIssueFingerprint otherFingerprint)
 		{
