[UBT]
* Made Rules class (more) thread safe. Needed for parallel makefile creation

diff --git a/Engine/Source/Programs/Shared/EpicGames.Build/System/Rules.cs b/Engine/Source/Programs/Shared/EpicGames.Build/System/Rules.cs
--- a/Engine/Source/Programs/Shared/EpicGames.Build/System/Rules.cs
+++ b/Engine/Source/Programs/Shared/EpicGames.Build/System/Rules.cs
@@ -6,6 +6,7 @@ using System.Collections.Generic;
 using System.IO;
 using System.Linq;
 using System.Text;
+using System.Threading;
 using System.Threading.Tasks;
 
 namespace UnrealBuildBase
@@ -158,6 +159,8 @@ namespace UnrealBuildBase
 			internal bool IsScanned;
 		}
 		static Directory Root = new();
+		static ReaderWriterLockSlim RootLock = new();
+
 
 		static Directory? FindDirectory(DirectoryReference dirRef, bool createIfNotFound)
 		{
@@ -188,7 +191,8 @@ namespace UnrealBuildBase
 		/// <param name="dirs">Directories to traverse</param>
 		public static void Populate(IEnumerable<DirectoryReference> dirs)
 		{
-			lock (Root)
+			RootLock.EnterWriteLock();
+			try
 			{
 				ParallelQueue? queue = null;
 				foreach (DirectoryReference dirRef in dirs)
@@ -201,7 +205,7 @@ namespace UnrealBuildBase
 					}
 
 					dir.IsScanned = true;
-					queue ??= new ();
+					queue ??= new();
 
 					queue.Enqueue(() =>
 					{
@@ -211,6 +215,14 @@ namespace UnrealBuildBase
 				}
 				queue?.Drain();
 			}
+			catch (Exception)
+			{
+				throw;
+			}
+			finally
+			{
+				RootLock.ExitWriteLock();
+			}
 		}
 
 		/// <summary>
@@ -221,17 +233,29 @@ namespace UnrealBuildBase
 		/// <returns></returns>
 		public static List<FileReference> Get(IEnumerable<DirectoryReference> dirs, RulesFileType type)
 		{
-			List<FileReference> files = [];
-			foreach (DirectoryReference dirRef in dirs)
+			RootLock.EnterReadLock();
+			try
 			{
-				Directory? dir = FindDirectory(dirRef, false);
-				if (dir != null)
+				List<FileReference> files = [];
+				foreach (DirectoryReference dirRef in dirs)
 				{
-					GetRecursively(files, dir, type);
+					Directory? dir = FindDirectory(dirRef, false);
+					if (dir != null)
+					{
+						GetRecursively(files, dir, type);
+					}
 				}
+				files.Sort(static (A, B) => A.FullName.CompareTo(B.FullName));
+				return files;
+			}
+			catch (Exception)
+			{
+				throw;
+			}
+			finally
+			{
+				RootLock.ExitReadLock();
 			}
-			files.Sort(static (A, B) => A.FullName.CompareTo(B.FullName));
-			return files;
 		}
 
 		private static void DistinctSortedList(List<FileReference> files)
