[EpicGames.Core]
* Added xml doc to public methods in StackStringBuffer

diff --git a/Engine/Source/Programs/Shared/EpicGames.Core/StackStringBuffer.cs b/Engine/Source/Programs/Shared/EpicGames.Core/StackStringBuffer.cs
--- a/Engine/Source/Programs/Shared/EpicGames.Core/StackStringBuffer.cs
+++ b/Engine/Source/Programs/Shared/EpicGames.Core/StackStringBuffer.cs
@@ -27,6 +27,9 @@ namespace EpicGames.Core
 			_chars = initialBuffer;
 		}
 
+		/// <summary>
+		/// Length of string buffer
+		/// </summary>
 		public int Length => _pos;
 
 		public ref char this[int index] => ref _chars[index];
@@ -93,6 +96,9 @@ namespace EpicGames.Core
 			return new string(_chars.Slice(0, _pos));
 		}
 
+		/// <summary>
+		/// Dispose buffer
+		/// </summary>
 		public void Dispose()
 		{
 			char[]? toReturn = _arrayToReturnToPool;
