[EpicGames.Core]
* Added xml doc to public methods in StackStringBuffer

diff --git a/Engine/Source/Programs/Shared/EpicGames.Core/StackStringBuffer.cs b/Engine/Source/Programs/Shared/EpicGames.Core/StackStringBuffer.cs
--- a/Engine/Source/Programs/Shared/EpicGames.Core/StackStringBuffer.cs
+++ b/Engine/Source/Programs/Shared/EpicGames.Core/StackStringBuffer.cs
@@ -32,6 +32,10 @@ namespace EpicGames.Core
 		/// </summary>
 		public int Length => _pos;
 
+		/// <summary>
+		/// Access a character by index
+		/// </summary>
+		/// <param name="index"></param>
 		public ref char this[int index] => ref _chars[index];
 
 		/// <summary>
