Provide command line argument to enable all codegen styles for testing.

#rnx

diff --git a/Engine/Source/Programs/Shared/EpicGames.UHT/Utils/UhtSession.cs b/Engine/Source/Programs/Shared/EpicGames.UHT/Utils/UhtSession.cs
--- a/Engine/Source/Programs/Shared/EpicGames.UHT/Utils/UhtSession.cs
+++ b/Engine/Source/Programs/Shared/EpicGames.UHT/Utils/UhtSession.cs
@@ -781,6 +781,11 @@ namespace EpicGames.UHT.Utils
 		/// </summary>
 		public bool CacheMessages { get; set; } = false;
 
+		/// <summary>
+		/// If true, force the target settings to generate all forms of code gen
+		/// </summary>
+		public bool EnableAllCodeGen { get; set; } = false;
+
 		/// <summary>
 		/// Collection of UHT tables
 		/// </summary>
@@ -2296,9 +2301,19 @@ namespace EpicGames.UHT.Utils
 
 				ManifestFile.Read();
 
-				if (Manifest != null && Tables != null)
+				if (Manifest != null)
 				{
-					Tables.AddPlugins(Manifest.UhtPlugins);
+					if (EnableAllCodeGen)
+					{
+						UHTTargetSettings targetSettings = Manifest.TargetSettings;
+						targetSettings.CompiledInObjectFormat = UhtCompiledInObjectFormat.All;
+						Manifest.TargetSettings = targetSettings;
+					}
+					if (Tables != null)
+					{
+						Tables.AddPlugins(Manifest.UhtPlugins);
+					}
+					Logger.LogInformation("UHT compiled-in object format {CompiledInObjectFormat}", Manifest.TargetSettings.CompiledInObjectFormat);
 				}
 			});
 		}
diff --git a/Engine/Source/Programs/UnrealBuildTool/Configuration/UEBuildTarget.cs b/Engine/Source/Programs/UnrealBuildTool/Configuration/UEBuildTarget.cs
--- a/Engine/Source/Programs/UnrealBuildTool/Configuration/UEBuildTarget.cs
+++ b/Engine/Source/Programs/UnrealBuildTool/Configuration/UEBuildTarget.cs
@@ -2749,7 +2749,6 @@ namespace UnrealBuildTool
 				// Generate max name/value pairs for UENUMs in the compiled-in data
 				Makefile.UHTTargetSettings.GenerateEnumMaxValues = true;
 			}
-			Logger.LogInformation($"UHT compiled-in object format {Makefile.UHTTargetSettings.CompiledInObjectFormat}");
 
 #if __VPROJECT_AVAILABLE__
 			{
diff --git a/Engine/Source/Programs/UnrealBuildTool/Modes/UnrealHeaderToolMode.cs b/Engine/Source/Programs/UnrealBuildTool/Modes/UnrealHeaderToolMode.cs
--- a/Engine/Source/Programs/UnrealBuildTool/Modes/UnrealHeaderToolMode.cs
+++ b/Engine/Source/Programs/UnrealBuildTool/Modes/UnrealHeaderToolMode.cs
@@ -549,6 +549,9 @@ namespace UnrealBuildTool.Modes
 		[CommandLine("-NoDefaultExporters", Description = "Disable all default exporters.  Useful for when a specific exporter is to be run")]
 		public bool bNoDefaultExporters = false;
 
+		[CommandLine("-EnableAllCodeGen", Description = "Include both the params and constinit code generation")]
+		public bool bEnableAllCodeGen = false;
+
 		/// <summary>
 		/// Initialize the options with the given command line arguments
 		/// </summary>
@@ -773,6 +776,7 @@ namespace UnrealBuildTool.Modes
 					NoOutput = Options.bNoOutput,
 					IncludeDebugOutput = Options.bIncludeDebugOutput,
 					NoDefaultExporters = Options.bNoDefaultExporters,
+					EnableAllCodeGen = Options.bEnableAllCodeGen,
 				};
 
 				if (Options.bWriteRef)
