EpicGames.Perforce: Add client name for stream view lookups

diff --git a/Engine/Source/Programs/Shared/EpicGames.Perforce/PerforceConnection.cs b/Engine/Source/Programs/Shared/EpicGames.Perforce/PerforceConnection.cs
--- a/Engine/Source/Programs/Shared/EpicGames.Perforce/PerforceConnection.cs
+++ b/Engine/Source/Programs/Shared/EpicGames.Perforce/PerforceConnection.cs
@@ -1316,6 +1316,20 @@ namespace EpicGames.Perforce
 		{
 			return (await TryGetStreamViewAsync(connection, streamName, changeNumber, cancellationToken)).Data;
 		}
+		
+		/// <summary>
+		/// Queries the view for a stream
+		/// </summary>
+		/// <param name="connection">Connection to the Perforce server</param>
+		/// <param name="clientName">Optional client name</param>
+		/// <param name="streamName">Name of the stream.</param>
+		/// <param name="changeNumber">Changelist at which to query the stream view</param>
+		/// <param name="cancellationToken">Token used to cancel the operation</param>
+		/// <returns>Response from the server; either a client record or error code</returns>
+		public static async Task<ClientRecord> GetStreamViewAsync(this IPerforceConnection connection, string? clientName, string streamName, int changeNumber, CancellationToken cancellationToken = default)
+		{
+			return (await TryGetStreamViewAsync(connection, clientName, streamName, changeNumber, cancellationToken)).Data;
+		}
 
 		/// <summary>
 		/// Queries the view for a stream
@@ -1326,6 +1340,20 @@ namespace EpicGames.Perforce
 		/// <param name="cancellationToken">Token used to cancel the operation</param>
 		/// <returns>Response from the server; either a client record or error code</returns>
 		public static Task<PerforceResponse<ClientRecord>> TryGetStreamViewAsync(this IPerforceConnection connection, string streamName, int changeNumber, CancellationToken cancellationToken = default)
+		{
+			return TryGetStreamViewAsync(connection, null, streamName, changeNumber, cancellationToken);
+		}
+		
+		/// <summary>
+		/// Queries the view for a stream
+		/// </summary>
+		/// <param name="connection">Connection to the Perforce server</param>
+		/// <param name="clientName">Optional client name</param>
+		/// <param name="streamName">Name of the stream.</param>
+		/// <param name="changeNumber">Changelist at which to query the stream view</param>
+		/// <param name="cancellationToken">Token used to cancel the operation</param>
+		/// <returns>Response from the server; either a client record or error code</returns>
+		public static Task<PerforceResponse<ClientRecord>> TryGetStreamViewAsync(this IPerforceConnection connection, string? clientName, string streamName, int changeNumber, CancellationToken cancellationToken = default)
 		{
 			List<string> arguments = new List<string> { "-o" };
 			arguments.Add($"-S{streamName}");
@@ -1333,7 +1361,12 @@ namespace EpicGames.Perforce
 			{
 				arguments.Add($"-c{changeNumber}");
 			}
-
+			
+			if (!String.IsNullOrEmpty(clientName))
+			{
+				arguments.Add(clientName);
+			}
+			
 			return SingleResponseCommandAsync<ClientRecord>(connection, "client", arguments, null, cancellationToken);
 		}
 
