[UBT]
* Minor drive-by opt

#FailedValidations Preflight

diff --git a/Engine/Source/Programs/Shared/EpicGames.Build/System/EnumeratePlugins.cs b/Engine/Source/Programs/Shared/EpicGames.Build/System/EnumeratePlugins.cs
--- a/Engine/Source/Programs/Shared/EpicGames.Build/System/EnumeratePlugins.cs
+++ b/Engine/Source/Programs/Shared/EpicGames.Build/System/EnumeratePlugins.cs
@@ -68,7 +68,7 @@ namespace UnrealBuildBase
 				}
 
 				// Sort the filenames to ensure that the plugin order is deterministic; otherwise response files will change with each build.
-				FileNames = [.. FileNames.OrderBy(x => x.FullName, StringComparer.OrdinalIgnoreCase)];
+				FileNames.Sort();
 				return FileNames;
 			})).Value;
 		}
