Horde: Remove stream salt from Gauntlet issues as they are generating too many individual issues

diff --git a/Engine/Source/Programs/Shared/EpicGames.Horde/Issues/Handlers/GauntletIssueHandler.cs b/Engine/Source/Programs/Shared/EpicGames.Horde/Issues/Handlers/GauntletIssueHandler.cs
--- a/Engine/Source/Programs/Shared/EpicGames.Horde/Issues/Handlers/GauntletIssueHandler.cs
+++ b/Engine/Source/Programs/Shared/EpicGames.Horde/Issues/Handlers/GauntletIssueHandler.cs
@@ -171,7 +171,7 @@ namespace EpicGames.Horde.Issues.Handlers
 				if (!hasCallstack)
 				{
 					// add job step salt if no Callstack property was found
-					key += $":{_context.StreamId}:{_context.NodeName}";
+					key += $":{_context.NodeName}";
 				}
 				keys.Add(key, IssueKeyType.None);
 				issueEvent.AuditLogger?.LogDebug("Fingerprint key: '{Key}' generated from event: '{Event}'", key, issueEvent.Render());
@@ -180,8 +180,8 @@ namespace EpicGames.Horde.Issues.Handlers
 			else
 			{
 				// Not enough information, make it an issue associated with only the job step
-				keys.Add($"{_context.StreamId}:{_context.NodeName}", IssueKeyType.None);
-				issueEvent.AuditLogger?.LogDebug("Fingerprint key: '{Key}' generated from event: '{Event}'", $"{_context.StreamId}:{_context.NodeName}", issueEvent.Render());
+				keys.Add($"{_context.NodeName}", IssueKeyType.None);
+				issueEvent.AuditLogger?.LogDebug("Fingerprint key: '{Key}' generated from event: '{Event}'", $"{_context.NodeName}", issueEvent.Render());
 			}
 			metadata.Add("Node", _context.NodeName);
 		}
