[EpicGames.Core]
* Fixed so we get information about what file is failing when json is parsed that has duplicates of keys

diff --git a/Engine/Source/Programs/Shared/EpicGames.Core/JsonObject.cs b/Engine/Source/Programs/Shared/EpicGames.Core/JsonObject.cs
--- a/Engine/Source/Programs/Shared/EpicGames.Core/JsonObject.cs
+++ b/Engine/Source/Programs/Shared/EpicGames.Core/JsonObject.cs
@@ -295,7 +295,14 @@ namespace EpicGames.Core
 		public static JsonObject Read(FileReference file)
 		{
 			JsonDocument doc = ParseUtf8File(file.FullName);
-			return new JsonObject(doc.RootElement);
+			try
+			{
+				return new JsonObject(doc.RootElement);
+			}
+			catch (Exception ex)
+			{
+				throw new JsonException($"Unable to parse {file}: {ex.Message}", file.FullName, null, null, ex);
+			}
 		}
 
 		/// <summary>
