[UBT]
* Fixed so UBT can still start even if ubahost.dll is corrupt

diff --git a/Engine/Source/Programs/Shared/EpicGames.UBA/ITrace.cs b/Engine/Source/Programs/Shared/EpicGames.UBA/ITrace.cs
--- a/Engine/Source/Programs/Shared/EpicGames.UBA/ITrace.cs
+++ b/Engine/Source/Programs/Shared/EpicGames.UBA/ITrace.cs
@@ -51,12 +51,16 @@ namespace EpicGames.UBA
 		{
 			if (name != null)
 			{
-				return new TraceImpl(name, traceFile, makeGlobal, channel);
-			}
-			else
-			{
-				return new TraceNull();
+				try
+				{
+					return new TraceImpl(name, traceFile, makeGlobal, channel);
+				}
+				catch
+				{
+				}
 			}
+
+			return new TraceNull();
 		}
 
 		/// <summary>
