Fixing minor BugHawk issue

#rnx

diff --git a/Engine/Source/Programs/Shared/EpicGames.UHT/Tokenizer/UhtTokenReaderStringExtensions.cs b/Engine/Source/Programs/Shared/EpicGames.UHT/Tokenizer/UhtTokenReaderStringExtensions.cs
--- a/Engine/Source/Programs/Shared/EpicGames.UHT/Tokenizer/UhtTokenReaderStringExtensions.cs
+++ b/Engine/Source/Programs/Shared/EpicGames.UHT/Tokenizer/UhtTokenReaderStringExtensions.cs
@@ -137,13 +137,13 @@ namespace EpicGames.UHT.Tokenizer
 				tokenReader.Require('(');
 				if (!tokenReader.PeekToken().IsSymbol(')'))
 				{
-					output = tokenReader.GetWrappedConstString();
+					output = tokenReader.GetWrappedConstString(exceptionContext);
 				}
 				tokenReader.Require(')');
 			}
 			else
 			{
-				output = tokenReader.GetWrappedConstString();
+				output = tokenReader.GetWrappedConstString(exceptionContext);
 			}
 			return output;
 		}
