[UBT]
* Added option to disable ipv6 dns resolving
* Fixed so horde and telemetry shares the same HttpMessageHandler to reduce dns resolving cost since they both talk to the same addresses (horde) by default

diff --git a/Engine/Source/Programs/Shared/EpicGames.Core/HttpClientExtensions.cs b/Engine/Source/Programs/Shared/EpicGames.Core/HttpClientExtensions.cs
--- a/Engine/Source/Programs/Shared/EpicGames.Core/HttpClientExtensions.cs
+++ b/Engine/Source/Programs/Shared/EpicGames.Core/HttpClientExtensions.cs
@@ -4,7 +4,9 @@ using System;
 using System.Collections.Generic;
 using System.Collections.Specialized;
 using System.Linq;
+using System.Net;
 using System.Net.Http;
+using System.Net.Sockets;
 using System.Text;
 using System.Text.Json;
 using System.Threading;
@@ -208,4 +210,92 @@ namespace EpicGames.Core
 			return new UriBuilder(uri) { Query = String.Join("&", queryParams) }.Uri;
 		}
 	}
+
+	public static class HttpMessageHandlers
+	{
+		public static bool UseIpv6DnsResolving { get; set; } = true;
+
+		/// <summary>
+		/// Creates default http message handler resolving both ipv4/6 dns.
+		/// </summary>
+		public static HttpMessageHandler CreateDefaultHandler()
+		{
+			if (UseIpv6DnsResolving)
+			{
+				return new NonOwningHandler(s_defaultHandler.Value);
+			}
+			else
+			{
+				return new NonOwningHandler(s_ipv4OnlyHandler.Value);
+			}
+		}
+
+		/// <summary>
+		/// Creates a http message handler that only resolves ipv4 dns and skips ipv6
+		/// This can be used to prevent long stalls on networks where ipv6 is not supported
+		/// </summary>
+		public static HttpMessageHandler CreateIpv4OnlyHandler()
+		{
+			return new NonOwningHandler(s_ipv4OnlyHandler.Value);
+		}
+
+		private static SocketsHttpHandler InternalCreateDefaultHandler()
+		{
+			return new SocketsHttpHandler
+			{
+				ConnectTimeout = TimeSpan.FromSeconds(5), // This needs to be higher even if there is Polly above because dns has to make it otherwise it will timeout over and over again
+				PooledConnectionLifetime = TimeSpan.FromMinutes(15),
+				PooledConnectionIdleTimeout = TimeSpan.FromMinutes(5),
+				//UseProxy = false,
+				//AllowAutoRedirect = false,
+			};
+		}
+
+		private static SocketsHttpHandler InternalCreateIpv4OnlyHandler()
+		{
+			SocketsHttpHandler socketHandler = InternalCreateDefaultHandler();
+			socketHandler.ConnectCallback = async (ctx, ct) =>
+			{
+				string host = ctx.DnsEndPoint.Host;
+				int port = ctx.DnsEndPoint.Port;
+
+				// v4-only resolution (NET 8 overload)
+				IPAddress[] addrs = await Dns.GetHostAddressesAsync(host, AddressFamily.InterNetwork, ct);
+				if (addrs.Length == 0)
+				{
+					throw new SocketException((int)SocketError.HostNotFound);
+				}
+
+				Exception? last = null;
+				foreach (IPAddress ip in addrs)
+				{
+					Socket sock = new(AddressFamily.InterNetwork, SocketType.Stream, ProtocolType.Tcp) { NoDelay = true };
+					using (ct.Register(() => { try { sock.Dispose(); } catch { } })) { }
+					try
+					{
+						await sock.ConnectAsync(new IPEndPoint(ip, port), ct);
+						return new NetworkStream(sock, ownsSocket: true);
+					}
+					catch (Exception ex)
+					{
+						last = ex;
+					}
+				}
+				throw last ?? new SocketException((int)SocketError.HostUnreachable);
+			};
+			return socketHandler;
+		}
+
+		private sealed class NonOwningHandler : DelegatingHandler
+		{
+			public NonOwningHandler(HttpMessageHandler inner) => InnerHandler = inner;
+
+			#pragma warning disable CA2215 // Intentionally skip calling base.Dispose to prevent cascading dispose
+			protected override void Dispose(bool disposing) { }
+			#pragma warning restore CA2215 // Dispose objects before losing scope
+		}
+
+		private static readonly Lazy<SocketsHttpHandler> s_defaultHandler = new(InternalCreateDefaultHandler);
+		private static readonly Lazy<SocketsHttpHandler> s_ipv4OnlyHandler = new(InternalCreateIpv4OnlyHandler);
+	}
 }
diff --git a/Engine/Source/Programs/Shared/EpicGames.Core/Telemetry/DefaultTelemetryHttpClient.cs b/Engine/Source/Programs/Shared/EpicGames.Core/Telemetry/DefaultTelemetryHttpClient.cs
--- a/Engine/Source/Programs/Shared/EpicGames.Core/Telemetry/DefaultTelemetryHttpClient.cs
+++ b/Engine/Source/Programs/Shared/EpicGames.Core/Telemetry/DefaultTelemetryHttpClient.cs
@@ -36,10 +36,8 @@ namespace EpicGames.Core.Telemetry
 		public static HttpClient GetClient()
 		{
 			// per MS documentation, with a lifetime specified we shouldn't have socket exhaustion issues
-			SocketsHttpHandler socketHandler = new SocketsHttpHandler {
-				ConnectTimeout  = TimeSpan.FromSeconds(2),
-				PooledConnectionLifetime = TimeSpan.FromMinutes(15)
-			};
+			HttpMessageHandler socketHandler = HttpMessageHandlers.CreateDefaultHandler();
+
 #pragma warning disable CA2000 // Dispose objects before losing scope
 			PolicyHttpMessageHandler pollyHandler = new PolicyHttpMessageHandler(GetRetryPolicy())
 			{
diff --git a/Engine/Source/Programs/Shared/EpicGames.Horde/HordeHttpMessageHandler.cs b/Engine/Source/Programs/Shared/EpicGames.Horde/HordeHttpMessageHandler.cs
--- a/Engine/Source/Programs/Shared/EpicGames.Horde/HordeHttpMessageHandler.cs
+++ b/Engine/Source/Programs/Shared/EpicGames.Horde/HordeHttpMessageHandler.cs
@@ -30,7 +30,7 @@ namespace EpicGames.Horde
 		/// <param name="logger"></param>
 		public HordeHttpMessageHandler(ILogger<HordeHttpMessageHandler> logger)
 		{
-			_instance = new SocketsHttpHandler();
+			_instance = HttpMessageHandlers.CreateDefaultHandler();
 			_instance = new PolicyHttpMessageHandler(request => CreateDefaultTimeoutRetryPolicy(request, logger)) { InnerHandler = _instance };
 			_instance = new PolicyHttpMessageHandler(request => CreateHordeTransientErrorPolicy(request, logger)) { InnerHandler = _instance };
 		}
diff --git a/Engine/Source/Programs/UnrealBuildTool/UnrealBuildTool.cs b/Engine/Source/Programs/UnrealBuildTool/UnrealBuildTool.cs
--- a/Engine/Source/Programs/UnrealBuildTool/UnrealBuildTool.cs
+++ b/Engine/Source/Programs/UnrealBuildTool/UnrealBuildTool.cs
@@ -270,6 +270,12 @@ namespace UnrealBuildTool
 			[CommandLine(Prefix = "-Session", Description = "Session identifier for this run of UBT, if unset defaults to a random Guid")]
 			public string? TelemetrySession = null;
 
+			/// <summary>
+			/// Whether telemetry/horde should use ipv6 or not. If network does not support ipv6 then using ipv6 can cause multi-second stalls
+			/// </summary>
+			[XmlConfigFile(Category = "BuildConfiguration")]
+			public bool bUseIpv6 { get; set; } = true;
+
 			/// <summary>
 			/// Initialize the options with the given command line arguments
 			/// </summary>
@@ -553,6 +559,11 @@ namespace UnrealBuildTool
 					XmlConfig.ApplyTo(Options);
 				}
 
+				if (!Options.bUseIpv6)
+				{
+					HttpMessageHandlers.UseIpv6DnsResolving = false;
+				}
+
 				// Start prefetching the contents of the engine folder
 				if ((ModeOptions & ToolModeOptions.StartPrefetchingEngine) != 0)
 				{
