[EpicGames.Core]
* Fixed static analyzer warnings regarding formatting and missing documentation

diff --git a/Engine/Source/Programs/Shared/EpicGames.Core/HttpClientExtensions.cs b/Engine/Source/Programs/Shared/EpicGames.Core/HttpClientExtensions.cs
--- a/Engine/Source/Programs/Shared/EpicGames.Core/HttpClientExtensions.cs
+++ b/Engine/Source/Programs/Shared/EpicGames.Core/HttpClientExtensions.cs
@@ -211,8 +211,14 @@ namespace EpicGames.Core
 		}
 	}
 
+	/// <summary>
+	/// Class that can/should be used to Create http message handlers. Under the hood, same handler will be used to reduce redundant dns lookup calls
+	/// </summary>
 	public static class HttpMessageHandlers
 	{
+		/// <summary>
+		/// Set to false to disable ipv6 dns resolving
+		/// </summary>
 		public static bool UseIpv6DnsResolving { get; set; } = true;
 
 		/// <summary>
@@ -270,7 +276,19 @@ namespace EpicGames.Core
 				foreach (IPAddress ip in addrs)
 				{
 					Socket sock = new(AddressFamily.InterNetwork, SocketType.Stream, ProtocolType.Tcp) { NoDelay = true };
-					using (ct.Register(() => { try { sock.Dispose(); } catch { } })) { }
+					using (ct.Register(() =>
+					{
+						try
+						{
+							sock.Dispose();
+						}
+						catch
+						{
+						}
+					}))
+					{
+					}
+
 					try
 					{
 						await sock.ConnectAsync(new IPEndPoint(ip, port), ct);
