Horde: Do not generate a flood of Build Health issues for a small set of known Gauntlet log event warnings
If we identify cases where we want these warnings, they should be defined as new log events, and force their warning in the issue handler
Though, we have been unable to come up with a single case in exploration

diff --git a/Engine/Source/Programs/Shared/EpicGames.Horde/Issues/Handlers/GauntletIssueHandler.cs b/Engine/Source/Programs/Shared/EpicGames.Horde/Issues/Handlers/GauntletIssueHandler.cs
--- a/Engine/Source/Programs/Shared/EpicGames.Horde/Issues/Handlers/GauntletIssueHandler.cs
+++ b/Engine/Source/Programs/Shared/EpicGames.Horde/Issues/Handlers/GauntletIssueHandler.cs
@@ -285,6 +285,14 @@ namespace EpicGames.Horde.Issues.Handlers
 					return true;
 				}
 
+				if (issueEvent.Severity == LogLevel.Warning)
+				{
+					// Do not generate a flood of Build Health issues for smallset of  known Gauntlet log event warnings
+					// If we identify cases where we want these warnings, they should be defined as new log events, and force their warning in the issue handler, here
+					// Though, we have been unable to come up with a single case in exploration					
+					return true;
+				}
+
 				bool isSevereEvent = s_knownSevereGauntletEvents.Contains(issueEvent.EventId.Value);
 				bool isTestEvent = !isSevereEvent && s_knownTestGauntletEvents.Contains(issueEvent.EventId.Value);
 				bool isErrorEvent = issueEvent.Severity >= LogLevel.Error;
