[EpicGames.Build]
* Changed Parallel2.x to use ManualResetEventSlim instead of EventWaitHandle to improve performance

diff --git a/Engine/Source/Programs/Shared/EpicGames.Build/System/Parallel2.cs b/Engine/Source/Programs/Shared/EpicGames.Build/System/Parallel2.cs
--- a/Engine/Source/Programs/Shared/EpicGames.Build/System/Parallel2.cs
+++ b/Engine/Source/Programs/Shared/EpicGames.Build/System/Parallel2.cs
@@ -31,7 +31,7 @@ namespace UnrealBuildBase
 		/// <param name="useTasks">Use tasks or threads</param>
 		public static void ForEach<TSource>(IEnumerable<TSource> source, Action<TSource> body, int helperCount = -1, bool useTasks = true)
 		{
-			using EventWaitHandle Handle = new(false, EventResetMode.ManualReset);
+			using ManualResetEventSlim Handle = new(initialState: false, spinCount: 20);
 
 			ForEachContext<TSource> context = new() { Enumerator = source.GetEnumerator(), Handle = Handle };
 
@@ -113,7 +113,7 @@ namespace UnrealBuildBase
 
 			WorkerLoop();
 
-			context.Handle.WaitOne();
+			context.Handle.Wait();
 
 			if (context.Exception != null)
 			{
@@ -141,7 +141,7 @@ namespace UnrealBuildBase
 				return;
 			}
 
-			using EventWaitHandle Handle = new(false, EventResetMode.ManualReset);
+			using ManualResetEventSlim Handle = new(initialState: false, spinCount: 20);
 
 			ForContext context = new() { Handle = Handle, Index = from, End = to, Finished = from };
 
@@ -200,7 +200,7 @@ namespace UnrealBuildBase
 
 			WorkerLoop();
 
-			context.Handle.WaitOne();
+			context.Handle.Wait();
 
 			if (context.Exception != null)
 			{
@@ -211,7 +211,7 @@ namespace UnrealBuildBase
 		private class ForEachContext<TSource>
 		{
 			internal required IEnumerator<TSource> Enumerator;
-			internal required EventWaitHandle Handle;
+			internal required ManualResetEventSlim Handle;
 			internal int Count;
 			internal int Finished;
 			internal bool Done;
@@ -220,7 +220,7 @@ namespace UnrealBuildBase
 
 		private class ForContext
 		{
-			internal required EventWaitHandle Handle;
+			internal required ManualResetEventSlim Handle;
 			internal int Index;
 			internal int End;
 			internal int Finished;
