[EpicGames.Build]
* Changed DirectoryEnumerator to use smaller file information class (FileFullDirectoryInformation instead of FileIdBothDirectoryInformation) to improve performance

diff --git a/Engine/Source/Programs/Shared/EpicGames.Build/System/DirectoryEnumerator.cs b/Engine/Source/Programs/Shared/EpicGames.Build/System/DirectoryEnumerator.cs
--- a/Engine/Source/Programs/Shared/EpicGames.Build/System/DirectoryEnumerator.cs
+++ b/Engine/Source/Programs/Shared/EpicGames.Build/System/DirectoryEnumerator.cs
@@ -69,7 +69,7 @@ namespace UnrealBuildBase
 
 					long alloc = 0;
 
-					int status = NtCreateFile(out dirHandle, FILE_LIST_DIRECTORY | SYNCHRONIZE, ref oa, out IO_STATUS_BLOCK iosb, ref alloc, 0u, FILE_SHARE_READ | FILE_SHARE_WRITE | FILE_SHARE_DELETE, FILE_OPEN, FILE_DIRECTORY_FILE | FILE_SYNCHRONOUS_IO_NONALERT | FILE_OPEN_FOR_BACKUP_INTENT, IntPtr.Zero, 0u);
+					int status = NtCreateFile(out dirHandle, FILE_LIST_DIRECTORY | SYNCHRONIZE, ref oa, out IO_STATUS_BLOCK iosb, ref alloc, 0u, FILE_SHARE_READ | FILE_SHARE_WRITE | FILE_SHARE_DELETE, FILE_OPEN, FILE_DIRECTORY_FILE | FILE_SYNCHRONOUS_IO_NONALERT, IntPtr.Zero, 0u);
 					if (status != STATUS_SUCCESS)
 					{
 						return null; //ThrowIo(status, $"NtCreateFile('{directoryPath}')");
@@ -191,7 +191,7 @@ namespace UnrealBuildBase
 
 			while (true)
 			{
-				int status = NtQueryDirectoryFile(dirHandle, IntPtr.Zero, IntPtr.Zero, IntPtr.Zero, out IO_STATUS_BLOCK iosbLocal, buf, (uint)bufferSize, FILE_INFORMATION_CLASS.FileIdBothDirectoryInformation, false, null, restart);
+				int status = NtQueryDirectoryFile(dirHandle, IntPtr.Zero, IntPtr.Zero, IntPtr.Zero, out IO_STATUS_BLOCK iosbLocal, buf, (uint)bufferSize, FILE_INFORMATION_CLASS.FileFullDirectoryInformation, false, null, restart);
 
 				if (status == STATUS_NO_MORE_FILES)
 				{
@@ -206,7 +206,7 @@ namespace UnrealBuildBase
 				byte* cursor = buf;
 				while (true)
 				{
-					FILE_ID_BOTH_DIR_INFORMATION* info = (FILE_ID_BOTH_DIR_INFORMATION*)cursor;
+					FILE_FULL_DIR_INFORMATION* info = (FILE_FULL_DIR_INFORMATION*)cursor;
 
 					int nameBytes = checked((int)info->FileNameLength);
 					char* namePtr = &info->FileName;
@@ -249,7 +249,6 @@ namespace UnrealBuildBase
 		private const uint FILE_OPEN = 0x00000001;
 		private const uint FILE_DIRECTORY_FILE = 0x00000001;
 		private const uint FILE_SYNCHRONOUS_IO_NONALERT = 0x00000020;
-		private const uint FILE_OPEN_FOR_BACKUP_INTENT = 0x00004000;
 		private const uint SYNCHRONIZE = 0x00100000;
 
 		[DllImport("ntdll.dll")]
@@ -340,11 +339,11 @@ namespace UnrealBuildBase
 
 		private enum FILE_INFORMATION_CLASS : int
 		{
-			FileIdBothDirectoryInformation = 37
+			FileFullDirectoryInformation = 2,
 		}
 
-		[StructLayout(LayoutKind.Sequential)]
-		private unsafe struct FILE_ID_BOTH_DIR_INFORMATION
+		[StructLayout(LayoutKind.Sequential, CharSet = CharSet.Unicode)]
+		public struct FILE_FULL_DIR_INFORMATION
 		{
 			public uint NextEntryOffset;
 			public uint FileIndex;
@@ -357,9 +356,6 @@ namespace UnrealBuildBase
 			public uint FileAttributes;
 			public uint FileNameLength;
 			public uint EaSize;
-			public byte ShortNameLength;
-			private fixed char ShortName[12];
-			public long FileId;
 			public char FileName;
 		}
 	}
