EpicGames.Perforce: Add OpenTelemetry tracing

diff --git a/Engine/Source/Programs/AutomationTool/AutomationUtils/AutomationUtils.Automation.csproj b/Engine/Source/Programs/AutomationTool/AutomationUtils/AutomationUtils.Automation.csproj
--- a/Engine/Source/Programs/AutomationTool/AutomationUtils/AutomationUtils.Automation.csproj
+++ b/Engine/Source/Programs/AutomationTool/AutomationUtils/AutomationUtils.Automation.csproj
@@ -14,6 +14,7 @@
     <AppendTargetFrameworkToOutputPath>false</AppendTargetFrameworkToOutputPath>
     <DebugType>pdbonly</DebugType>
     <DebugType Condition="'$(IsWindows)' == 'true' And '$(IsArm64)' == 'true'">portable</DebugType>
+    <CopyLocalLockFileAssemblies>true</CopyLocalLockFileAssemblies>
   </PropertyGroup>
   <PropertyGroup Condition="'$(EpicGamesMsBuild)' == 'true'">
     <OutputPath>..\..\..\..\Binaries\DotNET\AutomationTool\AutomationUtils</OutputPath>
diff --git a/Engine/Source/Programs/Horde/HordeServer.Shared/Utilities/OpenTelemetryTracers.cs b/Engine/Source/Programs/Horde/HordeServer.Shared/Utilities/OpenTelemetryTracers.cs
--- a/Engine/Source/Programs/Horde/HordeServer.Shared/Utilities/OpenTelemetryTracers.cs
+++ b/Engine/Source/Programs/Horde/HordeServer.Shared/Utilities/OpenTelemetryTracers.cs
@@ -24,12 +24,18 @@ namespace HordeServer.Utilities
 		/// Name of MongoDB tracer (aka activity source)
 		/// </summary>
 		public const string MongoDbName = "MongoDB";
+		
+		/// <summary>
+		/// Name of Perforce client tracer (aka activity source)
+		/// By default, tracer is declared inside EpicGames.Perforce.
+		/// </summary>
+		public const string PerforceName = "Perforce";
 
 		/// <summary>
 		/// List of all source names configured in this class.
 		/// They are needed at startup when initializing OpenTelemetry
 		/// </summary>
-		public static string[] SourceNames => new[] { HordeName, MongoDbName };
+		public static string[] SourceNames => [HordeName, MongoDbName, PerforceName];
 
 		/// <summary>
 		/// Default tracer used in Horde
diff --git a/Engine/Source/Programs/Horde/Plugins/Build/HordeServer.Build/VersionControl/Perforce/PerforceService.cs b/Engine/Source/Programs/Horde/Plugins/Build/HordeServer.Build/VersionControl/Perforce/PerforceService.cs
--- a/Engine/Source/Programs/Horde/Plugins/Build/HordeServer.Build/VersionControl/Perforce/PerforceService.cs
+++ b/Engine/Source/Programs/Horde/Plugins/Build/HordeServer.Build/VersionControl/Perforce/PerforceService.cs
@@ -91,6 +91,9 @@ namespace HordeServer.VersionControl.Perforce
 
 			/// <inheritdoc/>
 			public ILogger Logger => _inner.Perforce.Logger;
+			
+			/// <inheritdoc/>
+			public Tracer Tracer => _inner.Perforce.Tracer;
 
 			/// <inheritdoc/>
 			public IPerforceOutput Command(string command, IReadOnlyList<string> arguments, IReadOnlyList<string>? fileArguments, byte[]? inputData, string? promptResponse, bool interceptIo)
diff --git a/Engine/Source/Programs/Shared/EpicGames.Perforce/EpicGames.Perforce.csproj b/Engine/Source/Programs/Shared/EpicGames.Perforce/EpicGames.Perforce.csproj
--- a/Engine/Source/Programs/Shared/EpicGames.Perforce/EpicGames.Perforce.csproj
+++ b/Engine/Source/Programs/Shared/EpicGames.Perforce/EpicGames.Perforce.csproj
@@ -75,6 +75,7 @@
 
   <ItemGroup>
     <PackageReference Include="System.Linq.Async" Version="6.0.3" />
+    <PackageReference Include="OpenTelemetry.Api" Version="1.8.1" />
   </ItemGroup>
 
 </Project>
diff --git a/Engine/Source/Programs/Shared/EpicGames.Perforce/IPerforceConnection.cs b/Engine/Source/Programs/Shared/EpicGames.Perforce/IPerforceConnection.cs
--- a/Engine/Source/Programs/Shared/EpicGames.Perforce/IPerforceConnection.cs
+++ b/Engine/Source/Programs/Shared/EpicGames.Perforce/IPerforceConnection.cs
@@ -3,6 +3,7 @@
 using System;
 using System.Collections.Generic;
 using Microsoft.Extensions.Logging;
+using OpenTelemetry.Trace;
 
 namespace EpicGames.Perforce
 {
@@ -20,6 +21,11 @@ namespace EpicGames.Perforce
 		/// Logger for this connection
 		/// </summary>
 		ILogger Logger { get; }
+		
+		/// <summary>
+		/// Tracer for this connection
+		/// </summary>
+		Tracer Tracer { get; }
 
 		/// <summary>
 		/// Queues a Perforce command to be executed
diff --git a/Engine/Source/Programs/Shared/EpicGames.Perforce/NativePerforceConnection.cs b/Engine/Source/Programs/Shared/EpicGames.Perforce/NativePerforceConnection.cs
--- a/Engine/Source/Programs/Shared/EpicGames.Perforce/NativePerforceConnection.cs
+++ b/Engine/Source/Programs/Shared/EpicGames.Perforce/NativePerforceConnection.cs
@@ -14,6 +14,7 @@ using System.Threading.Tasks;
 using EpicGames.Core;
 using Microsoft.Extensions.Logging;
 using Microsoft.Extensions.Logging.Abstractions;
+using OpenTelemetry.Trace;
 
 namespace EpicGames.Perforce
 {
@@ -303,15 +304,19 @@ namespace EpicGames.Perforce
 		public IPerforceSettings Settings { get; }
 
 		/// <inheritdoc/>
-		public ILogger Logger { get; }
+		public ILogger Logger { get; set; }
+		
+		/// <inheritdoc/>
+		public Tracer Tracer { get; set; }
 
 		/// <summary>
 		/// Constructor
 		/// </summary>
 		/// <param name="settings">Settings for the connection</param>
 		/// <param name="logger">Logger for messages</param>
-		public NativePerforceConnection(IPerforceSettings settings, ILogger logger)
-			: this(settings, 2, 128 * 1024, logger)
+		/// <param name="tracer">Tracer</param>
+		public NativePerforceConnection(IPerforceSettings settings, ILogger logger, Tracer? tracer = null)
+			: this(settings, 2, 128 * 1024, logger, tracer)
 		{
 		}
 
@@ -322,10 +327,12 @@ namespace EpicGames.Perforce
 		/// <param name="bufferCount">Number of buffers to create for streaming response data</param>
 		/// <param name="bufferSize">Size of each buffer</param>
 		/// <param name="logger">Logger for messages</param>
-		public NativePerforceConnection(IPerforceSettings settings, int bufferCount, int bufferSize, ILogger logger)
+		/// <param name="tracer">Tracer</param> 
+		public NativePerforceConnection(IPerforceSettings settings, int bufferCount, int bufferSize, ILogger logger, Tracer? tracer = null)
 		{
 			Settings = settings;
 			Logger = logger;
+			Tracer = tracer ?? PerforceConnection.DefaultTracer;
 
 			_uniqueId = Interlocked.Increment(ref s_nextUniqueId);
 
@@ -354,15 +361,12 @@ namespace EpicGames.Perforce
 		/// <summary>
 		/// Create an instance of the native client
 		/// </summary>
-		/// <param name="settings"></param>
-		/// <param name="logger"></param>
-		/// <returns></returns>
-		public static async Task<NativePerforceConnection> CreateAsync(IPerforceSettings settings, ILogger logger)
+		public static async Task<NativePerforceConnection> CreateAsync(IPerforceSettings settings, ILogger logger, Tracer? tracer = null)
 		{
 			NativePerforceConnection? connection = null;
 			try
 			{
-				connection = new NativePerforceConnection(settings, logger);
+				connection = new NativePerforceConnection(settings, logger, tracer);
 				await connection.ConnectAsync();
 				return connection;
 			}
diff --git a/Engine/Source/Programs/Shared/EpicGames.Perforce/PerforceConnection.cs b/Engine/Source/Programs/Shared/EpicGames.Perforce/PerforceConnection.cs
--- a/Engine/Source/Programs/Shared/EpicGames.Perforce/PerforceConnection.cs
+++ b/Engine/Source/Programs/Shared/EpicGames.Perforce/PerforceConnection.cs
@@ -11,6 +11,7 @@ using System.Text.RegularExpressions;
 using System.Threading;
 using System.Threading.Tasks;
 using Microsoft.Extensions.Logging;
+using OpenTelemetry.Trace;
 
 namespace EpicGames.Perforce
 {
@@ -19,14 +20,19 @@ namespace EpicGames.Perforce
 	/// </summary>
 	public sealed class PerforceConnection : IPerforceConnection
 	{
+		/// <summary>
+		/// Default tracer. A custom tracer can also be passed to constructor or factory methods.
+		/// </summary>
+		public static Tracer DefaultTracer { get; } = TracerProvider.Default.GetTracer("Perforce");
+		
 		/// <summary>
 		/// Create a new Perforce connection
 		/// </summary>
-		/// <param name="logger"></param>
-		/// <returns></returns>
-		public static Task<IPerforceConnection> CreateAsync(ILogger logger)
+		/// <param name="logger">Interface for logging</param>
+		/// <param name="tracer">Tracer</param>
+		public static Task<IPerforceConnection> CreateAsync(ILogger logger, Tracer? tracer = null)
 		{
-			return CreateAsync(PerforceSettings.Default, logger);
+			return CreateAsync(PerforceSettings.Default, logger, tracer);
 		}
 
 		/// <summary>
@@ -35,10 +41,10 @@ namespace EpicGames.Perforce
 		/// <param name="serverAndPort">The server address and port</param>
 		/// <param name="userName">The user name</param>
 		/// <param name="logger">Interface for logging</param>
-		/// <returns></returns>
-		public static Task<IPerforceConnection> CreateAsync(string? serverAndPort, string? userName, ILogger logger)
+		/// /// <param name="tracer">Tracer</param>
+		public static Task<IPerforceConnection> CreateAsync(string? serverAndPort, string? userName, ILogger logger, Tracer? tracer = null)
 		{
-			return CreateAsync(CombineSettings(serverAndPort, userName), logger);
+			return CreateAsync(CombineSettings(serverAndPort, userName), logger, tracer);
 		}
 
 		/// <summary>
@@ -48,10 +54,10 @@ namespace EpicGames.Perforce
 		/// <param name="userName">The user name</param>
 		/// <param name="clientName">The client name</param>
 		/// <param name="logger">Interface for logging</param>
-		/// <returns></returns>
-		public static Task<IPerforceConnection> CreateAsync(string? serverAndPort, string? userName, string? clientName, ILogger logger)
+		/// <param name="tracer">Tracer</param>
+		public static Task<IPerforceConnection> CreateAsync(string? serverAndPort, string? userName, string? clientName, ILogger logger, Tracer? tracer = null)
 		{
-			return CreateAsync(CombineSettings(serverAndPort, userName, clientName), logger);
+			return CreateAsync(CombineSettings(serverAndPort, userName, clientName), logger, tracer);
 		}
 
 		/// <summary>
@@ -63,27 +69,27 @@ namespace EpicGames.Perforce
 		/// <param name="appName"></param>
 		/// <param name="appVersion"></param>
 		/// <param name="logger">Interface for logging</param>
-		/// <returns></returns>
-		public static Task<IPerforceConnection> CreateAsync(string? serverAndPort, string? userName, string? clientName, string? appName, string? appVersion, ILogger logger)
+		/// <param name="tracer">Tracer</param>
+		public static Task<IPerforceConnection> CreateAsync(string? serverAndPort, string? userName, string? clientName, string? appName, string? appVersion, ILogger logger, Tracer? tracer = null)
 		{
-			return CreateAsync(CombineSettings(serverAndPort, userName, clientName, appName, appVersion), logger);
+			return CreateAsync(CombineSettings(serverAndPort, userName, clientName, appName, appVersion), logger, tracer);
 		}
 
 		/// <summary>
 		/// Create a new Perforce connection
 		/// </summary>
 		/// <param name="settings">Settings for the connection</param>
-		/// <param name="logger"></param>
-		/// <returns></returns>
-		public static async Task<IPerforceConnection> CreateAsync(IPerforceSettings settings, ILogger logger)
+		/// <param name="logger">Interface for logging</param>
+		/// <param name="tracer">Tracer</param>
+		public static async Task<IPerforceConnection> CreateAsync(IPerforceSettings settings, ILogger logger, Tracer? tracer = null)
 		{
 			if (settings.PreferNativeClient && NativePerforceConnection.IsSupported())
 			{
-				return await NativePerforceConnection.CreateAsync(settings, logger);
+				return await NativePerforceConnection.CreateAsync(settings, logger, tracer);
 			}
 			else
 			{
-				return new PerforceConnection(settings, logger);
+				return new PerforceConnection(settings, logger, tracer);
 			}
 		}
 
@@ -203,24 +209,22 @@ namespace EpicGames.Perforce
 		/// Additional options to append to the command line
 		/// </summary>
 		public List<string> GlobalOptions { get; } = new List<string>();
-
-		/// <summary>
-		/// The logging interface
-		/// </summary>
-		public ILogger Logger
-		{
-			get;
-			set;
-		}
+		
+		/// <inheritdoc/>
+		public ILogger Logger { get; set; }
+		
+		/// <inheritdoc/>
+		public Tracer Tracer { get; set; }
 
 		/// <summary>
 		/// Constructor
 		/// </summary>
 		/// <param name="serverAndPort">The server address and port</param>
 		/// <param name="userName">The user name</param>
+		/// <param name="tracer">Tracer</param>
 		/// <param name="logger">Interface for logging</param>
-		public PerforceConnection(string? serverAndPort, string? userName, ILogger logger)
-			: this(serverAndPort, userName, null, logger)
+		public PerforceConnection(string? serverAndPort, string? userName, ILogger logger, Tracer? tracer = null)
+			: this(serverAndPort, userName, null, logger, tracer)
 		{
 		}
 
@@ -230,9 +234,10 @@ namespace EpicGames.Perforce
 		/// <param name="serverAndPort">The server address and port</param>
 		/// <param name="userName">The user name</param>
 		/// <param name="clientName">The client name</param>
+		/// <param name="tracer">Tracer</param>
 		/// <param name="logger">Interface for logging</param>
-		public PerforceConnection(string? serverAndPort, string? userName, string? clientName, ILogger logger)
-			: this(serverAndPort, userName, clientName, null, null, logger)
+		public PerforceConnection(string? serverAndPort, string? userName, string? clientName, ILogger logger, Tracer? tracer = null)
+			: this(serverAndPort, userName, clientName, null, null, logger, tracer)
 		{
 			AssemblyName entryAssemblyName = Assembly.GetEntryAssembly()!.GetName();
 			if (entryAssemblyName.Name != null)
@@ -250,8 +255,9 @@ namespace EpicGames.Perforce
 		/// <param name="clientName">The client name</param>
 		/// <param name="appName"></param>
 		/// <param name="appVersion"></param>
+		/// <param name="tracer">Tracer</param>
 		/// <param name="logger">Interface for logging</param>
-		public PerforceConnection(string? serverAndPort, string? userName, string? clientName, string? appName, string? appVersion, ILogger logger)
+		public PerforceConnection(string? serverAndPort, string? userName, string? clientName, string? appName, string? appVersion, ILogger logger, Tracer? tracer = null)
 		{
 			ServerAndPort = serverAndPort;
 			UserName = userName;
@@ -259,15 +265,14 @@ namespace EpicGames.Perforce
 			AppName = appName;
 			AppVersion = appVersion;
 			Logger = logger;
+			Tracer = tracer ?? DefaultTracer;
 		}
 
 		/// <summary>
 		/// Constructor
 		/// </summary>
-		/// <param name="settings"></param>
-		/// <param name="logger"></param>
-		public PerforceConnection(IPerforceSettings settings, ILogger logger)
-			: this(settings.ServerAndPort, settings.UserName, settings.ClientName, settings.AppName, settings.AppVersion, logger)
+		public PerforceConnection(IPerforceSettings settings, ILogger logger, Tracer? tracer = null)
+			: this(settings.ServerAndPort, settings.UserName, settings.ClientName, settings.AppName, settings.AppVersion, logger, tracer)
 		{
 			HostName = settings.HostName;
 		}
@@ -277,7 +282,7 @@ namespace EpicGames.Perforce
 		/// </summary>
 		/// <param name="other">Connection to copy settings from</param>
 		public PerforceConnection(PerforceConnection other)
-			: this(other.ServerAndPort, other.UserName, other.ClientName, other.AppName, other.AppVersion, other.Logger)
+			: this(other.ServerAndPort, other.UserName, other.ClientName, other.AppName, other.AppVersion, other.Logger, other.Tracer)
 		{
 			GlobalOptions.AddRange(other.GlobalOptions);
 		}
@@ -386,7 +391,7 @@ namespace EpicGames.Perforce
 		public static Task<IPerforceConnection> WithClientAsync(this IPerforceConnection perforce, string? clientName)
 		{
 			PerforceSettings settings = new PerforceSettings(perforce.Settings) { ClientName = clientName };
-			return PerforceConnection.CreateAsync(settings, perforce.Logger);
+			return PerforceConnection.CreateAsync(settings, perforce.Logger, perforce.Tracer);
 		}
 
 		/// <summary>
@@ -400,6 +405,24 @@ namespace EpicGames.Perforce
 		}
 
 		#region Command wrappers
+		
+		private static void AddAttributes(IPerforceConnection perforce, TelemetrySpan span, string commandType, string command, IReadOnlyList<string> arguments, IReadOnlyList<string>? fileArguments)
+		{
+			span.SetAttribute("epic.perforce.command_type", commandType)
+				.SetAttribute("epic.perforce.args", arguments.ToArray())
+				.SetAttribute("epic.perforce.file_args", fileArguments?.ToArray())
+				.SetAttribute("epic.perforce.native", perforce.Settings.PreferNativeClient)
+				.SetAttribute("epic.perforce.client", perforce.Settings.ClientName)
+				.SetAttribute("epic.perforce.username", perforce.Settings.UserName)
+				.SetAttribute("epic.perforce.server_port", perforce.Settings.ServerAndPort)
+				.SetAttribute("epic.perforce.hostname", perforce.Settings.HostName);
+			
+			// Useful attributes but mostly hinting for Datadog
+			span.SetAttribute("type", "db")
+				.SetAttribute("db.query", command + " " + String.Join(' ', arguments))
+				.SetAttribute("operation.name", command)
+				.SetAttribute("service.name", "Perforce");
+		}
 
 		/// <summary>
 		/// Execute a command and parse the response
@@ -429,10 +452,14 @@ namespace EpicGames.Perforce
 		/// <returns>List of objects returned by the server</returns>
 		public static async Task<List<PerforceResponse>> CommandAsync(this IPerforceConnection perforce, string command, IReadOnlyList<string> arguments, IReadOnlyList<string>? fileArguments, byte[]? inputData, Type? statRecordType, CancellationToken cancellationToken = default)
 		{
-			await using (IPerforceOutput response = perforce.Command(command, arguments, fileArguments, inputData, null, false))
-			{
-				return await response.ReadResponsesAsync(statRecordType, cancellationToken);
-			}
+			using TelemetrySpan span = perforce.Tracer.StartActiveSpan(command);
+			AddAttributes(perforce, span, "normal", command, arguments, fileArguments);
+			
+			await using IPerforceOutput response = perforce.Command(command, arguments, fileArguments, inputData, null, false);
+			List<PerforceResponse> responses = await response.ReadResponsesAsync(statRecordType, cancellationToken);
+			
+			span.SetAttribute("epic.perforce.res.count", responses.Count);
+			return responses;
 		}
 
 		/// <summary>
@@ -449,15 +476,20 @@ namespace EpicGames.Perforce
 		/// <returns>List of objects returned by the server</returns>
 		public static async IAsyncEnumerable<PerforceResponse> StreamCommandAsync(this IPerforceConnection perforce, string command, IReadOnlyList<string> arguments, IReadOnlyList<string>? fileArguments, byte[]? inputData, Type? statRecordType, bool interceptIo, [EnumeratorCancellation] CancellationToken cancellationToken)
 		{
+			using TelemetrySpan span = perforce.Tracer.StartActiveSpan(command);
+			AddAttributes(perforce, span, "stream", command, arguments, fileArguments);
+			int c = 0;
 #pragma warning disable CA1849 // Call async methods when in an async method
 			await using (IPerforceOutput output = perforce.Command(command, arguments, fileArguments, inputData, null, interceptIo))
 			{
 				await foreach (PerforceResponse response in output.ReadStreamingResponsesAsync(statRecordType, cancellationToken))
 				{
+					c++;
 					yield return response;
 				}
 			}
 #pragma warning restore CA1849 // Call async methods when in an async method
+			span.SetAttribute("epic.perforce.res.count", c);
 		}
 
 		/// <summary>
@@ -472,16 +504,21 @@ namespace EpicGames.Perforce
 		/// <returns>List of objects returned by the server</returns>
 		public static async IAsyncEnumerable<PerforceResponse<T>> StreamCommandAsync<T>(this IPerforceConnection perforce, string command, IReadOnlyList<string> arguments, IReadOnlyList<string>? fileArguments = null, byte[]? inputData = null, [EnumeratorCancellation] CancellationToken cancellationToken = default) where T : class
 		{
+			using TelemetrySpan span = perforce.Tracer.StartActiveSpan(command);
+			AddAttributes(perforce, span, "stream", command, arguments, fileArguments);
+			int c = 0;
 #pragma warning disable CA1849 // Call async methods when in an async method
 			await using (IPerforceOutput output = perforce.Command(command, arguments, fileArguments, inputData, null, false))
 			{
 				Type statRecordType = typeof(T);
 				await foreach (PerforceResponse response in output.ReadStreamingResponsesAsync(statRecordType, cancellationToken))
 				{
+					c++;
 					yield return new PerforceResponse<T>(response);
 				}
 			}
 #pragma warning restore CA1849 // Call async methods when in an async method
+			span.SetAttribute("epic.perforce.res.count", c);
 		}
 
 		/// <summary>
@@ -496,11 +533,11 @@ namespace EpicGames.Perforce
 		/// <returns>List of objects returned by the server</returns>
 		public static async Task RecordCommandAsync(this IPerforceConnection perforce, string command, IReadOnlyList<string> arguments, byte[]? inputData, Action<PerforceRecord> handleRecord, CancellationToken cancellationToken = default)
 		{
+			using TelemetrySpan span = perforce.Tracer.StartActiveSpan(command);
+			AddAttributes(perforce, span, "record", command, arguments, null);
 #pragma warning disable CA1849 // Call async methods when in an async method
-			await using (IPerforceOutput response = perforce.Command(command, arguments, null, inputData, null, false))
-			{
-				await response.ReadRecordsAsync(handleRecord, cancellationToken);
-			}
+			await using IPerforceOutput response = perforce.Command(command, arguments, null, inputData, null, false);
+			await response.ReadRecordsAsync(handleRecord, cancellationToken);
 #pragma warning restore CA1849 // Call async methods when in an async method
 		}
 
diff --git a/Engine/Source/Programs/Shared/EpicGames.Perforce/PerforceResponse.cs b/Engine/Source/Programs/Shared/EpicGames.Perforce/PerforceResponse.cs
--- a/Engine/Source/Programs/Shared/EpicGames.Perforce/PerforceResponse.cs
+++ b/Engine/Source/Programs/Shared/EpicGames.Perforce/PerforceResponse.cs
@@ -38,7 +38,7 @@ namespace EpicGames.Perforce
 		public bool Failed => InternalData is PerforceError;
 
 		/// <summary>
-		/// Accessor for the succcessful response data. Throws an exception if the response is an error.
+		/// Accessor for the successful response data. Throws an exception if the response is an error.
 		/// </summary>
 		public object Data
 		{
