EpicGames.Core: Use a null ticker in DefaultClock

diff --git a/Engine/Source/Programs/Shared/EpicGames.Core/Clock.cs b/Engine/Source/Programs/Shared/EpicGames.Core/Clock.cs
--- a/Engine/Source/Programs/Shared/EpicGames.Core/Clock.cs
+++ b/Engine/Source/Programs/Shared/EpicGames.Core/Clock.cs
@@ -90,13 +90,13 @@ public class DefaultClock : IClock
 	/// <inheritdoc/>
 	public ITicker AddTicker(string name, TimeSpan interval, Func<CancellationToken, ValueTask<TimeSpan?>> tickAsync, ILogger logger)
 	{
-		throw new NotImplementedException("Not available in default implementation");
+		return new NullTicker();
 	}
 
 	/// <inheritdoc/>
 	public ITicker AddSharedTicker(string name, TimeSpan interval, Func<CancellationToken, ValueTask> tickAsync, ILogger logger)
 	{
-		throw new NotImplementedException("Not available in default implementation");
+		return new NullTicker();
 	}
 }
 
