[EpicGames.Core]
* Added AsSpan to StackStringBuilder

diff --git a/Engine/Source/Programs/Shared/EpicGames.Core/StackStringBuffer.cs b/Engine/Source/Programs/Shared/EpicGames.Core/StackStringBuffer.cs
--- a/Engine/Source/Programs/Shared/EpicGames.Core/StackStringBuffer.cs
+++ b/Engine/Source/Programs/Shared/EpicGames.Core/StackStringBuffer.cs
@@ -54,16 +54,12 @@ namespace EpicGames.Core
 		public void Append(char c)
 		{
 			int pos = _pos;
-			if ((uint)pos < (uint)_chars.Length)
-			{
-				_chars[pos] = c;
-				_pos = pos + 1;
-			}
-			else
+			if ((uint)pos >= (uint)_chars.Length)
 			{
 				Grow(1);
-				Append(c);
 			}
+			_chars[pos] = c;
+			_pos = pos + 1;
 		}
 
 		/// <summary>
@@ -100,6 +96,12 @@ namespace EpicGames.Core
 			return new string(_chars.Slice(0, _pos));
 		}
 
+		/// <summary>
+		/// As span
+		/// </summary>
+		/// <returns></returns>
+		public ReadOnlySpan<char> AsSpan() => _chars.Slice(0, _pos);
+
 		/// <summary>
 		/// Dispose buffer
 		/// </summary>
