[EpicGames.Core]
* Optimized FileSystemReference MakeRelativeTo and added a version that takes in a StackStringBuilder

diff --git a/Engine/Source/Programs/Shared/EpicGames.Core/FileSystemReference.cs b/Engine/Source/Programs/Shared/EpicGames.Core/FileSystemReference.cs
--- a/Engine/Source/Programs/Shared/EpicGames.Core/FileSystemReference.cs
+++ b/Engine/Source/Programs/Shared/EpicGames.Core/FileSystemReference.cs
@@ -250,104 +250,129 @@ namespace EpicGames.Core
 		/// <returns>True if the path is under the given directory</returns>
 		public bool ContainsAnyNames(IEnumerable<string> names, DirectoryReference baseDir) => IsUnderDirectory(baseDir) && names.Any(x => ContainsName(x, baseDir.FullName.Length));
 
-		static readonly ThreadLocal<StringBuilder> s_makeRelativeToStringBuilder = new ThreadLocal<StringBuilder>(() => new StringBuilder(260));
-
 		/// <summary>
 		/// Creates a relative path from the given base directory
 		/// </summary>
 		/// <param name="directory">The directory to create a relative path from</param>
 		/// <returns>A relative path from the given directory</returns>
+		[MethodImpl(MethodImplOptions.AggressiveOptimization)]
 		public string MakeRelativeTo(DirectoryReference directory)
 		{
-			StringBuilder result = s_makeRelativeToStringBuilder.Value!.Clear();
-			WriteRelativeTo(directory, result);
-			return result.ToString();
+			if (String.Equals(FullName, directory.FullName, Comparison))
+			{
+				return ".";
+			}
+			Span<char> initial = stackalloc char[512];
+			StackStringBuilder sb = new(initial);
+			MakeRelativeTo(ref sb, directory);
+			string result = sb.ToString();
+			sb.Dispose();
+			return result;
 		}
 
 		/// <summary>
-		/// Appends a relative path to a string builder
+		/// Creates a relative path from the given base directory
 		/// </summary>
-		/// <param name="directory"></param>
-		/// <param name="result"></param>
+		/// <param name="outBuilder">The string builder where the result is written to</param>
+		/// <param name="directory">The directory to create a relative path from</param>
+#pragma warning disable CA1045
 		[MethodImpl(MethodImplOptions.AggressiveOptimization)]
-		public void WriteRelativeTo(DirectoryReference directory, StringBuilder result)
+		public void MakeRelativeTo(ref StackStringBuilder outBuilder, DirectoryReference directory)
 		{
-			// Find how much of the path is common between the two paths. This length does not include a trailing directory separator character.
-			int commonDirectoryLength = -1;
-			for (int idx = 0; ; idx++)
+			string targetDir = FullName;
+			int targetLength = targetDir.Length;
+			string baseDir = directory.FullName;
+			int baseLength = baseDir.Length;
+			char sep = Path.DirectorySeparatorChar;
+
+			// 1) find common directory boundary by segments
+			int ti = 0;
+			int bi = 0;
+			int common = -1;
+
+			while (true)
 			{
-				if (idx == FullName.Length)
+				// scan to next separator (or end) in target and base
+				int tn = ti;
+				while (tn < targetLength && targetDir[tn] != sep)
 				{
-					// The two paths are identical. Just return the "." character.
-					if (idx == directory.FullName.Length)
-					{
-						result.Append('.');
-						return;
-					}
+					tn++;
+				}
 
-					// Check if we're finishing on a complete directory name
-					if (directory.FullName[idx] == Path.DirectorySeparatorChar)
-					{
-						commonDirectoryLength = idx;
-					}
-					break;
+				int bn = bi;
+				while (bn < baseLength && baseDir[bn] != sep)
+				{
+					bn++;
 				}
-				else if (idx == directory.FullName.Length)
+
+				int tlen = tn - ti;
+				int blen = bn - bi;
+
+				if (tlen != blen || String.Compare(targetDir, ti, baseDir, bi, tlen, Comparison) != 0)
 				{
-					// Check whether the end of the directory name coincides with a boundary for the current name.
-					if (FullName[idx] == Path.DirectorySeparatorChar)
-					{
-						commonDirectoryLength = idx;
-					}
-					break;
+					break; // diverged within this segment
 				}
-				else
+
+				common = tn; // after matched segment
+
+				if (tn == targetLength || bn == baseLength)
 				{
-					// Check the two paths match, and bail if they don't. Increase the common directory length if we've reached a separator.
-					if (String.Compare(FullName, idx, directory.FullName, idx, 1, Comparison) != 0)
-					{
-						break;
-					}
-					if (FullName[idx] == Path.DirectorySeparatorChar)
-					{
-						commonDirectoryLength = idx;
-					}
+					break;
 				}
+
+				ti = tn + 1;
+				bi = bn + 1;
 			}
 
-			// If there's no relative path, just return the absolute path
-			if (commonDirectoryLength == -1)
+			if (common == -1)
 			{
-				result.Append(FullName);
+				outBuilder.Append(targetDir); // no relative path possible per your original behavior
 				return;
 			}
 
-			// Append all the '..' separators to get back to the common directory, then the rest of the string to reach the target item
-			for (int idx = commonDirectoryLength + 1; idx < directory.FullName.Length; idx++)
+			// 2) count how many ".." from base after 'common'
+			int ups = 0;
+			int scan = (common < baseLength && baseDir[common] == sep) ? common + 1 : common;
+			while (scan < baseLength)
 			{
-				// Move up a directory
-				if (result.Length != 0)
+				ups++;
+				int next = baseDir.IndexOf(sep, scan);
+				if (next < 0)
 				{
-					result.Append(Path.DirectorySeparatorChar);
+					break;
 				}
-				result.Append("..");
+				scan = next + 1;
+			}
 
-				// Scan to the next directory separator
-				while (idx < directory.FullName.Length && directory.FullName[idx] != Path.DirectorySeparatorChar)
+			// 3) tail from target after 'common'
+			int tailStart = (common < targetLength && targetDir[common] == sep) ? common + 1 : common;
+			int tailLen = targetLength - tailStart;
+
+			bool needSep = false;
+			for (int i = 0; i < ups; i++)
+			{
+				if (needSep)
 				{
-					idx++;
+					outBuilder.Append(sep);
 				}
+				outBuilder.Append("..");
+				needSep = true;
 			}
 
-			if (commonDirectoryLength + 1 < FullName.Length)
+			if (tailLen > 0)
 			{
-				if (result.Length != 0)
+				if (needSep)
 				{
-					result.Append(Path.DirectorySeparatorChar);
+					outBuilder.Append(sep);
 				}
-				result.Append(FullName, commonDirectoryLength + 1, FullName.Length - commonDirectoryLength - 1);
+				outBuilder.Append(targetDir.AsSpan(tailStart, tailLen));
+			}
+			else if (targetLength == baseLength)
+			{
+				outBuilder.Append('.');
 			}
 		}
+#pragma warning restore CA1045
 
 		/// <summary>
 		/// Normalize the path to using forward slashes
