[EpicGames.Core]
* Removed unnecessary using directive

diff --git a/Engine/Source/Programs/Shared/EpicGames.Core/FileSystemReference.cs b/Engine/Source/Programs/Shared/EpicGames.Core/FileSystemReference.cs
--- a/Engine/Source/Programs/Shared/EpicGames.Core/FileSystemReference.cs
+++ b/Engine/Source/Programs/Shared/EpicGames.Core/FileSystemReference.cs
@@ -5,7 +5,6 @@ using System.Collections.Generic;
 using System.IO;
 using System.Linq;
 using System.Runtime.CompilerServices;
-using System.Text;
 using System.Threading;
 
 namespace EpicGames.Core
