UHT: Support for "None" documentation policy. Don't enfore policy for deprecated types

#rnx

diff --git a/Engine/Source/Programs/Shared/EpicGames.UHT/Specifiers/UhtDefaultSpecifiers.cs b/Engine/Source/Programs/Shared/EpicGames.UHT/Specifiers/UhtDefaultSpecifiers.cs
--- a/Engine/Source/Programs/Shared/EpicGames.UHT/Specifiers/UhtDefaultSpecifiers.cs
+++ b/Engine/Source/Programs/Shared/EpicGames.UHT/Specifiers/UhtDefaultSpecifiers.cs
@@ -157,10 +157,10 @@ namespace EpicGames.UHT.Parsers
 			SetMetaData(specifierContext, "DevelopmentStatus", "EarlyAccess");
 		}
 
-		[UhtSpecifier(Extends = UhtTableNames.Default, ValueType = UhtSpecifierValueType.None, When = UhtSpecifierWhen.Immediate)]
-		private static void DocumentationPolicySpecifier(UhtSpecifierContext specifierContext)
+		[UhtSpecifier(Extends = UhtTableNames.Default, ValueType = UhtSpecifierValueType.String, When = UhtSpecifierWhen.Immediate)]
+		private static void DocumentationPolicySpecifier(UhtSpecifierContext specifierContext, StringView value)
 		{
-			SetMetaData(specifierContext, UhtNames.DocumentationPolicy, "Strict");
+			SetMetaData(specifierContext, UhtNames.DocumentationPolicy, value);
 		}
 
 		[UhtSpecifier(Extends = UhtTableNames.Default, ValueType = UhtSpecifierValueType.String, When = UhtSpecifierWhen.Immediate)]
@@ -219,9 +219,11 @@ namespace EpicGames.UHT.Parsers
 		private static void DocumentationPolicySpecifierValidator(UhtType type, UhtMetaData metaData, UhtMetaDataKey key, StringView value)
 		{
 			const string StrictValue = "Strict";
-			if (!value.Span.Equals(StrictValue, StringComparison.OrdinalIgnoreCase))
+			const string NoneValue = "None";
+			if (!value.Span.Equals(StrictValue, StringComparison.OrdinalIgnoreCase) &&
+			    !value.Span.Equals(NoneValue, StringComparison.OrdinalIgnoreCase))
 			{
-				type.LogError(metaData.LineNumber, $"'{key}' metadata was '{value}' but it must be '{StrictValue}'");
+				type.LogError(metaData.LineNumber, $"'{key}' metadata was '{value}' but it must be '{StrictValue}' or '{NoneValue}'");
 			}
 		}
 
diff --git a/Engine/Source/Programs/Shared/EpicGames.UHT/Types/UhtScriptStruct.cs b/Engine/Source/Programs/Shared/EpicGames.UHT/Types/UhtScriptStruct.cs
--- a/Engine/Source/Programs/Shared/EpicGames.UHT/Types/UhtScriptStruct.cs
+++ b/Engine/Source/Programs/Shared/EpicGames.UHT/Types/UhtScriptStruct.cs
@@ -921,7 +921,7 @@ namespace EpicGames.UHT.Types
 			}
 
 			// Validate RigVM
-			if (RigVMStructInfo != null && MetaData.ContainsKey("Deprecated") && !RigVMStructInfo.HasGetUpgradeInfoMethod)
+			if (RigVMStructInfo != null && MetaData.ContainsKey(UhtNames.Deprecated) && !RigVMStructInfo.HasGetUpgradeInfoMethod)
 			{
 				this.LogError($"RigVMStruct '{SourceName}' is marked as deprecated but is missing 'GetUpgradeInfo method.");
 				this.LogError("Please implement a method like below:");
diff --git a/Engine/Source/Programs/Shared/EpicGames.UHT/Types/UhtStruct.cs b/Engine/Source/Programs/Shared/EpicGames.UHT/Types/UhtStruct.cs
--- a/Engine/Source/Programs/Shared/EpicGames.UHT/Types/UhtStruct.cs
+++ b/Engine/Source/Programs/Shared/EpicGames.UHT/Types/UhtStruct.cs
@@ -360,6 +360,16 @@ namespace EpicGames.UHT.Types
 				Dictionary<string, UhtProperty> toolTipToType = new();
 				foreach (UhtProperty property in Properties)
 				{
+					// we don't require tooltips / comments on hidden properties.
+					// only properties exposed to the user interface require documentation.
+					string inputFlag = property.MetaData.GetValueOrDefault(UhtNames.Input);
+					string outputFlag = property.MetaData.GetValueOrDefault(UhtNames.Output);
+					string visibleFlag = property.MetaData.GetValueOrDefault(UhtNames.Visible);
+					if (inputFlag.Length == 0 && outputFlag.Length == 0 && visibleFlag.Length == 0)
+					{
+						continue;
+					}
+
 					string toolTip = property.GetToolTipText();
 					if (toolTip.Length == 0 || toolTip == property.GetDisplayNameText())
 					{
diff --git a/Engine/Source/Programs/Shared/EpicGames.UHT/Types/UhtType.cs b/Engine/Source/Programs/Shared/EpicGames.UHT/Types/UhtType.cs
--- a/Engine/Source/Programs/Shared/EpicGames.UHT/Types/UhtType.cs
+++ b/Engine/Source/Programs/Shared/EpicGames.UHT/Types/UhtType.cs
@@ -1294,6 +1294,13 @@ namespace EpicGames.UHT.Types
 				return;
 			}
 
+			// deprecated structs don't require the documentation policy.
+			string deprecatedValue = MetaData.GetValueOrDefaultHierarchical(UhtNames.Deprecated);
+			if (!String.IsNullOrEmpty(deprecatedValue))
+			{
+				return;
+			}
+
 			string policyName = DocumentationPolicyName;
 			if (String.IsNullOrEmpty(policyName))
 			{
diff --git a/Engine/Source/Programs/Shared/EpicGames.UHT/Utils/UhtNames.cs b/Engine/Source/Programs/Shared/EpicGames.UHT/Utils/UhtNames.cs
--- a/Engine/Source/Programs/Shared/EpicGames.UHT/Utils/UhtNames.cs
+++ b/Engine/Source/Programs/Shared/EpicGames.UHT/Utils/UhtNames.cs
@@ -31,6 +31,7 @@ namespace EpicGames.UHT.Utils
 		public const string Constant = "Constant";
 		public const string CppFromBpEvent = "CppFromBpEvent";
 		public const string CustomThunk = "CustomThunk";
+		public const string Deprecated = "Deprecated";
 		public const string DeprecatedFunction = "DeprecatedFunction";
 		public const string DisplayName = "DisplayName";
 		public const string DocumentationPolicy = "DocumentationPolicy";
diff --git a/Engine/Source/Programs/UnrealBuildTool/Modes/UnrealHeaderToolMode.cs b/Engine/Source/Programs/UnrealBuildTool/Modes/UnrealHeaderToolMode.cs
--- a/Engine/Source/Programs/UnrealBuildTool/Modes/UnrealHeaderToolMode.cs
+++ b/Engine/Source/Programs/UnrealBuildTool/Modes/UnrealHeaderToolMode.cs
@@ -448,6 +448,15 @@ namespace UnrealBuildTool.Modes
 				FloatRangesRequired = true,
 			};
 
+			_documentationPolicies["None"] = new()
+			{
+				ClassOrStructCommentRequired = false,
+				FunctionToolTipsRequired = false,
+				MemberToolTipsRequired = false,
+				ParameterToolTipsRequired = false,
+				FloatRangesRequired = false,
+			};
+
 			if (_ini.TryGetValues(Section, Key, out IReadOnlyList<string>? StringList))
 			{
 				foreach (string Value in StringList)
