UHT: Only require documentation policy based UI ranges on input members

#rnx

diff --git a/Engine/Source/Programs/Shared/EpicGames.UHT/Specifiers/UhtDefaultSpecifiers.cs b/Engine/Source/Programs/Shared/EpicGames.UHT/Specifiers/UhtDefaultSpecifiers.cs
--- a/Engine/Source/Programs/Shared/EpicGames.UHT/Specifiers/UhtDefaultSpecifiers.cs
+++ b/Engine/Source/Programs/Shared/EpicGames.UHT/Specifiers/UhtDefaultSpecifiers.cs
@@ -197,7 +197,10 @@ namespace EpicGames.UHT.Parsers
 		{
 			if (!UhtFCString.IsNumeric(value.Span))
 			{
-				type.LogError($"Metadata value for '{key}' is non-numeric : '{value}'");
+				if (value.Length > 0)
+				{
+					type.LogError($"Metadata value for '{key}' is non-numeric : '{value}'");
+				}
 			}
 		}
 
diff --git a/Engine/Source/Programs/Shared/EpicGames.UHT/Types/UhtStruct.cs b/Engine/Source/Programs/Shared/EpicGames.UHT/Types/UhtStruct.cs
--- a/Engine/Source/Programs/Shared/EpicGames.UHT/Types/UhtStruct.cs
+++ b/Engine/Source/Programs/Shared/EpicGames.UHT/Types/UhtStruct.cs
@@ -362,10 +362,10 @@ namespace EpicGames.UHT.Types
 				{
 					// we don't require tooltips / comments on hidden properties.
 					// only properties exposed to the user interface require documentation.
-					string inputFlag = property.MetaData.GetValueOrDefault(UhtNames.Input);
-					string outputFlag = property.MetaData.GetValueOrDefault(UhtNames.Output);
-					string visibleFlag = property.MetaData.GetValueOrDefault(UhtNames.Visible);
-					if (inputFlag.Length == 0 && outputFlag.Length == 0 && visibleFlag.Length == 0)
+					bool inputFlag = property.MetaData.ContainsKey(UhtNames.Input);
+					bool outputFlag = property.MetaData.ContainsKey(UhtNames.Output);
+					bool visibleFlag = property.MetaData.ContainsKey(UhtNames.Visible);
+					if (!inputFlag && !outputFlag && !visibleFlag)
 					{
 						continue;
 					}
@@ -394,6 +394,12 @@ namespace EpicGames.UHT.Types
 				{
 					if (child is UhtDoubleProperty || child is UhtFloatProperty)
 					{
+						bool inputFlag = child.MetaData.ContainsKey(UhtNames.Input);
+						bool visibleFlag = child.MetaData.ContainsKey(UhtNames.Visible);
+						if (!inputFlag && !visibleFlag)
+						{
+							continue;
+						}
 						if (!CheckUIMinMaxRangeFromMetaData(child))
 						{
 							child.LogError($"Property '{SourceName}::{child.SourceName}' does not provide a valid UIMin / UIMax (DocumentationPolicy).");
