UHT: Follow up: Only require documentation policy based UI ranges on input members

#rnx

diff --git a/Engine/Source/Programs/Shared/EpicGames.UHT/Specifiers/UhtDefaultSpecifiers.cs b/Engine/Source/Programs/Shared/EpicGames.UHT/Specifiers/UhtDefaultSpecifiers.cs
--- a/Engine/Source/Programs/Shared/EpicGames.UHT/Specifiers/UhtDefaultSpecifiers.cs
+++ b/Engine/Source/Programs/Shared/EpicGames.UHT/Specifiers/UhtDefaultSpecifiers.cs
@@ -197,10 +197,7 @@ namespace EpicGames.UHT.Parsers
 		{
 			if (!UhtFCString.IsNumeric(value.Span))
 			{
-				if (value.Length > 0)
-				{
-					type.LogError($"Metadata value for '{key}' is non-numeric : '{value}'");
-				}
+				type.LogError($"Metadata value for '{key}' is non-numeric : '{value}'");
 			}
 		}
 
diff --git a/Engine/Source/Programs/Shared/EpicGames.UHT/Types/UhtStruct.cs b/Engine/Source/Programs/Shared/EpicGames.UHT/Types/UhtStruct.cs
--- a/Engine/Source/Programs/Shared/EpicGames.UHT/Types/UhtStruct.cs
+++ b/Engine/Source/Programs/Shared/EpicGames.UHT/Types/UhtStruct.cs
@@ -321,11 +321,18 @@ namespace EpicGames.UHT.Types
 
 		private static bool CheckUIMinMaxRangeFromMetaData(UhtType child)
 		{
+			bool inputFlag = child.MetaData.ContainsKey(UhtNames.Input);
+			bool visibleFlag = child.MetaData.ContainsKey(UhtNames.Visible);
+			if (!inputFlag && !visibleFlag)
+			{
+				return true;
+			}
+
 			string uiMin = child.MetaData.GetValueOrDefault(UhtNames.UIMin);
 			string uiMax = child.MetaData.GetValueOrDefault(UhtNames.UIMax);
 			if (uiMin.Length == 0 || uiMax.Length == 0)
 			{
-				return false;
+				return true;
 			}
 
 			// NOTE: Old UHT didn't handle parse errors
@@ -394,12 +401,6 @@ namespace EpicGames.UHT.Types
 				{
 					if (child is UhtDoubleProperty || child is UhtFloatProperty)
 					{
-						bool inputFlag = child.MetaData.ContainsKey(UhtNames.Input);
-						bool visibleFlag = child.MetaData.ContainsKey(UhtNames.Visible);
-						if (!inputFlag && !visibleFlag)
-						{
-							continue;
-						}
 						if (!CheckUIMinMaxRangeFromMetaData(child))
 						{
 							child.LogError($"Property '{SourceName}::{child.SourceName}' does not provide a valid UIMin / UIMax (DocumentationPolicy).");
