UAT/UBT: Strip OnBuildingProjects action from CompileScriptModule::InitializeScriptModules() and CsProjBuilder::Build() APIs to reduce complexity.
Log a generic "Building {Count} projects..." message instead (the logger was not available when OnBuildingProjects was added).
Add some additional stopwatch profiling log messages.

#rnx

diff --git a/Engine/Source/Programs/AutomationTool/Program.cs b/Engine/Source/Programs/AutomationTool/Program.cs
--- a/Engine/Source/Programs/AutomationTool/Program.cs
+++ b/Engine/Source/Programs/AutomationTool/Program.cs
@@ -480,7 +480,7 @@ namespace AutomationToolDriver
 		{
 			ILogger Logger = Log.Logger;
 			Logger.LogInformation("Initializing script modules...");
-			var StartTime = DateTime.UtcNow;
+			Stopwatch Timer = Stopwatch.StartNew();
 			string ScriptsForProject = (string)AutomationToolCommandLine.GetValueUnchecked("-ScriptsForProject");
 			List<string> AdditionalScriptDirs = (List<string>) AutomationToolCommandLine.GetValueUnchecked("-ScriptDir");
 			bool bForceCompile = AutomationToolCommandLine.IsSetGlobal("-Compile");
@@ -492,11 +492,7 @@ namespace AutomationToolDriver
 			bool bBuildSuccess;
 			HashSet<FileReference> ScriptModuleAssemblyPaths = CompileScriptModule.InitializeScriptModules(
 					Rules.RulesFileType.AutomationModule, ScriptsForProject, AdditionalScriptDirs, bForceCompile, bNoCompile, bUseBuildRecords, 
-					out bBuildSuccess, (int Count) =>
-                    {
-						Logger.LogInformation("Building {Count} projects...", Count);
-					},
-					Log.Logger);
+					out bBuildSuccess, Logger);
 
 			if (!bBuildSuccess)
             {
@@ -523,7 +519,7 @@ namespace AutomationToolDriver
 
 			Type AutomationTools_Automation = AutomationUtilsAssembly.GetType("AutomationTool.Automation");
 			MethodInfo Automation_Process = AutomationTools_Automation.GetMethod("ProcessAsync");
-			Logger.LogInformation("Total script module initialization time: {InitTime:0.00} s.", (DateTime.UtcNow - StartTime).TotalMilliseconds / 1000);
+			Logger.LogInformation("Total script module initialization time: {InitTime:0.00} s.", Timer.Elapsed.TotalSeconds);
 			return await (Task<ExitCode>) Automation_Process.Invoke(null,
 				new object[] {AutomationToolCommandLine, StartupListener, ScriptModuleAssemblyPaths});
 		}
diff --git a/Engine/Source/Programs/Shared/EpicGames.Build/System/CompileScriptModules.cs b/Engine/Source/Programs/Shared/EpicGames.Build/System/CompileScriptModules.cs
--- a/Engine/Source/Programs/Shared/EpicGames.Build/System/CompileScriptModules.cs
+++ b/Engine/Source/Programs/Shared/EpicGames.Build/System/CompileScriptModules.cs
@@ -3,6 +3,7 @@
 using System;
 using System.Collections.Concurrent;
 using System.Collections.Generic;
+using System.Diagnostics;
 using System.IO;
 using System.Linq;
 using System.Runtime.CompilerServices;
@@ -141,13 +142,13 @@ namespace UnrealBuildBase
 		/// <param name="bNoCompile"></param>
 		/// <param name="bUseBuildRecords"></param>
 		/// <param name="bBuildSuccess"></param>
-		/// <param name="OnBuildingProjects">Action to invoke when projects get built</param>
 		/// <param name="Logger"></param>
 		/// <returns>Collection of all the projects.  They will have been compiled.</returns>
 		public static HashSet<FileReference> InitializeScriptModules(Rules.RulesFileType RulesFileType,
 			string? ScriptsForProjectFileName, List<string>? AdditionalScriptsFolders, bool bForceCompile, bool bNoCompile, bool bUseBuildRecords,
-			out bool bBuildSuccess, Action<int> OnBuildingProjects, ILogger Logger)
+			out bool bBuildSuccess, ILogger Logger)
 		{
+			Stopwatch Timer = Stopwatch.StartNew();
 			List<DirectoryReference> GameDirectories = GetGameDirectories(ScriptsForProjectFileName, Logger);
 			List<DirectoryReference> AdditionalDirectories = GetAdditionalDirectories(AdditionalScriptsFolders);
 			List<DirectoryReference> GameBuildDirectories = GetAdditionalBuildDirectories(GameDirectories);
@@ -160,12 +161,15 @@ namespace UnrealBuildBase
 				.. AdditionalDirectories,
 			];
 			BaseDirectories = BaseDirectories.Distinct().ToList();
+			Logger.LogDebug("Find {Count} directories time: {TimeSeconds:0.00} s", BaseDirectories.Count, Timer.Elapsed.TotalSeconds);
 
+			Timer.Restart();
 			HashSet<FileReference> FoundProjects =
 				[.. Rules.FindAllRulesSourceFiles(RulesFileType,
 				// Project scripts require source engine builds
 				GameFolders: Unreal.IsEngineInstalled() ? GameDirectories : [],
 				ForeignPlugins: null, AdditionalSearchPaths: [.. AdditionalDirectories, .. GameBuildDirectories])];
+			Logger.LogDebug("Find {Count} projects time: {TimeSeconds:0.00} s", FoundProjects.Count(), Timer.Elapsed.TotalSeconds);
 
 			BuildFlags BuildFlags = BuildFlags.None;
 			if (bForceCompile)
@@ -203,7 +207,6 @@ namespace UnrealBuildBase
 					null,
 					EngineBuildFlags,
 					out bEngineBuildSuccess,
-					OnBuildingProjects,
 					Logger);
 
 				bool bProjectBuildSuccess = false;
@@ -220,7 +223,6 @@ namespace UnrealBuildBase
 					null,
 					ProjectBuildFlags,
 					out bProjectBuildSuccess,
-					OnBuildingProjects,
 					Logger);
 
 				bBuildSuccess = bEngineBuildSuccess && bProjectBuildSuccess;
@@ -235,7 +237,7 @@ namespace UnrealBuildBase
 			}
 			else
 			{
-				BuildResults = Build(RulesFileType, FoundProjects, BaseDirectories, null, BuildFlags, out bBuildSuccess, OnBuildingProjects, Logger);
+				BuildResults = Build(RulesFileType, FoundProjects, BaseDirectories, null, BuildFlags, out bBuildSuccess, Logger);
 			}
 
 			return GetTargetPaths(BuildResults);
@@ -274,12 +276,11 @@ namespace UnrealBuildBase
 		/// <param name="DefineConstants">Collection of constants to be added to the project</param>
 		/// <param name="BuildFlags">Collection of flags to customize compilation</param>
 		/// <param name="bBuildSuccess"></param>
-		/// <param name="OnBuildingProjects">Action to invoke when projects get built</param>
 		/// <param name="Logger"></param>
 		/// <returns>Collection of all the projects.  They will have been compiled.</returns>
 		public static Dictionary<FileReference, CsProjBuildRecordEntry> Build(Rules.RulesFileType RulesFileType,
 			HashSet<FileReference> FoundProjects, IEnumerable<DirectoryReference> BaseDirectories, IEnumerable<string>? DefineConstants,
-			BuildFlags BuildFlags, out bool bBuildSuccess, Action<int> OnBuildingProjects, ILogger Logger)
+			BuildFlags BuildFlags, out bool bBuildSuccess, ILogger Logger)
 		{
 			CsProjBuildHook Hook = new Hook(Logger, RulesFileType);
 
@@ -383,7 +384,7 @@ namespace UnrealBuildBase
 			}
 
 			// Fall back to the slower approach: use msbuild to load csproj files & build as necessary
-			return Build(FoundProjects, BuildFlags.HasFlag(BuildFlags.ForceCompile), out bBuildSuccess, Hook, BaseDirectories, DefineConstants, OnBuildingProjects, Logger);
+			return Build(FoundProjects, BuildFlags.HasFlag(BuildFlags.ForceCompile), out bBuildSuccess, Hook, BaseDirectories, DefineConstants, Logger);
 		}
 
 		/// <summary>
@@ -394,9 +395,9 @@ namespace UnrealBuildBase
 		[MethodImpl(MethodImplOptions.NoInlining)]
 		static Dictionary<FileReference, CsProjBuildRecordEntry> Build(HashSet<FileReference> FoundProjects,
 			bool bForceCompile, out bool bBuildSuccess, CsProjBuildHook Hook, IEnumerable<DirectoryReference> BaseDirectories,
-			IEnumerable<string>? DefineConstants, Action<int> OnBuildingProjects, ILogger Logger)
+			IEnumerable<string>? DefineConstants, ILogger Logger)
 		{
-			return CsProjBuilder.Build(FoundProjects, bForceCompile, out bBuildSuccess, Hook, BaseDirectories, DefineConstants ?? ([]), OnBuildingProjects, Logger);
+			return CsProjBuilder.Build(FoundProjects, bForceCompile, out bBuildSuccess, Hook, BaseDirectories, DefineConstants ?? ([]), Logger);
 		}
 
 		/// <summary>
diff --git a/Engine/Source/Programs/Shared/EpicGames.Build/System/EnumeratePlugins.cs b/Engine/Source/Programs/Shared/EpicGames.Build/System/EnumeratePlugins.cs
--- a/Engine/Source/Programs/Shared/EpicGames.Build/System/EnumeratePlugins.cs
+++ b/Engine/Source/Programs/Shared/EpicGames.Build/System/EnumeratePlugins.cs
@@ -224,19 +224,7 @@ namespace UnrealBuildBase
 			bool bBuildSuccess = true;
 			Dictionary<FileReference, CsProjBuildRecordEntry>? BuiltPlugins =
 				CompileScriptModule.Build(UnrealBuildBase.Rules.RulesFileType.UbtPlugin, [.. NotBuilt],
-				BaseDirectories, DefineConstants, CompileScriptModule.BuildFlags.UseBuildRecords, out bBuildSuccess,
-				Count =>
-				{
-					if (Log.OutputFile != null)
-					{
-						Logger.LogInformation("Building {Count} plugins (see Log '{LogFile}' for more details)", Count, Log.OutputFile);
-					}
-					else
-					{
-						Logger.LogInformation("Building {Count} plugins", Count);
-					}
-				}, Logger
-			);
+				BaseDirectories, DefineConstants, CompileScriptModule.BuildFlags.UseBuildRecords, out bBuildSuccess, Logger);
 
 			// Add any built plugins back into the cache
 			foreach (KeyValuePair<FileReference, CsProjBuildRecordEntry> BuiltPlugin in BuiltPlugins)
diff --git a/Engine/Source/Programs/Shared/EpicGames.MsBuild/CsProjBuilder.cs b/Engine/Source/Programs/Shared/EpicGames.MsBuild/CsProjBuilder.cs
--- a/Engine/Source/Programs/Shared/EpicGames.MsBuild/CsProjBuilder.cs
+++ b/Engine/Source/Programs/Shared/EpicGames.MsBuild/CsProjBuilder.cs
@@ -57,18 +57,17 @@ namespace EpicGames.MsBuild
 		/// <param name="hook">Interface to fetch data about the building environment</param>
 		/// <param name="baseDirectories">Base directories of the engine and project</param>
 		/// <param name="defineConstants">Collection of constants to be defined while building projects</param>
-		/// <param name="onBuildingProjects">Action invoked to notify caller regarding the number of projects being built</param>
 		/// <param name="logger">Destination logger</param>
 		public static Dictionary<FileReference, CsProjBuildRecordEntry> Build(HashSet<FileReference> foundProjects,
 			bool bForceCompile, out bool bBuildSuccess, CsProjBuildHook hook, IEnumerable<DirectoryReference> baseDirectories,
-			IEnumerable<string> defineConstants, Action<int> onBuildingProjects, ILogger logger)
+			IEnumerable<string> defineConstants, ILogger logger)
 		{
 
 			// Register the MS build path prior to invoking the internal routine.  By not having the internal routine
 			// inline, we avoid having the issue of the Microsoft.Build libraries being resolved prior to the build path
 			// being set.
 			RegisterMsBuildPath(hook);
-			return BuildInternal(foundProjects, bForceCompile, out bBuildSuccess, hook, baseDirectories, defineConstants, onBuildingProjects, logger);
+			return BuildInternal(foundProjects, bForceCompile, out bBuildSuccess, hook, baseDirectories, defineConstants, logger);
 		}
 
 		/// <summary>
@@ -80,13 +79,12 @@ namespace EpicGames.MsBuild
 		/// <param name="hook">Interface to fetch data about the building environment</param>
 		/// <param name="baseDirectories">Base directories of the engine and project</param>
 		/// <param name="defineConstants">Collection of constants to be defined while building projects</param>
-		/// <param name="onBuildingProjects">Action invoked to notify caller regarding the number of projects being built</param>
 		/// <param name="logger">Destination logger</param>
 		private static Dictionary<FileReference, CsProjBuildRecordEntry> BuildInternal(HashSet<FileReference> foundProjects,
-			bool bForceCompile, out bool bBuildSuccess, CsProjBuildHook hook, IEnumerable<DirectoryReference> baseDirectories, IEnumerable<string> defineConstants,
-			Action<int> onBuildingProjects, ILogger logger)
+			bool bForceCompile, out bool bBuildSuccess, CsProjBuildHook hook, IEnumerable<DirectoryReference> baseDirectories,
+			IEnumerable<string> defineConstants, ILogger logger)
 		{
-			Stopwatch stopwatch = new();
+			Stopwatch stopwatch = Stopwatch.StartNew();
 			FileReference csharpTargetsFile = FileReference.Combine(hook.EngineDirectory, "Source", "Programs", "Shared", "UnrealEngine.CSharp.targets");
 			Dictionary<string, string> globalProperties = new Dictionary<string, string>
 			{
@@ -113,6 +111,7 @@ namespace EpicGames.MsBuild
 			using ProjectCollection projectCollection = new ProjectCollection(globalProperties);
 			ConcurrentDictionary<FileReference, Project> projects = new();
 			ConcurrentBag<FileReference> platformSpecificProjects = new();
+			logger.LogDebug("Setup {Count} projects time: {TimeSeconds:0.00} s", foundProjects.Count, stopwatch.Elapsed.TotalSeconds);
 
 			stopwatch.Restart();
 			// Load all found projects
@@ -255,7 +254,7 @@ namespace EpicGames.MsBuild
 				buildRecords.AddOrUpdate(entry.ProjectFile, entry, (key, oldValue) => entry);
 			});
 
-			logger.LogDebug("Generate build records {Count} projects time: {TimeSeconds:0.00} s", projects.Count, stopwatch.Elapsed.TotalSeconds);
+			logger.LogDebug("Generate {Count} projects build records time: {TimeSeconds:0.00} s", projects.Count, stopwatch.Elapsed.TotalSeconds);
 
 			// Ensure no projects are set to the same OutputPath, as this will cause build issues with multiprocess compiling
 			{
@@ -296,7 +295,7 @@ namespace EpicGames.MsBuild
 
 			if (outOfDateProjects.Any())
 			{
-				onBuildingProjects(outOfDateProjects.Count());
+				logger.LogInformation("Building {Count} projects...", outOfDateProjects.Count());
 
 				IEnumerable<FileReference> singleBuildProjects = outOfDateProjects.Intersect(platformSpecificProjects);
 				bBuildSuccess = RunDotnetBuildAsync(outOfDateProjects.Except(singleBuildProjects), globalProperties, hook, logger, default).Result;
@@ -478,8 +477,7 @@ namespace EpicGames.MsBuild
 
 			using IDisposable mutex = await mutexTask;
 
-			Stopwatch stopwatch = new();
-			stopwatch.Start();
+			Stopwatch stopwatch = Stopwatch.StartNew();
 
 			FileReference buildProject;
 			if (projects.Count() > 1)
