[UBT]
* Added UBA null server and null logger

diff --git a/Engine/Source/Programs/Shared/EpicGames.UBA/ILogger.cs b/Engine/Source/Programs/Shared/EpicGames.UBA/ILogger.cs
--- a/Engine/Source/Programs/Shared/EpicGames.UBA/ILogger.cs
+++ b/Engine/Source/Programs/Shared/EpicGames.UBA/ILogger.cs
@@ -66,5 +66,10 @@ namespace EpicGames.UBA
 		{
 			return new LoggerImpl(logger);
 		}
+
+		/// <summary>
+		/// Creates null logger
+		/// </summary>
+		public static ILogger CreateNullLogger() => new NullLogger();
 	}
 }
\ No newline at end of file
diff --git a/Engine/Source/Programs/Shared/EpicGames.UBA/IServer.cs b/Engine/Source/Programs/Shared/EpicGames.UBA/IServer.cs
--- a/Engine/Source/Programs/Shared/EpicGames.UBA/IServer.cs
+++ b/Engine/Source/Programs/Shared/EpicGames.UBA/IServer.cs
@@ -17,7 +17,7 @@ namespace EpicGames.UBA
 		/// <param name="crypto">Enable crypto by using a 32 character crypto string (representing a 16 byte value)</param>
 		/// <returns></returns>
 		public abstract bool StartServer(string ip = "", int port = -1, string crypto = "");
-		
+
 		/// <summary>
 		/// Stop the server
 		/// </summary>
@@ -57,5 +57,10 @@ namespace EpicGames.UBA
 		{
 			return new ServerImpl(maxWorkers, sendSize, logger, useQuic);
 		}
+
+		/// <summary>
+		/// Create a null server
+		/// </summary>
+		public static IServer CreateNullServer() => new NullServer();
 	}
 }
diff --git a/Engine/Source/Programs/Shared/EpicGames.UBA/Impl/LoggerImpl.cs b/Engine/Source/Programs/Shared/EpicGames.UBA/Impl/LoggerImpl.cs
--- a/Engine/Source/Programs/Shared/EpicGames.UBA/Impl/LoggerImpl.cs
+++ b/Engine/Source/Programs/Shared/EpicGames.UBA/Impl/LoggerImpl.cs
@@ -21,9 +21,6 @@ namespace EpicGames.UBA.Impl
 		readonly object _lock = new object();
 
 		#region DllImport
-		[DllImport("UbaHost", CharSet = CharSet.Auto)]
-		static extern nint GetDefaultLogWriter();
-
 		[DllImport("UbaHost", CharSet = CharSet.Auto)]
 		static extern nint CreateCallbackLogWriter(BeginScopeCallback begin, EndScopeCallback end, LogCallback log);
 
@@ -95,4 +92,29 @@ namespace EpicGames.UBA.Impl
 
 		void Log(LogEntryType type, nint ptr, uint len) => Log(type, Marshal.PtrToStringAuto(ptr, (int)len) ?? String.Empty);
 	}
+
+	internal class NullLogger : ILogger
+	{
+		#region DllImport
+		[DllImport("UbaHost", CharSet = CharSet.Auto)]
+		static extern nint LogWriter_GetNull();
+		#endregion
+
+		/// <inheritdoc/>
+		public void BeginScope() { }
+		/// <inheritdoc/>
+		public void EndScope() { }
+		/// <inheritdoc/>
+		public void Log(LogEntryType type, string message) { }
+		/// <inheritdoc/>
+		public nint GetHandle() => LogWriter_GetNull();
+		/// <inheritdoc/>
+		public void Dispose()
+		{
+			Dispose(true);
+			GC.SuppressFinalize(this);
+		}
+		/// <inheritdoc/>
+		protected virtual void Dispose(bool disposing) { }
+	}
 }
\ No newline at end of file
diff --git a/Engine/Source/Programs/Shared/EpicGames.UBA/Impl/ServerImpl.cs b/Engine/Source/Programs/Shared/EpicGames.UBA/Impl/ServerImpl.cs
--- a/Engine/Source/Programs/Shared/EpicGames.UBA/Impl/ServerImpl.cs
+++ b/Engine/Source/Programs/Shared/EpicGames.UBA/Impl/ServerImpl.cs
@@ -92,4 +92,31 @@ namespace EpicGames.UBA.Impl
 
 		#endregion
 	}
+
+	/// <summary>
+	/// Null server
+	/// </summary>
+	internal class NullServer : IServer
+	{
+		/// <inheritdoc/>
+		public bool StartServer(string ip = "", int port = -1, string crypto = "") => false;
+		/// <inheritdoc/>
+		public void StopServer() { }
+		/// <inheritdoc/>
+		public bool AddNamedConnection(string name) => false;
+		/// <inheritdoc/>
+		public bool AddClient(string ip, int port, string crypto = "") => false;
+		/// <inheritdoc/>
+		public void SetClientsConfig(string path) { }
+		/// <inheritdoc/>
+		public nint GetHandle() => 0;
+		/// <inheritdoc/>
+		public void Dispose()
+		{
+			Dispose(true);
+			GC.SuppressFinalize(this);
+		}
+		/// <inheritdoc/>
+		protected virtual void Dispose(bool disposing) { }
+	}
 }
