Overriding AppID for Verse QA Tests so DataBricks collects information across Horde

#rnx

diff --git a/Engine/Source/Programs/Shared/EpicGames.Core/Telemetry/DataRouterTelemetryService.cs b/Engine/Source/Programs/Shared/EpicGames.Core/Telemetry/DataRouterTelemetryService.cs
--- a/Engine/Source/Programs/Shared/EpicGames.Core/Telemetry/DataRouterTelemetryService.cs
+++ b/Engine/Source/Programs/Shared/EpicGames.Core/Telemetry/DataRouterTelemetryService.cs
@@ -111,6 +111,16 @@ namespace EpicGames.Core.Telemetry
 
 			return urlParameters;
 		}
+
+		/// <summary>
+		/// Direct access to override the App ID.
+		/// This is used to redirect telemetry data in Databricks.
+		/// </summary>
+		/// <param name="appId"></param>
+		public void OverrideAppId(string appId)
+		{
+			this.AppId = appId;
+		}
 	}
 
 	/// <summary>
