Added missing xml docs for code moved from UGS

diff --git a/Engine/Source/Programs/Shared/EpicGames.Core/ProgressTextReader.cs b/Engine/Source/Programs/Shared/EpicGames.Core/ProgressTextReader.cs
--- a/Engine/Source/Programs/Shared/EpicGames.Core/ProgressTextReader.cs
+++ b/Engine/Source/Programs/Shared/EpicGames.Core/ProgressTextReader.cs
@@ -6,16 +6,25 @@ using System.Linq;
 
 namespace EpicGames.Core
 {
+	/// <summary>
+	/// Keeps track of progress reporting of a child process
+	/// </summary>
 	public class ProgressValue
 	{
 		Tuple<string, float> _state = null!;
 		readonly Stack<Tuple<float, float>> _ranges = new Stack<Tuple<float, float>>();
 
+		/// <summary>
+		/// CTOR
+		/// </summary>
 		public ProgressValue()
 		{
 			Clear();
 		}
 
+		/// <summary>
+		/// Reset the progress tracking
+		/// </summary>
 		public void Clear()
 		{
 			_state = new Tuple<string, float>("Starting...", 0.0f);
@@ -24,8 +33,15 @@ namespace EpicGames.Core
 			_ranges.Push(new Tuple<float, float>(0.0f, 1.0f));
 		}
 
+		/// <summary>
+		/// The current state
+		/// </summary>
 		public Tuple<string, float> Current => _state;
 
+		/// <summary>
+		/// Add new message
+		/// </summary>
+		/// <param name="message"></param>
 		public void Set(string message)
 		{
 			if (_ranges.Count == 1)
@@ -34,6 +50,11 @@ namespace EpicGames.Core
 			}
 		}
 
+		/// <summary>
+		/// Add a new message with a fraction of how far it has progressed
+		/// </summary>
+		/// <param name="message"></param>
+		/// <param name="fraction"></param>
 		public void Set(string message, float fraction)
 		{
 			if (_ranges.Count == 1)
@@ -46,21 +67,36 @@ namespace EpicGames.Core
 			}
 		}
 
+		/// <summary>
+		/// Update a new fraction of progress without updating the message
+		/// </summary>
+		/// <param name="fraction"></param>
 		public void Set(float fraction)
 		{
 			_state = new Tuple<string, float>(_state.Item1, RelativeToAbsoluteFraction(fraction));
 		}
 
+		/// <summary>
+		/// Increment the current progress
+		/// </summary>
+		/// <param name="fraction"></param>
 		public void Increment(float fraction)
 		{
 			Set(_state.Item2 + RelativeToAbsoluteFraction(fraction));
 		}
 
+		/// <summary>
+		/// Push a new range with a max fraction which indicates of far this scope can go
+		/// </summary>
+		/// <param name="maxFraction"></param>
 		public void Push(float maxFraction)
 		{
 			_ranges.Push(new Tuple<float, float>(_state.Item2, RelativeToAbsoluteFraction(maxFraction)));
 		}
 
+		/// <summary>
+		/// Pop the latest range as that step is now completed
+		/// </summary>
 		public void Pop()
 		{
 			if (_ranges.Count > 1)
@@ -76,10 +112,22 @@ namespace EpicGames.Core
 		}
 	}
 
+	/// <summary>
+	/// Parses child process output that contains the @progress directives to indicate their progress
+	/// </summary>
 	public static class ProgressTextReader
 	{
+		/// <summary>
+		/// The directive used to indicate progress
+		/// </summary>
 		public const string DirectivePrefix = "@progress ";
 
+		/// <summary>
+		/// Parse progress value from a text line
+		/// </summary>
+		/// <param name="line"></param>
+		/// <param name="value"></param>
+		/// <returns></returns>
 		public static string? ParseLine(string line, ProgressValue value)
 		{
 			string trimLine = line.Trim();
