Explicitly cast element size initializer to int32 to avoid issues with templated wrapper constructor for properties.

diff --git a/Engine/Source/Programs/Shared/EpicGames.UHT/Types/Properties/UhtBoolProperty.cs b/Engine/Source/Programs/Shared/EpicGames.UHT/Types/Properties/UhtBoolProperty.cs
--- a/Engine/Source/Programs/Shared/EpicGames.UHT/Types/Properties/UhtBoolProperty.cs
+++ b/Engine/Source/Programs/Shared/EpicGames.UHT/Types/Properties/UhtBoolProperty.cs
@@ -238,7 +238,7 @@ namespace EpicGames.UHT.Types
 				builder.AppendTabs(tabs).Append("}\r\n");
 			}
 			AppendConstInitMemberDefStart(builder, context, name, nameSuffix, outerFunc, IsBitfield ? "0" : offset, tabs);
-			builder.Append("sizeof(").Append(CppTypeText).Append("), ");
+			builder.Append("(int32)sizeof(").Append(CppTypeText).Append("), ");
 			builder.Append(IsNativeBool ? "true, " : "false, ");
 			if (Outer == context.OuterStruct)
 			{
