Explicitly case sizes of enums and set elements to int32 for constructor

diff --git a/Engine/Source/Programs/Shared/EpicGames.UHT/Types/Properties/UhtEnumProperty.cs b/Engine/Source/Programs/Shared/EpicGames.UHT/Types/Properties/UhtEnumProperty.cs
--- a/Engine/Source/Programs/Shared/EpicGames.UHT/Types/Properties/UhtEnumProperty.cs
+++ b/Engine/Source/Programs/Shared/EpicGames.UHT/Types/Properties/UhtEnumProperty.cs
@@ -282,7 +282,7 @@ namespace EpicGames.UHT.Types
 				}
 				else
 				{
-					builder.Append("sizeof(");
+					builder.Append("(int32)sizeof(");
 					UnderlyingProperty.AppendText(builder, UhtPropertyTextType.Generic).Append("), ");
 				}
 				AppendMemberDefRef(builder, context, Enum, Exporters.CodeGen.UhtSingletonType.ConstInit);
diff --git a/Engine/Source/Programs/Shared/EpicGames.UHT/Types/Properties/UhtSetProperty.cs b/Engine/Source/Programs/Shared/EpicGames.UHT/Types/Properties/UhtSetProperty.cs
--- a/Engine/Source/Programs/Shared/EpicGames.UHT/Types/Properties/UhtSetProperty.cs
+++ b/Engine/Source/Programs/Shared/EpicGames.UHT/Types/Properties/UhtSetProperty.cs
@@ -168,8 +168,8 @@ namespace EpicGames.UHT.Types
 			}
 			AppendConstInitMemberDefStart(builder, context, name, nameSuffix, outerFunc, offset, tabs);
 			ValueProperty.AppendConstInitMemberPtr(builder, context, name, GetNameSuffix(nameSuffix, "_ElementProp"), tabs).Append(", ");
-			builder.Append("sizeof(").AppendPropertyText(ValueProperty, UhtPropertyTextType.Generic).Append("), ");
-			builder.Append("alignof(").AppendPropertyText(ValueProperty, UhtPropertyTextType.Generic).Append("), ");
+			builder.Append("(int32)sizeof(").AppendPropertyText(ValueProperty, UhtPropertyTextType.Generic).Append("), ");
+			builder.Append("(int32)alignof(").AppendPropertyText(ValueProperty, UhtPropertyTextType.Generic).Append("), ");
 			AppendConstInitMemberDefEnd(builder, context);
 			return builder;
 		}
