Horde: Add support for grouping issues on log channel boundaries

diff --git a/Engine/Source/Programs/Horde/Plugins/Build/HordeServer.Build.Tests/Issues/IssueHandlerTests.cs b/Engine/Source/Programs/Horde/Plugins/Build/HordeServer.Build.Tests/Issues/IssueHandlerTests.cs
--- a/Engine/Source/Programs/Horde/Plugins/Build/HordeServer.Build.Tests/Issues/IssueHandlerTests.cs
+++ b/Engine/Source/Programs/Horde/Plugins/Build/HordeServer.Build.Tests/Issues/IssueHandlerTests.cs
@@ -76,7 +76,7 @@ namespace HordeServer.Tests.Issues
 				"{ \"time\":\"2025-01-01T01:20:30.999Z\",\"level\":\"Warning\",\"message\":\"LogLocTextHelper: Warning: Text conflict for namespace '' and key 'ABCDEFGHIJKLMNOPQRSTUVWXYZABCDEH'. First entry is /SamePlugin/Scripts/SubFolder/AssetName.AssetName.Sub.Message.ContextualMessages(0).ContextualMessages.Messages(0).Messages.Message:'Text in first asset.'. Conflicting entry is /SamePlugin/Scripts/SubFolder/CopiedAsset.CopiedAsset.Sub.Message.ContextualMessages(0).ContextualMessages.Messages(0).Messages.Message:'Text in second asset.'.\",\"properties\":{ \"_channel\":{ \"$type\":\"Channel\",\"$text\":\"LogLocTextHelper\"},\"_severity\":{ \"$type\":\"Severity\",\"$text\":\"Warning\"},\"from\":\"\",\"location\":\"/SamePlugin/Scripts/SubFolder/AssetName.AssetName.Sub.Message.ContextualMessages(0).ContextualMessages.Messages(0).Messages.Message\",\"locNamespace\":\"\",\"locKey\":\"ABCDEFGHIJKLMNOPQRSTUVWXYZABCDEH\",\"locKeyMetaData\":\"\",\"text\":\"Text in first asset.\",\"textMetaData\":\"\",\"conflictLocation\":\"/SamePlugin/Scripts/SubFolder/CopiedAsset.CopiedAsset.Sub.Message.ContextualMessages(0).ContextualMessages.Messages(0).Messages.Message\",\"conflictText\":\"Text in second asset.\",\"conflictTextMetaData\":\"\",\"id\":304},\"format\":\"{_channel}: {_severity}: Text conflict{from} for namespace '{locNamespace}' and key '{locKey}'{locKeyMetaData}. First entry is {location}:'{text}'{textMetaData}. Conflicting entry is {conflictLocation}:'{conflictText}'{conflictTextMetaData}.\"}", // This "location" created an exception in the code at one point
 			};
 
-			IssueHandlerContext context = new IssueHandlerContext(new StreamId("Release-XX.XX"), new TemplateId(), "Localize ProjectX", new Dictionary<string, string>());
+			IssueHandlerContext context = new IssueHandlerContext(new StreamId("Release-XX.XX"), new TemplateId(), "Localize ProjectX", new Dictionary<string, string>(), new Dictionary<string, string>());
 			LocalizationIssueHandler localizationIssueHandler = new LocalizationIssueHandler(context);
 			List<IssueEventGroup> groupedIssues = TestSingleIssueHandler(localizationIssueHandler, jsonStructuredLogs);
 
diff --git a/Engine/Source/Programs/Horde/Plugins/Build/HordeServer.Build.Tests/Issues/IssueServiceTests.cs b/Engine/Source/Programs/Horde/Plugins/Build/HordeServer.Build.Tests/Issues/IssueServiceTests.cs
--- a/Engine/Source/Programs/Horde/Plugins/Build/HordeServer.Build.Tests/Issues/IssueServiceTests.cs
+++ b/Engine/Source/Programs/Horde/Plugins/Build/HordeServer.Build.Tests/Issues/IssueServiceTests.cs
@@ -163,6 +163,10 @@ namespace HordeServer.Tests.Issues
 
 			BuildConfig buildConfig = new BuildConfig();
 			buildConfig.Projects.Add(projectConfig);
+			buildConfig.Issues = new List<IssueConfig> { new IssueConfig { ChannelGroups = new List<IssueChannelGroupConfig> { 
+			new IssueChannelGroupConfig { Channels = new List<string> { "AssetCheck" } },
+			new IssueChannelGroupConfig { Channels = new List<string> { "LogLinker", "LogBlueprint" } },
+			} } };
 
 			GlobalConfig globalConfig = new GlobalConfig();
 			globalConfig.Plugins.AddBuildConfig(buildConfig);
@@ -2245,6 +2249,66 @@ namespace HordeServer.Tests.Issues
 			}
 		}
 
+		// Parse structured logs to create IssueEvents (and LogEvents for debug purposes)
+		static List<JsonLogEvent> ParseJSONStructuredLogs(string[] jsonStructuredLogs)
+		{
+			List<JsonLogEvent> jsonEvents = new List<JsonLogEvent>();
+			foreach (string jsonStructuredLog in jsonStructuredLogs)
+			{
+				byte[] bytes = Encoding.UTF8.GetBytes(jsonStructuredLog);
+				JsonLogEvent jsonEvent;
+				Assert.IsTrue(JsonLogEvent.TryParse(bytes, out jsonEvent));
+				jsonEvents.Add(jsonEvent);
+			}
+
+			return jsonEvents;
+		}
+
+		[TestMethod]
+		public async Task ChannelGroupingTestAsync()
+		{
+			string[] jsonStructuredLogs =
+			{
+				@"{""time"":""2025-10-15T18:10:10.409Z"",""level"":""Error"",""message"":""AssetCheck: Error: This is an error"",""properties"":{ ""_channel"":{ ""$type"":""Channel"",""$text"":""AssetCheck""},""_severity"":{ ""$type"":""Severity"",""$text"":""Error""} },""format"":""{_channel}: {_severity}: This is an error""}",
+				@"{""time"":""2025-10-15T18:10:11.409Z"",""level"":""Error"",""message"":""LogBlueprint: Error: This is a blueprint error 1"",""properties"":{ ""_channel"":{ ""$type"":""Channel"",""$text"":""LogBlueprint""},""_severity"":{ ""$type"":""Severity"",""$text"":""Error""} },""format"":""{_channel}: {_severity}: This is a blueprint error 1""}",
+				@"{""time"":""2025-10-15T18:10:12.409Z"",""level"":""Error"",""message"":""LogBlueprint: Error: This is a blueprint error 2"",""properties"":{ ""_channel"":{ ""$type"":""Channel"",""$text"":""LogBlueprint""},""_severity"":{ ""$type"":""Severity"",""$text"":""Error""} },""format"":""{_channel}: {_severity}: This is a blueprint error 2""}",
+				@"{""time"":""2025-10-15T18:10:13.409Z"",""level"":""Warning"",""message"":""LogLinker: Error: This is a linker error"",""properties"":{ ""_channel"":{ ""$type"":""Channel"",""$text"":""LogBlueprint""},""_severity"":{ ""$type"":""Severity"",""$text"":""Warning""} },""format"":""{_channel}: {_severity}: This is a linker error""}",
+				@"{""time"":""2025-10-15T18:10:13.409Z"",""level"":""Error"",""message"":""LogOther: Error: This is another error"",""properties"":{ ""_channel"":{ ""$type"":""Channel"",""$text"":""LogOther""},""_severity"":{ ""$type"":""Severity"",""$text"":""Error""} },""format"":""{_channel}: {_severity}: This is another error""}",
+				@"{""time"":""2025-10-15T18:10:13.409Z"",""level"":""Error"",""message"":""NoChannel: Error: This is an error with no channel"",""properties"":{ ""_severity"":{ ""$type"":""Severity"",""$text"":""Error""} },""format"":""{_channel}: {_severity}: This is an error with no channel""}"
+			};
+			//LogBlueprint
+			List<JsonLogEvent> logEvents = ParseJSONStructuredLogs(jsonStructuredLogs);
+
+			// #2
+			// Scenario: Log Channel Grouping
+			// Expected: Issues are generated upon log channel boundaries
+			{
+				IJob job = CreateJob(MainStreamId, 120, "Test Build", Graph);
+				await using (TestJsonLogger logger = await CreateLoggerAsync(job, 0, 0))
+				{
+					foreach (JsonLogEvent outJsonLogEvent in logEvents)
+					{
+						logger.LogJsonLogEvent(outJsonLogEvent);
+					}
+				}
+
+				await UpdateCompleteStepAsync(job, 0, 0, JobStepOutcome.Failure);
+
+				IReadOnlyList<IIssue> issues = await IssueCollection.FindIssuesAsync();
+				issues = issues.OrderBy(x => x.Summary).ToList();
+				Assert.AreEqual(3, issues.Count);
+
+				Assert.AreEqual("AssetCheck errors in Update Version Files", issues[0].Summary);
+				Assert.AreEqual("Hashed_AssetCheck", issues[0].Fingerprints[0].Type);
+
+				Assert.AreEqual("LogBlueprint errors in Update Version Files", issues[1].Summary);
+				Assert.AreEqual("Hashed_LogBlueprint_LogLinker", issues[1].Fingerprints[0].Type);
+
+				Assert.AreEqual("LogOther errors in Update Version Files", issues[2].Summary);
+				Assert.AreEqual("Hashed", issues[2].Fingerprints[0].Type);
+			}
+		}
+
 		[TestMethod]
 		public async Task GauntletTestAsync()
 		{
diff --git a/Engine/Source/Programs/Horde/Plugins/Build/HordeServer.Build/BuildConfig.cs b/Engine/Source/Programs/Horde/Plugins/Build/HordeServer.Build/BuildConfig.cs
--- a/Engine/Source/Programs/Horde/Plugins/Build/HordeServer.Build/BuildConfig.cs
+++ b/Engine/Source/Programs/Horde/Plugins/Build/HordeServer.Build/BuildConfig.cs
@@ -21,11 +21,12 @@ using HordeServer.Artifacts;
 using HordeServer.Commits;
 using HordeServer.Configuration;
 using HordeServer.Devices;
-using HordeServer.VersionControl.Perforce;
+using HordeServer.Issues;
 using HordeServer.Plugins;
 using HordeServer.Projects;
 using HordeServer.Streams;
 using HordeServer.Utilities;
+using HordeServer.VersionControl.Perforce;
 
 #pragma warning disable CA2227 // Change x to be read-only by removing the property setter
 
@@ -90,6 +91,16 @@ namespace HordeServer
 		public string IssueFixedTag { get; set; } = "#horde";
 
 		/// <summary>
+		/// List of build health issue configurations
+		/// </summary>
+		public List<IssueConfig> Issues { get; set; } = new List<IssueConfig>();
+
+		/// <summary>
+		/// Issue channel groupings
+		/// </summary>
+		[JsonIgnore]
+		public IReadOnlyDictionary<string, string> IssueChannelGroups { get; private set; } = null!;
+
 		/// The maximum days since node success to scan for suspects
 		/// </summary>
 		public int IssueFindSuspectsMaxSuccessDays { get; set; } = 14;
@@ -108,6 +119,35 @@ namespace HordeServer
 			Acl = configOptions.ParentAcl;
 			Streams = Projects.SelectMany(x => x.Streams).ToList();
 
+			Dictionary<string, string> issueChannelGroups = new Dictionary<string, string>();
+			IssueChannelGroups = issueChannelGroups;
+
+			// build up the issue channel groupings
+			Issues.ForEach(ic =>
+			{
+				ic.PostLoad();
+
+				ic.ChannelGroups?.ForEach(group =>
+				{
+					if (group.Channels == null || group.Channels.Count == 0)
+					{
+						return;
+					}
+
+					string channelGroup = String.Join("_", group.Channels);
+
+					group.Channels.ForEach(channel =>
+					{
+						if (String.IsNullOrEmpty(channel) || issueChannelGroups.ContainsKey(channel))
+						{
+							return;
+						}
+
+						issueChannelGroups[channel] = channelGroup;
+					});
+				});
+			});
+
 			_projectLookup.Clear();
 			_streamLookup.Clear();
 			foreach (ProjectConfig project in Projects)
diff --git a/Engine/Source/Programs/Horde/Plugins/Build/HordeServer.Build/Issues/IssueConfig.cs b/Engine/Source/Programs/Horde/Plugins/Build/HordeServer.Build/Issues/IssueConfig.cs
new file mode 100644
--- /dev/null
+++ b/Engine/Source/Programs/Horde/Plugins/Build/HordeServer.Build/Issues/IssueConfig.cs
@@ -0,0 +1,47 @@
+// Copyright Epic Games, Inc. All Rights Reserved.
+
+using EpicGames.Core;
+using HordeServer.Configuration;
+using HordeServer.Utilities;
+
+namespace HordeServer.Issues
+{
+
+	/// <summary>
+	/// Configuration for grouping to apply based on logging channels
+	/// </summary>
+	public class IssueChannelGroupConfig
+	{
+		/// <summary>
+		/// A list of channel logging boundaries to group issues on
+		/// </summary>
+		public List<string>? Channels { get; set; }
+	}
+
+	/// <summary>
+	/// Stores configuration for build health issues
+	/// </summary>
+	[JsonSchema("https://unrealengine.com/horde/issues")]
+	[JsonSchemaCatalog("Horde Issue", "Horde issue configuration file", new[] { "*.issues.json", "Issues/*.json" })]
+	[ConfigDoc("*.issues.json", "[Horde](../../../README.md) > [Configuration](../../Config.md)", "Config/Schema/Issues.md")]
+	[ConfigIncludeRoot]
+	[ConfigMacroScope]
+	public class IssueConfig
+	{
+		/// <summary>
+		/// Issue grouping channels
+		/// </summary>
+		public List<IssueChannelGroupConfig>? ChannelGroups { get; set; }
+
+		/// <summary>
+		/// Process Issue Config
+		/// </summary>
+		public void PostLoad()
+		{
+			ChannelGroups?.ForEach(group =>
+			{
+				group.Channels?.Sort();
+			});
+		}
+	}
+}
diff --git a/Engine/Source/Programs/Horde/Plugins/Build/HordeServer.Build/Issues/IssueService.cs b/Engine/Source/Programs/Horde/Plugins/Build/HordeServer.Build/Issues/IssueService.cs
--- a/Engine/Source/Programs/Horde/Plugins/Build/HordeServer.Build/Issues/IssueService.cs
+++ b/Engine/Source/Programs/Horde/Plugins/Build/HordeServer.Build/Issues/IssueService.cs
@@ -747,10 +747,12 @@ namespace HordeServer.Issues
 			{
 				throw new ArgumentException($"Unable to retrieve log {step.LogId}");
 			}
+			
+			BuildConfig buildConfig = _buildConfig.CurrentValue;
 
 			// Create the DI container for issue handlers
 			ServiceCollection services = new ServiceCollection();
-			services.AddSingleton(new IssueHandlerContext(job.StreamId, job.TemplateId, node.Name, node.Annotations));
+			services.AddSingleton(new IssueHandlerContext(job.StreamId, job.TemplateId, node.Name, buildConfig.IssueChannelGroups, node.Annotations));
 
 			using ServiceProvider serviceProvider = services.BuildServiceProvider();
 
diff --git a/Engine/Source/Programs/Shared/EpicGames.Horde/Issues/Handlers/HashedIssueHandler.cs b/Engine/Source/Programs/Shared/EpicGames.Horde/Issues/Handlers/HashedIssueHandler.cs
--- a/Engine/Source/Programs/Shared/EpicGames.Horde/Issues/Handlers/HashedIssueHandler.cs
+++ b/Engine/Source/Programs/Shared/EpicGames.Horde/Issues/Handlers/HashedIssueHandler.cs
@@ -1,5 +1,6 @@
 // Copyright Epic Games, Inc. All Rights Reserved.
 
+using System;
 using System.Collections.Generic;
 using System.Linq;
 using System.Text;
@@ -49,71 +50,167 @@ namespace EpicGames.Horde.Issues.Handlers
 			return true;
 		}
 
+		/// <summary>
+		/// Get the channel for issue event
+		/// </summary>
+		/// <returns></returns>
+		static string? GetLogChannel(IssueEvent logEvent)
+		{
+			try
+			{
+
+				string? channel = null;
+
+				foreach (JsonLogEvent line in logEvent.Lines)
+				{
+					using JsonDocument document = JsonDocument.Parse(line.Data);
+
+					string? channelType;
+					string? channelText;
+					if (document.RootElement.TryGetNestedProperty("properties._channel.$type", out channelType) && document.RootElement.TryGetNestedProperty("properties._channel.$text", out channelText))
+					{
+						if (channelType != "Channel")
+						{
+							continue;
+						}
+
+						channel = channelText;
+
+						if (!String.IsNullOrEmpty(channel))
+						{
+							break;
+						}
+					}
+				}
+
+				return channel;
+
+			} catch { 
+				return null; 
+			}
+		}
+
 		/// <inheritdoc/>
 		public override IEnumerable<IssueEventGroup> GetIssues()
 		{
 			List<IssueEventGroup> issues = new List<IssueEventGroup>();
 
-			IssueEventGroup? genericFingerprint = null;
-			IssueEventGroup? genericErrorsFingerprint = null;
-			HashSet<Md5Hash> hashes = new HashSet<Md5Hash>();
-
 			// keep hash consistent when only have general, non-unique events
 			bool allGeneral = _issueEvents.FirstOrDefault(stepEvent => stepEvent.EventId == null || !IsGeneralEventId(stepEvent.EventId.Value)) == null;
 
+			Dictionary<string, List<IssueEvent>> fingerEvents = new Dictionary<string, List<IssueEvent>>();
+			Dictionary<IssueEvent, string?> eventScopes = new Dictionary<IssueEvent, string?>();
+
 			foreach (IssueEvent stepEvent in _issueEvents)
 			{
-				string hashSource = stepEvent.Render();
-
-				if (!allGeneral && stepEvent.EventId != null)
+				string? channelGroup = null;
+				if (_context.ChannelGroups.Count > 0)
 				{
-					// If the event is general, salt the hash with the stream id, template, otherwise it will be aggressively matched.
-					// Consider salting with node name, though template id should be enough and have better grouping
-					if (IsGeneralEventId(stepEvent.EventId.Value))
+					string? channel = GetLogChannel(stepEvent);
+					eventScopes.Add(stepEvent, channel);
+					if (!String.IsNullOrEmpty(channel))
 					{
-						hashSource += $"step:{_context.StreamId}:{_context.TemplateId}";
+						_context.ChannelGroups.TryGetValue(channel, out channelGroup);
 					}
 				}
 
-				if (hashes.Count < 25 && TryGetHash(hashSource, out Md5Hash hash))
+				string fingerprint = "Hashed";
+
+				if (!String.IsNullOrEmpty(channelGroup))
 				{
-					hashes.Add(hash);
-
-					IssueEventGroup issue = new IssueEventGroup("Hashed", "{Severity} in {Meta:Node}", IssueChangeFilter.All);
-					issue.Events.Add(stepEvent);
-					issue.Keys.AddHash(hash);
-					stepEvent.AuditLogger?.LogDebug("Fingerprint key: '{Key}' generated from event: '{Event}'", hash, stepEvent.Render());
-					AddFileMetadata(issue, stepEvent);
-					issue.Metadata.Add("Node", _context.NodeName);
-					issues.Add(issue);
+					fingerprint += $"_{channelGroup}";
 				}
-				else
+
+				List<IssueEvent>? events = null;
+				if (!fingerEvents.TryGetValue(fingerprint, out events))
 				{
-					if (stepEvent.Severity == LogLevel.Error || stepEvent.Severity == LogLevel.Critical)
+					events = new List<IssueEvent>();
+					fingerEvents.Add(fingerprint, events);
+				}
+
+				events.Add(stepEvent);
+			}
+
+			foreach (KeyValuePair<string, List<IssueEvent>> item in fingerEvents)
+			{
+				string fingerprintType = item.Key;
+
+				IssueEventGroup? genericFingerprint = null;
+				IssueEventGroup? genericErrorsFingerprint = null;
+				HashSet<Md5Hash> hashes = new HashSet<Md5Hash>();
+
+				foreach (IssueEvent stepEvent in item.Value)
+				{
+					string hashSource = stepEvent.Render();
+					string? channel = null;
+					eventScopes.TryGetValue(stepEvent, out channel);
+
+					if (!allGeneral && stepEvent.EventId != null)
+					{
+						// If the event is general, salt the hash with the stream id, template, otherwise it will be aggressively matched.
+						// Consider salting with node name, though template id should be enough and have better grouping
+						if (IsGeneralEventId(stepEvent.EventId.Value))
+						{
+							hashSource += $"step:{_context.StreamId}:{_context.TemplateId}";
+						}
+					}
+
+					string summaryTemplate = channel == null ? "{Severity} in {Meta:Node}" : "{Meta:Channel} {Severity} in {Meta:Node}";
+
+					if (hashes.Count < 25 && TryGetHash(hashSource, out Md5Hash hash))
 					{
-						if (genericErrorsFingerprint == null)
+						hashes.Add(hash);
+
+						IssueEventGroup issue = new IssueEventGroup(fingerprintType, summaryTemplate, IssueChangeFilter.All);
+						issue.Events.Add(stepEvent);
+						issue.Keys.AddHash(hash);
+						stepEvent.AuditLogger?.LogDebug("Fingerprint key: '{Key}' generated from event: '{Event}'", hash, stepEvent.Render());
+						AddFileMetadata(issue, stepEvent);
+						issue.Metadata.Add("Node", _context.NodeName);						
+						if (!String.IsNullOrEmpty(channel))
 						{
-							genericErrorsFingerprint = new IssueEventGroup("Hashed", "{Severity} in {Meta:Node}", IssueChangeFilter.All);
-							genericErrorsFingerprint.Keys.Add(IssueKey.FromStepAndSeverity(_context.StreamId, _context.TemplateId, _context.NodeName, LogLevel.Error));
-							stepEvent.AuditLogger?.LogDebug("Fingerprint key: '{Key}' generated from event: '{Event}'", IssueKey.FromStepAndSeverity(_context.StreamId, _context.TemplateId, _context.NodeName, LogLevel.Error), stepEvent.Render());
-							genericErrorsFingerprint.Metadata.Add("Node", _context.NodeName);
-							issues.Add(genericErrorsFingerprint);
+							issue.Metadata.Add("Channel", channel);
 						}
-						genericErrorsFingerprint.Events.Add(stepEvent);
-						AddFileMetadata(genericErrorsFingerprint, stepEvent);
+						issues.Add(issue);
 					}
 					else
 					{
-						if (genericFingerprint == null)
+						if (stepEvent.Severity == LogLevel.Error || stepEvent.Severity == LogLevel.Critical)
 						{
-							genericFingerprint = new IssueEventGroup("Hashed", "{Severity} in {Meta:Node}", IssueChangeFilter.All);
-							genericFingerprint.Keys.Add(IssueKey.FromStep(_context.StreamId, _context.TemplateId, _context.NodeName));
-							stepEvent.AuditLogger?.LogDebug("Fingerprint key: '{Key}' generated from event: '{Event}'", IssueKey.FromStep(_context.StreamId, _context.TemplateId, _context.NodeName), stepEvent.Render());
-							genericFingerprint.Metadata.Add("Node", _context.NodeName);
-							issues.Add(genericFingerprint);
+							if (genericErrorsFingerprint == null)
+							{
+								genericErrorsFingerprint = new IssueEventGroup(fingerprintType, summaryTemplate, IssueChangeFilter.All);
+								genericErrorsFingerprint.Keys.Add(IssueKey.FromStepAndSeverity(_context.StreamId, _context.TemplateId, _context.NodeName, LogLevel.Error));
+								stepEvent.AuditLogger?.LogDebug("Fingerprint key: '{Key}' generated from event: '{Event}'", IssueKey.FromStepAndSeverity(_context.StreamId, _context.TemplateId, _context.NodeName, LogLevel.Error), stepEvent.Render());
+								genericErrorsFingerprint.Metadata.Add("Node", _context.NodeName);
+								if (!String.IsNullOrEmpty(channel))
+								{
+									genericErrorsFingerprint.Metadata.Add("Channel", channel);
+								}
+
+								issues.Add(genericErrorsFingerprint);
+							}
+							genericErrorsFingerprint.Events.Add(stepEvent);
+							AddFileMetadata(genericErrorsFingerprint, stepEvent);
+						}
+						else
+						{
+							if (genericFingerprint == null)
+							{
+								genericFingerprint = new IssueEventGroup(fingerprintType, summaryTemplate, IssueChangeFilter.All);
+								genericFingerprint.Keys.Add(IssueKey.FromStep(_context.StreamId, _context.TemplateId, _context.NodeName));
+								stepEvent.AuditLogger?.LogDebug("Fingerprint key: '{Key}' generated from event: '{Event}'", IssueKey.FromStep(_context.StreamId, _context.TemplateId, _context.NodeName), stepEvent.Render());
+								genericFingerprint.Metadata.Add("Node", _context.NodeName);
+								if (!String.IsNullOrEmpty(channel))
+								{
+									genericFingerprint.Metadata.Add("Channel", channel);
+								}
+
+								issues.Add(genericFingerprint);
+							}
+							genericFingerprint.Events.Add(stepEvent);
+							AddFileMetadata(genericFingerprint, stepEvent);
 						}
-						genericFingerprint.Events.Add(stepEvent);
-						AddFileMetadata(genericFingerprint, stepEvent);
 					}
 				}
 			}
diff --git a/Engine/Source/Programs/Shared/EpicGames.Horde/Issues/IssueHandler.cs b/Engine/Source/Programs/Shared/EpicGames.Horde/Issues/IssueHandler.cs
--- a/Engine/Source/Programs/Shared/EpicGames.Horde/Issues/IssueHandler.cs
+++ b/Engine/Source/Programs/Shared/EpicGames.Horde/Issues/IssueHandler.cs
@@ -63,6 +63,11 @@ namespace EpicGames.Horde.Issues
 		/// </summary>
 		public string NodeName { get; }
 
+		/// <summary>
+		/// Log Channel => Channel Grouping
+		/// </summary>
+		public IReadOnlyDictionary<string, string> ChannelGroups { get; }
+
 		/// <summary>
 		/// Annotations for this node
 		/// </summary>
@@ -71,12 +76,13 @@ namespace EpicGames.Horde.Issues
 		/// <summary>
 		/// Constructor
 		/// </summary>
-		public IssueHandlerContext(StreamId streamId, TemplateId templateId, string nodeName, IReadOnlyDictionary<string, string> nodeAnnotations)
+		public IssueHandlerContext(StreamId streamId, TemplateId templateId, string nodeName, IReadOnlyDictionary<string, string> channelGroups, IReadOnlyDictionary<string, string> nodeAnnotations)
 		{
 			StreamId = streamId;
 			TemplateId = templateId;
 			NodeName = nodeName;
 			NodeAnnotations = nodeAnnotations;
+			ChannelGroups = channelGroups;
 		}
 	}
 }
