Fix recipe id initialization in test session response

#rnx

diff --git a/Engine/Source/Programs/Shared/EpicGames.Horde/Jobs/TestData/TestDataMessages.cs b/Engine/Source/Programs/Shared/EpicGames.Horde/Jobs/TestData/TestDataMessages.cs
--- a/Engine/Source/Programs/Shared/EpicGames.Horde/Jobs/TestData/TestDataMessages.cs
+++ b/Engine/Source/Programs/Shared/EpicGames.Horde/Jobs/TestData/TestDataMessages.cs
@@ -864,7 +864,7 @@ namespace EpicGames.Horde.Jobs.TestData
 		/// <summary>
 		/// The recipe id associated with the test session
 		/// </summary>
-		public string? RecipeId { get; set; } = String.Empty;
+		public string? RecipeId { get; set; }
 
 		/// <summary>
 		/// The outcome of the test session
