fix property specifier not removed from name when parsing property values

#rnx

diff --git a/Engine/Source/Programs/Shared/EpicGames.Core/MessageTemplate.cs b/Engine/Source/Programs/Shared/EpicGames.Core/MessageTemplate.cs
--- a/Engine/Source/Programs/Shared/EpicGames.Core/MessageTemplate.cs
+++ b/Engine/Source/Programs/Shared/EpicGames.Core/MessageTemplate.cs
@@ -229,6 +229,13 @@ namespace EpicGames.Core
 				{
 					string name = format.Substring(offsets[idx].Item1, offsets[idx].Item2);
 
+					// Parse for format specifier
+					int specifierIdx = name.IndexOf(':', 0);
+					if (specifierIdx != -1)
+					{
+						name = name.Substring(0, specifierIdx);
+					}
+
 					int number;
 					if (Int32.TryParse(name, out number))
 					{
