EpicGames.Perforce: Add support for modtime client option

- Maintains the original Perforce timestamp when moved to/from the cache
- Requires opting in via "method" setting for the workspace type, like how use of have table works

diff --git a/Engine/Source/Programs/Horde/Drivers/JobDriver.Tests/WorkspaceInfoTest.cs b/Engine/Source/Programs/Horde/Drivers/JobDriver.Tests/WorkspaceInfoTest.cs
--- a/Engine/Source/Programs/Horde/Drivers/JobDriver.Tests/WorkspaceInfoTest.cs
+++ b/Engine/Source/Programs/Horde/Drivers/JobDriver.Tests/WorkspaceInfoTest.cs
@@ -40,6 +40,7 @@ namespace JobDriver.Tests
 			Assert.AreEqual(new MWO { UseHaveTable = false }, GetMwOptions("name=managedWorkspace&useHaveTable=FalSe"));
 			Assert.AreEqual(new MWO { UseHaveTable = true }, GetMwOptions("name=managedWorkspace&useHaveTable=true"));
 			Assert.AreEqual(new MWO { UseHaveTable = true }, GetMwOptions("name=managedWorkspace&useHaveTable=TrUE"));
+			Assert.AreEqual(new MWO { UseModTime = true }, GetMwOptions("name=managedWorkspace&UseModTime=true"));
 		}
 
 		static ManagedWorkspaceOptions GetMwOptions(string? method, long? minScratchSpace = null)
diff --git a/Engine/Source/Programs/Horde/Drivers/JobDriver/Utility/WorkspaceInfo.cs b/Engine/Source/Programs/Horde/Drivers/JobDriver/Utility/WorkspaceInfo.cs
--- a/Engine/Source/Programs/Horde/Drivers/JobDriver/Utility/WorkspaceInfo.cs
+++ b/Engine/Source/Programs/Horde/Drivers/JobDriver/Utility/WorkspaceInfo.cs
@@ -281,6 +281,7 @@ namespace JobDriver.Utility
 			const string MaxFileConcurrencyKey = "maxFileConcurrency";
 			const string MinScratchSpaceKey = "minScratchSpace";
 			const string UseHaveTableKey = "useHaveTable";
+			const string UseModTimeKey = "useModTime";
 			const string PreferNativeClientKey = "preferNativeClient";
 
 			ManagedWorkspaceOptions options = new ManagedWorkspaceOptions();
@@ -317,6 +318,10 @@ namespace JobDriver.Utility
 					{
 						options = options with { PreferNativeClient = true };
 					}
+					if (String.Equals(nameValues[UseModTimeKey], "true", StringComparison.OrdinalIgnoreCase))
+					{
+						options = options with { UseModTime = true };
+					}
 				}
 			}
 
diff --git a/Engine/Source/Programs/Shared/EpicGames.Perforce.Managed.Tests/ManagedWorkspaceTest.cs b/Engine/Source/Programs/Shared/EpicGames.Perforce.Managed.Tests/ManagedWorkspaceTest.cs
--- a/Engine/Source/Programs/Shared/EpicGames.Perforce.Managed.Tests/ManagedWorkspaceTest.cs
+++ b/Engine/Source/Programs/Shared/EpicGames.Perforce.Managed.Tests/ManagedWorkspaceTest.cs
@@ -244,9 +244,9 @@ public class ManagedWorkspaceTest : BasePerforceFixtureTest
 	[DataRow(false, DisplayName = "Without have-table")]
 	public async Task UnshelveAsync(bool useHaveTable)
 	{
-		ManagedWorkspace ws = await CreateManagedWorkspaceAsync(useHaveTable);
-		await SyncAsync(ws, 7);
-		await ws.UnshelveAsync(PerforceConnection, 8, CancellationToken.None);
+		ManagedWorkspace mw = await CreateManagedWorkspaceAsync(useHaveTable);
+		await SyncAsync(mw, 7);
+		await mw.UnshelveAsync(PerforceConnection, 8, CancellationToken.None);
 		Stream.GetChangelist(8).AssertDepotFiles(SyncDir);
 
 		// Have-table still correspond to CL 7 as CL 8 is shelved, only p4 printed to workspace
@@ -285,6 +285,57 @@ public class ManagedWorkspaceTest : BasePerforceFixtureTest
 		// Assert - cache becomes empty
 		PerforceFixture.AssertCacheEquals(CacheDir, Array.Empty<string>());
 	}
+	
+	[TestMethod]
+	[DataRow(true, DisplayName = "With have-table")]
+	[DataRow(false, DisplayName = "Without have-table")]
+	public async Task ModTime_NormalSync_PreservesDepotTimestamp_Async(bool useHaveTable)
+	{
+		// Arrange
+		ManagedWorkspace ws = await CreateManagedWorkspaceAsync(useHaveTable, true);
+		
+		// Act
+		await SyncAsync(ws, 7);
+		
+		// Assert
+		ChangelistFixture cl = Stream.GetChangelist(7);
+		DepotFileFixture file = cl.GetFile("//Foo/Main/main.cpp")!;
+		Assert.IsNotNull(file);
+		Assert.AreEqual(await GetDepotFileModTimeAsync(file), GetLocalFileModTime(file));
+	}
+	
+	[TestMethod]
+	[DataRow(true, DisplayName = "With have-table")]
+	[DataRow(false, DisplayName = "Without have-table")]
+	public async Task ModTime_CachedSync_PreservesDepotTimestamp_Async(bool useHaveTable)
+	{
+		// Arrange
+		ManagedWorkspace ws = await CreateManagedWorkspaceAsync(useHaveTable, true);
+		await SyncAsync(ws, 7);
+		await SyncAsync(ws, 1);
+		
+		// Act
+		await SyncAsync(ws, 7);
+		
+		// Assert
+		ChangelistFixture cl = Stream.GetChangelist(7);
+		DepotFileFixture file = cl.GetFile("//Foo/Main/main.cpp")!;
+		Assert.IsNotNull(file);
+		Assert.AreEqual(await GetDepotFileModTimeAsync(file), GetLocalFileModTime(file));
+	}
+
+	private DateTime GetLocalFileModTime(DepotFileFixture file)
+	{
+		FileInfo fi = new (Path.Join(SyncDir, file.ClientFile));
+		return fi.LastWriteTimeUtc;
+	}
+	
+	private async Task<DateTime> GetDepotFileModTimeAsync(DepotFileFixture file)
+	{
+		List<FStatRecord> fstatRecords = await PerforceConnection.FStatAsync($"{file.DepotFile}#{file.Revision}").ToListAsync();
+		Assert.AreEqual(1, fstatRecords.Count);
+		return fstatRecords[0].HeadModTime;
+	}
 
 	[TestMethod]
 	public async Task FixtureMetadataMatchesServerMetadataAsync()
@@ -350,9 +401,9 @@ public class ManagedWorkspaceTest : BasePerforceFixtureTest
 		cl.AssertDepotFiles(SyncDir);
 	}
 
-	private async Task<ManagedWorkspace> CreateManagedWorkspaceAsync(bool useHaveTable)
+	private async Task<ManagedWorkspace> CreateManagedWorkspaceAsync(bool useHaveTable, bool useModTime = false)
 	{
-		ManagedWorkspaceOptions options = new() { UseHaveTable = useHaveTable };
+		ManagedWorkspaceOptions options = new() { UseHaveTable = useHaveTable, UseModTime = useModTime};
 		ManagedWorkspace ws = await ManagedWorkspace.CreateAsync(Environment.MachineName, TempDir, options, TracerProvider.Default.GetTracer("Test"), _mwLogger, CancellationToken.None);
 		await ws.SetupAsync(PerforceConnection, StreamName, CancellationToken.None);
 		return ws;
diff --git a/Engine/Source/Programs/Shared/EpicGames.Perforce.Managed/ManagedWorkspace.cs b/Engine/Source/Programs/Shared/EpicGames.Perforce.Managed/ManagedWorkspace.cs
--- a/Engine/Source/Programs/Shared/EpicGames.Perforce.Managed/ManagedWorkspace.cs
+++ b/Engine/Source/Programs/Shared/EpicGames.Perforce.Managed/ManagedWorkspace.cs
@@ -94,6 +94,7 @@ namespace EpicGames.Perforce.Managed
 	///	</param>
 	/// <param name="Partitioned">Whether to allow using partitioned workspaces</param>
 	/// <param name="PreferNativeClient">Whether to prefer the native p4 client</param>
+	/// <param name="UseModTime">Whether to set a synced file's timestamp to match P4 stream/depot timestamp ("modtime" client option)</param>
 	public record ManagedWorkspaceOptions
 	(
 		int NumParallelSyncThreads,
@@ -101,18 +102,20 @@ namespace EpicGames.Perforce.Managed
 		long MinScratchSpace,
 		bool UseHaveTable,
 		bool Partitioned,
-		bool PreferNativeClient
+		bool PreferNativeClient,
+		bool UseModTime
 	)
 	{
-		/// <inheritdoc cref="ManagedWorkspaceOptions(Int32, Int32, Int64, Boolean, Boolean, Boolean)"/>
-		public ManagedWorkspaceOptions(int? numParallelSyncThreads = null, int? maxFileConcurrency = null, long minScratchSpace = 50L * 1024, bool useHaveTable = true, bool partitioned = false, bool preferNativeClient = false)
+		/// <inheritdoc cref="ManagedWorkspaceOptions(Int32, Int32, Int64, Boolean, Boolean, Boolean, Boolean)"/>
+		public ManagedWorkspaceOptions(int? numParallelSyncThreads = null, int? maxFileConcurrency = null, long minScratchSpace = 50L * 1024, bool useHaveTable = true, bool partitioned = false, bool preferNativeClient = false, bool useModTime = false)
 			: this(
 				  numParallelSyncThreads ?? GetDefaultThreadCount(4),
 				  maxFileConcurrency ?? GetDefaultThreadCount(4),
 				  minScratchSpace,
 				  useHaveTable,
 				  partitioned,
-				  preferNativeClient)
+				  preferNativeClient,
+				  useModTime)
 		{ }
 
 		static int GetDefaultThreadCount(int defaultValue)
@@ -1341,6 +1344,11 @@ namespace EpicGames.Perforce.Managed
 			client.Host = _hostName;
 			client.Stream = streamName;
 
+			if (_options.UseModTime)
+			{
+				client.Options = ClientOptions.ModTime;	
+			}
+
 			if (_options.Partitioned)
 			{
 				// Partitioned and read-only types store their have table separately on the server, compared to normal (writeable) clients
@@ -2359,8 +2367,12 @@ namespace EpicGames.Perforce.Managed
 				FileReference.Move(fileToMove._trackedFile.GetLocation(), targetFile);
 				try
 				{
-					FileReference.SetLastWriteTimeUtc(targetFile, DateTime.UtcNow);
-					fileToMove._workspaceFile.UpdateMetadata();
+					DateTime lastWriteTime = _options.UseModTime
+						? new DateTime(fileToMove._workspaceFile._lastModifiedTicks, DateTimeKind.Utc)
+						: DateTime.UtcNow;
+					
+					FileReference.SetLastWriteTimeUtc(targetFile, lastWriteTime);
+					fileToMove._workspaceFile.UpdateMetadata();	
 				}
 				catch
 				{
