Don't declare destructors virtual in final classes

diff --git a/Engine/Source/Programs/Shared/EpicGames.UHT/Exporters/CodeGen/UhtHeaderCodeGeneratorHFile.cs b/Engine/Source/Programs/Shared/EpicGames.UHT/Exporters/CodeGen/UhtHeaderCodeGeneratorHFile.cs
--- a/Engine/Source/Programs/Shared/EpicGames.UHT/Exporters/CodeGen/UhtHeaderCodeGeneratorHFile.cs
+++ b/Engine/Source/Programs/Shared/EpicGames.UHT/Exporters/CodeGen/UhtHeaderCodeGeneratorHFile.cs
@@ -1495,13 +1495,27 @@ namespace EpicGames.UHT.Exporters.CodeGen
 		{
 			if (!classObj.ClassExportFlags.HasAnyFlags(UhtClassExportFlags.HasDestructor))
 			{
+				builder.Append('\t');
+
+				if (!classObj.ClassFlags.HasAnyFlags(EClassFlags.Interface))
+				{
+					builder.Append(api);
+				}
+
+				if (!classObj.ClassExportFlags.HasAnyFlags(UhtClassExportFlags.Final))
+				{
+					builder.Append("virtual ");
+				}
+
+				builder.Append("~").Append(classObj.SourceName);
+
 				if (classObj.ClassFlags.HasAnyFlags(EClassFlags.Interface))
 				{
-					builder.Append('\t').Append("virtual ~").Append(classObj.SourceName).Append("() = default; \\\r\n");	
+					builder.Append("() = default; \\\r\n");	
 				}
 				else
 				{
-					builder.Append('\t').Append(api).Append("virtual ~").Append(classObj.SourceName).Append("(); \\\r\n");	
+					builder.Append("(); \\\r\n");	
 				}
 			}
 			return builder;
