Fix warning CA1834: Use 'StringBuilder.Append(char)' instead of 'StringBuilder.Append(string)' when the input is a constant unit string
#rnx

diff --git a/Engine/Source/Programs/Shared/EpicGames.UHT/Exporters/CodeGen/UhtHeaderCodeGeneratorHFile.cs b/Engine/Source/Programs/Shared/EpicGames.UHT/Exporters/CodeGen/UhtHeaderCodeGeneratorHFile.cs
--- a/Engine/Source/Programs/Shared/EpicGames.UHT/Exporters/CodeGen/UhtHeaderCodeGeneratorHFile.cs
+++ b/Engine/Source/Programs/Shared/EpicGames.UHT/Exporters/CodeGen/UhtHeaderCodeGeneratorHFile.cs
@@ -1507,7 +1507,7 @@ namespace EpicGames.UHT.Exporters.CodeGen
 					builder.Append("virtual ");
 				}
 
-				builder.Append("~").Append(classObj.SourceName);
+				builder.Append('~').Append(classObj.SourceName);
 
 				if (classObj.ClassFlags.HasAnyFlags(EClassFlags.Interface))
 				{
