Shared: Address some trivial analysis warnings

#rnx

diff --git a/Engine/Source/Programs/Shared/EpicGames.Build/System/Parallel2.cs b/Engine/Source/Programs/Shared/EpicGames.Build/System/Parallel2.cs
--- a/Engine/Source/Programs/Shared/EpicGames.Build/System/Parallel2.cs
+++ b/Engine/Source/Programs/Shared/EpicGames.Build/System/Parallel2.cs
@@ -83,7 +83,7 @@ namespace UnrealBuildBase
 					}
 					finished = 1;
 				}
-			};
+			}
 
 			if (helperCount == -1)
 			{
diff --git a/Engine/Source/Programs/Shared/EpicGames.Build/UEBuildPlatformSDK.cs b/Engine/Source/Programs/Shared/EpicGames.Build/UEBuildPlatformSDK.cs
--- a/Engine/Source/Programs/Shared/EpicGames.Build/UEBuildPlatformSDK.cs
+++ b/Engine/Source/Programs/Shared/EpicGames.Build/UEBuildPlatformSDK.cs
@@ -1835,7 +1835,7 @@ namespace EpicGames.Core
 				if (File.Exists(HookExe))
 				{
 					Logger.LogDebug("Running {Hook} hook {HookExe}", Hook, HookExe);
-					using var _ = Timeline.ScopeEvent("RunAutoSDKHooks");
+					using ITimelineEvent _ = Timeline.ScopeEvent("RunAutoSDKHooks");
 
 					// run it
 					Process HookProcess = new Process();
diff --git a/Engine/Source/Programs/Shared/EpicGames.Core/Telemetry/DataRouterTelemetryService.cs b/Engine/Source/Programs/Shared/EpicGames.Core/Telemetry/DataRouterTelemetryService.cs
--- a/Engine/Source/Programs/Shared/EpicGames.Core/Telemetry/DataRouterTelemetryService.cs
+++ b/Engine/Source/Programs/Shared/EpicGames.Core/Telemetry/DataRouterTelemetryService.cs
@@ -119,7 +119,7 @@ namespace EpicGames.Core.Telemetry
 		/// <param name="appId"></param>
 		public void OverrideAppId(string appId)
 		{
-			this.AppId = appId;
+			AppId = appId;
 		}
 	}
 
diff --git a/Engine/Source/Programs/Shared/EpicGames.Serialization/CbField.cs b/Engine/Source/Programs/Shared/EpicGames.Serialization/CbField.cs
--- a/Engine/Source/Programs/Shared/EpicGames.Serialization/CbField.cs
+++ b/Engine/Source/Programs/Shared/EpicGames.Serialization/CbField.cs
@@ -1930,7 +1930,7 @@ namespace EpicGames.Serialization
 					}
 					break;
 				default:
-					throw new NotImplementedException($"Unhandled type '{field.GetType().ToString()}' in cb-json converter");
+					throw new NotImplementedException($"Unhandled type '{field.GetType()}' in cb-json converter");
 			}
 		}
 	}
@@ -2212,7 +2212,7 @@ namespace EpicGames.Serialization
 	class CbArrayDebugView(CbArray array)
 	{
 		[DebuggerBrowsable(DebuggerBrowsableState.RootHidden)]
-		public object?[] Value => array.Select(x => x.Value).ToArray();
+		public object?[] Value => [.. array.Select(x => x.Value)];
 	}
 
 	/// <summary>
@@ -2408,7 +2408,7 @@ namespace EpicGames.Serialization
 		}
 
 		[DebuggerBrowsable(DebuggerBrowsableState.RootHidden)]
-		public Property[] Properties => obj.Select(x => new Property { Name = x.Name.ToString(), Value = x.Value }).ToArray();
+		public Property[] Properties => [.. obj.Select(x => new Property { Name = x.Name.ToString(), Value = x.Value })];
 	}
 
 	/// <summary>
diff --git a/Engine/Source/Programs/Shared/EpicGames.Serialization/CompressedBuffer.cs b/Engine/Source/Programs/Shared/EpicGames.Serialization/CompressedBuffer.cs
--- a/Engine/Source/Programs/Shared/EpicGames.Serialization/CompressedBuffer.cs
+++ b/Engine/Source/Programs/Shared/EpicGames.Serialization/CompressedBuffer.cs
@@ -55,7 +55,7 @@ namespace EpicGames.Serialization
 		public ulong TotalCompressedSize { get; set; }
 
 		/** The hash of the uncompressed data. */
-		public byte[] RawHash { get; set; } = Array.Empty<byte>();
+		public byte[] RawHash { get; set; } = [];
 
 		public void ByteSwap()
 		{
@@ -116,7 +116,7 @@ namespace EpicGames.Serialization
 			byte[] crcData = new byte[blocksByteUsed + headerData.Length];
 			Array.Copy(headerData, crcData, headerData.Length);
 
-			uint[] blocks = Array.Empty<uint>();
+			uint[] blocks = [];
 
 			if (blocksByteUsed != 0)
 			{
@@ -275,7 +275,7 @@ namespace EpicGames.Serialization
 			long blockSize = DefaultBlockSize;
 			long blockCount = (rawContents.LongLength + blockSize - 1) / blockSize;
 			Span<byte> contentsSpan = new Span<byte>(rawContents);
-			List<byte[]> blocks = new List<byte[]>();
+			List<byte[]> blocks = [];
 
 			for (int i = 0; i < blockCount; i++)
 			{
@@ -301,7 +301,7 @@ namespace EpicGames.Serialization
 			const long DefaultBlockSize = 256 * 1024;
 			long blockSize = DefaultBlockSize;
 			long blockCount = (sourceStream.Length + blockSize - 1) / blockSize;
-			List<byte[]> blocks = new List<byte[]>();
+			List<byte[]> blocks = [];
 
 			for (int i = 0; i < blockCount; i++)
 			{
@@ -321,7 +321,7 @@ namespace EpicGames.Serialization
 
 			byte blockSizeExponent = (byte)Math.Floor(Math.Log2(blockSize));
 
-			List<byte[]> compressedBlocks = new List<byte[]>();
+			List<byte[]> compressedBlocks = [];
 			using Hasher hasher = Hasher.New();
 
 			ulong uncompressedContentLength = (ulong)blocks.Sum(b => b.LongLength);
@@ -366,7 +366,7 @@ namespace EpicGames.Serialization
 				RawHash = hashData
 			};
 
-			byte[] headerAndBlocks = WriteHeaderToBuffer(header, compressedBlocks.Select(b => (uint)b.Length).ToArray());
+			byte[] headerAndBlocks = WriteHeaderToBuffer(header, [.. compressedBlocks.Select(b => (uint)b.Length)]);
 
 			using BinaryWriter writer = new BinaryWriter(s, Encoding.Default, leaveOpen: true);
 
