UnrealBuildTool: Environment.GetEnvironmentVariable can return null

#rnx

diff --git a/Engine/Source/Programs/Shared/EpicGames.Build/UEBuildPlatformSDK.cs b/Engine/Source/Programs/Shared/EpicGames.Build/UEBuildPlatformSDK.cs
--- a/Engine/Source/Programs/Shared/EpicGames.Build/UEBuildPlatformSDK.cs
+++ b/Engine/Source/Programs/Shared/EpicGames.Build/UEBuildPlatformSDK.cs
@@ -504,6 +504,7 @@ namespace EpicGames.Core
 		{
 			return GetInstalledSDKVersion();
 		}
+
 		public string? GetInstalledVersion(string SDKName)
 		{
 			return GetSDKInfo(SDKName)?.Current;
@@ -1625,8 +1626,8 @@ namespace EpicGames.Core
 			// apply PATH modifications
 			if (PathRemoves != null || PathAdds != null)
 			{
-				string OrigPathVar = Environment.GetEnvironmentVariable("PATH")!;
-				IEnumerable<string> PathVars = OrigPathVar.Split( Path.PathSeparator );
+				string OrigPathVar = Environment.GetEnvironmentVariable("PATH") ?? String.Empty;
+				IEnumerable<string> PathVars = OrigPathVar.Split(Path.PathSeparator);
 				if (PathRemoves != null)
 				{
 					PathVars = PathVars.Except(PathRemoves, FileUtils.PlatformPathComparer);
