Introduced a common ICloudStorage implementation for searching and downloading of artifacts and updated RetreiveCloudArtifact in BuildGraph to use it. The intention is to replace UGS implementations here and allow for easier access of cloud artifacts in UAT in general as well as other more standalone uses in C# applications.

diff --git a/Engine/Source/Programs/AutomationTool/AutomationUtils/Automation.cs b/Engine/Source/Programs/AutomationTool/AutomationUtils/Automation.cs
--- a/Engine/Source/Programs/AutomationTool/AutomationUtils/Automation.cs
+++ b/Engine/Source/Programs/AutomationTool/AutomationUtils/Automation.cs
@@ -12,6 +12,8 @@ using EpicGames.Horde;
 using OpenTracing.Util;
 using UnrealBuildBase;
 using System.Threading.Tasks;
+using AutomationUtils;
+using EpicGames.ProjectStore;
 using Microsoft.Extensions.Logging;
 using Microsoft.Extensions.DependencyInjection;
 
@@ -268,6 +270,22 @@ namespace AutomationTool
 			serviceCollection.AddLogging(builder => builder.AddEpicDefault());
 			serviceCollection.AddHorde(options => options.AllowAuthPrompt = !CommandUtils.IsBuildMachine);
 
+			serviceCollection.AddScoped<BuildCommandContainer>();
+			serviceCollection.AddCloudStorage((CloudStorageOptions options) =>
+			{
+				BuildCommandContainer currentBuildCommand = CommandUtils.ServiceScope.ServiceProvider.GetService<BuildCommandContainer>();
+				
+				string projectFile = currentBuildCommand.CurrentCommand.ParseParamValue("Project");
+				options.ReadFromIni(projectFile != null ? new FileReference(projectFile) : null /* no project file specified, only engine inis will be applied */);
+
+				options.AllowAuthPrompt = !CommandUtils.IsBuildMachine;
+				options.FindExecutablesFromUnrealRoot();
+				options.EchoZenCliToStdout = true;
+
+				options.OnProcessCreated += process => ProcessManager.AddProcess(ZenUATProcessUtil.GetUATProcessStub(process));
+				options.OnProcessExited += process => ProcessManager.RemoveProcess(ZenUATProcessUtil.GetUATProcessStub(process));
+			});
+
 			await using ServiceProvider serviceProvider = serviceCollection.BuildServiceProvider();
 			try
 			{
@@ -284,7 +302,17 @@ namespace AutomationTool
 						throw new AutomationException("Failed to find command {0}", CommandInfo.CommandName);
 					}
 
+					// each command is considered a "request" in the DI framework e.g. owns scoped providers
+					// allows for systems that require per command configuration
+					using IServiceScope serviceScope = serviceProvider.CreateScope();
+					CommandUtils.ServiceScope = serviceScope;
+
 					BuildCommand Command = (BuildCommand)Activator.CreateInstance(CommandType);
+
+					// make the current build command available in the scope
+					BuildCommandContainer buildCommandContainer = serviceScope.ServiceProvider.GetService<BuildCommandContainer>();
+					buildCommandContainer.CurrentCommand = Command;
+
 					Command.Params = CommandInfo.Arguments.ToArray();
 					try
 					{
@@ -313,6 +341,7 @@ namespace AutomationTool
 			finally
 			{
 				CommandUtils.ServiceProvider = null!;
+				CommandUtils.ServiceScope = null!;
 			}
 		}
 
@@ -413,6 +442,11 @@ AutomationTool.exe [-verbose] [-Help|-List|-CompileOnly] [-p4] Command0 [-Arg0 -
 		private static bool bShouldKillProcesses = true;
 	}
 
+	public class BuildCommandContainer
+	{
+		public BuildCommand CurrentCommand { get; set; }
+	}
+
 
 	// This class turns stringly-typed commandline option values into named compile-time values
 	public class GlobalCommandLine
diff --git a/Engine/Source/Programs/AutomationTool/AutomationUtils/CloudStorageUtils.cs b/Engine/Source/Programs/AutomationTool/AutomationUtils/CloudStorageUtils.cs
new file mode 100644
--- /dev/null
+++ b/Engine/Source/Programs/AutomationTool/AutomationUtils/CloudStorageUtils.cs
@@ -0,0 +1,65 @@
+// Copyright Epic Games, Inc. All Rights Reserved.
+
+using System;
+using System.Collections.Generic;
+using System.Linq;
+using AutomationTool;
+using EpicGames.Core;
+using EpicGames.ProjectStore;
+using JetBrains.Annotations;
+using UnrealBuildTool;
+
+namespace AutomationUtils
+{
+    public static class CloudStorageOptionsExtensions
+    {
+	    /// <summary>
+	    /// Attempt to initialize options from inis
+	    /// </summary>
+	    /// <param name="projectPath"></param>
+	    public static void ReadFromIni(this CloudStorageOptions options, [CanBeNull] FileReference projectPath)
+	    {
+		    Dictionary<UnrealTargetPlatform, ConfigHierarchy> engineConfigs = new Dictionary<UnrealTargetPlatform, ConfigHierarchy>();
+		    Dictionary<UnrealTargetPlatform, ConfigHierarchy> gameConfigs = new Dictionary<UnrealTargetPlatform, ConfigHierarchy>();
+		    foreach (UnrealTargetPlatform targetPlatformType in UnrealTargetPlatform.GetValidPlatforms())
+		    {
+			    DirectoryReference projectDir = projectPath?.Directory;
+
+			    ConfigHierarchy engineConfig = ConfigCache.ReadHierarchy(ConfigHierarchyType.Engine, projectDir, targetPlatformType);
+			    engineConfigs.Add(targetPlatformType, engineConfig);
+
+			    ConfigHierarchy gameConfig = ConfigCache.ReadHierarchy(ConfigHierarchyType.Game, projectDir, targetPlatformType); 
+			    gameConfigs.Add(targetPlatformType, gameConfig);
+		    }
+
+		    ConfigHierarchy config = engineConfigs[HostPlatform.Current.HostEditorPlatform];
+		    bool foundConfig = config.TryGetValueGeneric("StorageServers", "Cloud", out IniCloudConfiguration foundCloudConfig);
+
+		    if (!foundConfig)
+		    {
+			    return;
+		    }
+
+		    string host = foundCloudConfig.Host;
+		    if (host.Contains(';', StringComparison.OrdinalIgnoreCase))
+		    {
+			    // if it's a list pick the first element
+			    host = host.Split(";").First();
+		    }
+		    options.Host = host;
+		    options.NamespaceId = foundCloudConfig.BuildsNamespace;
+		    options.BaselineBranch = foundCloudConfig.BuildsBaselineBranch;
+	    }
+		
+	    struct IniCloudConfiguration
+	    {
+#pragma warning disable IDE1006 // Static analysis wants these to be named differently, but they must be named the same as the config file properties
+		    public string Host = String.Empty;
+		    public string BuildsNamespace = String.Empty;
+		    public string BuildsBaselineBranch = String.Empty;
+#pragma warning restore IDE1006
+
+		    public IniCloudConfiguration() {}
+	    }
+    }
+}
diff --git a/Engine/Source/Programs/AutomationTool/AutomationUtils/CommandUtils.cs b/Engine/Source/Programs/AutomationTool/AutomationUtils/CommandUtils.cs
--- a/Engine/Source/Programs/AutomationTool/AutomationUtils/CommandUtils.cs
+++ b/Engine/Source/Programs/AutomationTool/AutomationUtils/CommandUtils.cs
@@ -18,6 +18,7 @@ using System.Threading;
 using System.Threading.Tasks;
 using System.Xml;
 using System.Xml.Serialization;
+using Microsoft.Extensions.DependencyInjection;
 using UnrealBuildBase;
 using UnrealBuildTool;
 
@@ -115,6 +116,12 @@ namespace AutomationTool
 		/// </summary>
 		public static IServiceProvider ServiceProvider { get; set; } = null!;
 
+		
+		/// <summary>
+		/// Access to scoped services which are created per command
+		/// </summary>
+		public static IServiceScope ServiceScope { get; set; } = null!;
+
 		/// <summary>
 		/// Global activity source for UAT. Can be used to add trace data and additional context to logs.
 		/// </summary>
diff --git a/Engine/Source/Programs/AutomationTool/AutomationUtils/ZenUtils.cs b/Engine/Source/Programs/AutomationTool/AutomationUtils/ZenUtils.cs
--- a/Engine/Source/Programs/AutomationTool/AutomationUtils/ZenUtils.cs
+++ b/Engine/Source/Programs/AutomationTool/AutomationUtils/ZenUtils.cs
@@ -8,7 +8,9 @@ using EpicGames.Serialization;
 using IdentityModel.OidcClient;
 using Microsoft.Extensions.Logging;
 using System;
+using System.Collections.Concurrent;
 using System.Collections.Generic;
+using System.Diagnostics;
 using System.IO;
 using System.Linq;
 using System.Net.Http;
@@ -708,4 +710,36 @@ namespace AutomationUtils
 			return UseZenStore && EnablePakStreaming;
 		}
 	}
+
+	public static class ZenUATProcessUtil
+	{
+		private static readonly ConcurrentDictionary<ManagedProcess, ZenUatProcessStub> ProcessStubs = new ConcurrentDictionary<ManagedProcess, ZenUatProcessStub>();
+		public static IProcess GetUATProcessStub(ManagedProcess process)
+		{
+			return ProcessStubs.GetOrAdd(process, p => new ZenUatProcessStub(p));
+		}
+	}
+
+	public class ZenUatProcessStub(ManagedProcess process) : IProcess
+	{
+		public void StopProcess(bool killDescendants = true)
+		{
+			process.Kill(killDescendants: killDescendants);
+		}
+
+		public bool HasExited => process.HasExited;
+
+		public string GetProcessName()
+		{
+			try
+			{
+				return Process.GetProcessById(process.Id).ProcessName;
+			}
+			catch (ArgumentException)
+			{
+				// process is not running anymore
+				return null;
+			}
+		}
+	}
 }
diff --git a/Engine/Source/Programs/AutomationTool/BuildGraph/Tasks/CreateCloudArtifactTask.cs b/Engine/Source/Programs/AutomationTool/BuildGraph/Tasks/CreateCloudArtifactTask.cs
--- a/Engine/Source/Programs/AutomationTool/BuildGraph/Tasks/CreateCloudArtifactTask.cs
+++ b/Engine/Source/Programs/AutomationTool/BuildGraph/Tasks/CreateCloudArtifactTask.cs
@@ -14,6 +14,7 @@ using EpicGames.Horde.Commits;
 using EpicGames.Horde.Storage;
 using EpicGames.Horde.Storage.Nodes;
 using EpicGames.Horde.Streams;
+using EpicGames.ProjectStore;
 using EpicGames.Serialization;
 using Microsoft.Extensions.DependencyInjection;
 using Microsoft.Extensions.Logging;
@@ -136,12 +137,6 @@ namespace AutomationTool.Tasks
 		[TaskParameter(Optional = true)]
 		public string BaselineBranch { get; set; } = null!;
 
-		/// <summary>
-        /// Set to upload more blobs in parallel, defaults to number of cores. Set to -1 to go as wide as possible.
-        /// </summary>
-        [TaskParameter(Optional = true)]
-        public int MaxBlobUploadParallelism { get; set; } = Environment.ProcessorCount;
-
 		/// <summary>
 		/// Enable to use multipart endpoints if valuable
 		/// </summary>
@@ -208,12 +203,12 @@ namespace AutomationTool.Tasks
 
 			string name = String.IsNullOrEmpty(_parameters.Name) ? _parameters.HordeArtifactID : _parameters.Name;
 			ArtifactType type = new ArtifactType(_parameters.Type);
-			string sanitizedHordeArtifactParam = SanitizeBucketValue(String.IsNullOrEmpty(_parameters.HordeArtifactID) ? name : _parameters.HordeArtifactID);
+			string sanitizedHordeArtifactParam = ICloudStorage.SanitizeBucketValue(String.IsNullOrEmpty(_parameters.HordeArtifactID) ? name : _parameters.HordeArtifactID);
 			ArtifactName hordeArtifactID = new ArtifactName(sanitizedHordeArtifactParam);
 			List<string> metadataParam = (_parameters.Metadata ?? String.Empty).Split(';', StringSplitOptions.RemoveEmptyEntries | StringSplitOptions.TrimEntries).ToList();
 
 			Dictionary<string, object> metadata = new Dictionary<string, object>();
-			string platform = SanitizeBucketValue(_parameters.Platform);
+			string platform = ICloudStorage.SanitizeBucketValue(_parameters.Platform);
 			
 			FileReference? projectPath = _parameters.Project;
 			string project = projectPath != null ? projectPath.GetFileNameWithoutExtension() : "";
@@ -319,7 +314,7 @@ namespace AutomationTool.Tasks
 					throw new AutomationException("Missing UE_HORDE_STREAMID environment variable; unable to determine current stream.");
 				}
 			}
-			string streamId = SanitizeBucketValue(streamIdTyped.ToString());
+			string streamId = ICloudStorage.SanitizeBucketValue(streamIdTyped.ToString());
 			metadata.Add("stream", streamId);
 			
 			string baselineBranch =_parameters.BaselineBranch;
@@ -334,7 +329,7 @@ namespace AutomationTool.Tasks
 				baselineBranch = streamId;
 			}
 
-			baselineBranch = SanitizeBucketValue(baselineBranch);
+			baselineBranch = ICloudStorage.SanitizeBucketValue(baselineBranch);
 
 			CommitId commitId = CommitId.Empty;
 
@@ -613,15 +608,6 @@ namespace AutomationTool.Tasks
 			return (engineConfigs, gameConfigs);
 		}
 
-		internal static string SanitizeBucketValue(string s)
-		{
-			// trim starting + and then replace any + and . with -
-			s = s.TrimStart('+');
-			s = s.Replace("+", "-", StringComparison.OrdinalIgnoreCase);
-			// . have special meaning in the bucket names, its used as a separator between components, as such component can not themselves contain .
-			return s.Replace(".", "-", StringComparison.OrdinalIgnoreCase);
-		}
-		
 		async Task RegisterHordeArtifactAsync(Uri buildUri, ArtifactName name, ArtifactType type, StreamId streamId, CommitId commitId, Dictionary<string, object> metadata)
 		{		
 			// Add keys used to collate and find builds
diff --git a/Engine/Source/Programs/AutomationTool/BuildGraph/Tasks/RetrieveCloudArtifactTask.cs b/Engine/Source/Programs/AutomationTool/BuildGraph/Tasks/RetrieveCloudArtifactTask.cs
--- a/Engine/Source/Programs/AutomationTool/BuildGraph/Tasks/RetrieveCloudArtifactTask.cs
+++ b/Engine/Source/Programs/AutomationTool/BuildGraph/Tasks/RetrieveCloudArtifactTask.cs
@@ -4,21 +4,20 @@ using System;
 using System.Collections.Generic;
 using System.IO;
 using System.Linq;
-using System.Net.Http;
-using System.Net.Http.Headers;
-using System.Net.Http.Json;
-using System.Net.Mime;
+using System.Threading;
 using System.Threading.Tasks;
 using System.Xml;
-using AutomationTool.Tasks.CloudDDC;
+using AutomationUtils;
 using EpicGames.Core;
 using EpicGames.Horde.Artifacts;
 using EpicGames.Horde.Commits;
+using EpicGames.Horde.Storage;
 using EpicGames.Horde.Streams;
+using EpicGames.ProjectStore;
 using EpicGames.Serialization;
+using Microsoft.Extensions.DependencyInjection;
 using Microsoft.Extensions.Logging;
 using UnrealBuildBase;
-using UnrealBuildTool;
 
 #nullable enable
 
@@ -84,10 +83,10 @@ namespace AutomationTool.Tasks
 		public string Platform { get; set; } = null!;
 
 		/// <summary>
-		/// The path to the uproject this artifact is created for
+		/// The path to the uproject this artifact is created for, this is optional if UAT is run with the -project option
 		/// </summary>
-		[TaskParameter]
-		public FileReference Project { get; set; } = null!;
+		[TaskParameter(Optional = true)]
+		public FileReference? Project { get; set; } = null!;
 
 		/// <summary>
 		/// The platform the artifact is created for
@@ -135,6 +134,18 @@ namespace AutomationTool.Tasks
 		/// Enable to create an Unreal Insights trace of the download process
 		/// </summary>
 		public bool EnableTracing { get; set; } = true;
+
+		/// <summary>
+		/// Optional pattern(s) to download from the artifact, if omitted the entire artifact is downloaded
+		/// </summary>
+		[TaskParameter(Optional = true, ValidationType = TaskParameterValidationType.FileSpec)]
+		public string? Include { get; set; } = null;
+
+		/// <summary>
+		/// Optional pattern(s) to exclude from the copy for example, Engine/NoCopy*.txt)
+		/// </summary>
+		[TaskParameter(Optional = true, ValidationType = TaskParameterValidationType.FileSpec)]
+		public string? Exclude { get; set; } = null;
 	}
 
 	/// <summary>
@@ -160,25 +171,36 @@ namespace AutomationTool.Tasks
 		/// <param name="tagNameToFileSet">Mapping from tag names to the set of files they include.</param>
 		public override async Task ExecuteAsync(JobContext job, HashSet<FileReference> buildProducts, Dictionary<string, HashSet<FileReference>> tagNameToFileSet)
 		{
-			string project = _parameters.Project.GetFileNameWithoutExtension();
-			FileReference? projectPath = _parameters.Project;
-			if (!File.Exists(_parameters.Project.FullName))
+			string? defaultProject = job.OwnerCommand.ParseParamValue("Project");
+			FileReference? defaultProjectFile = null;
+			if (defaultProject != null && File.Exists(defaultProject))
+			{
+				defaultProjectFile = new FileReference(defaultProject);
+			}
+
+			FileReference? projectPath = _parameters.Project != null ? _parameters.Project : defaultProjectFile;
+			string? project = projectPath != null ? projectPath.GetFileNameWithoutExtension() : null;
+			if (projectPath == null || !File.Exists(projectPath.FullName))
 			{
-				// if the project is not an uproject that exists we do the scanning for engine inis without the project specific overrides
+				// if the project is not an uproject that exists we can not scan for inis
 				projectPath = null;
-				Logger.LogInformation("\'{Project}\' is not a valid file. Loading engine configs without project overrides.", _parameters.Project);
+				Logger.LogInformation("\'{Project}\' is not a valid file. Not parsing ini files, config needs to be explicitly provided.", project);
+			}
+
+			if (project == null)
+			{
+				Logger.LogError("Project is null and no ProjectFile set for build graph. Unable to determine project for this upload.");
+				return;
 			}
 			string platform = _parameters.Platform;
 
-			bool foundConfig = false;
-			CloudConfiguration? cloudConfig = null;
+			ICloudStorage? cloudStorage = CommandUtils.ServiceScope.ServiceProvider.GetService<ICloudStorage>();
+			CloudBuildOptionsOverride optionsOverride = new CloudBuildOptionsOverride();
 
-			if (projectPath != null)
+			if (cloudStorage == null)
 			{
-				(Dictionary<UnrealTargetPlatform, ConfigHierarchy> engineConfigs, Dictionary<UnrealTargetPlatform, ConfigHierarchy> _) = CreateCloudArtifactTask.GetIniConfigs(projectPath);
-				ConfigHierarchy config = engineConfigs[HostPlatform.Current.HostEditorPlatform];
-				foundConfig = config.TryGetValueGeneric("StorageServers", "Cloud", out CloudConfiguration foundCloudConfig);
-				cloudConfig = foundCloudConfig;
+				// TODO: If this actually happens we could likely initialize cloud storage ourselves
+				throw new NotImplementedException("Failed to retrieve ICloudStorage from service provider, unable to continue");
 			}
 
 			// Figure out the current change and stream id
@@ -202,50 +224,26 @@ namespace AutomationTool.Tasks
 
 			string? cloudHostEnvVar = Environment.GetEnvironmentVariable(OperatingSystem.IsWindows() ? "UE-CloudPublishHost" : "UE_CloudPublishHost");
 			
-			string host;
 			if (!String.IsNullOrEmpty(_parameters.Host))
 			{
-				host = _parameters.Host;
+				optionsOverride.OverrideHost = _parameters.Host;
 			}
 			else if (!String.IsNullOrEmpty(cloudHostEnvVar))
 			{
-				host = cloudHostEnvVar;
-			}
-			else if (foundConfig && cloudConfig != null && !String.IsNullOrEmpty(cloudConfig.Value.Host))
-			{
-				string cloudConfigHost = cloudConfig.Value.Host;
-				if (cloudConfigHost.Contains(';', StringComparison.OrdinalIgnoreCase))
-				{
-					// if its a list pick the first element
-					cloudConfigHost = cloudConfigHost.Split(";").First();
-				}
-
-				host = cloudConfigHost;
-			}
-			else
-			{
-				throw new AutomationException("Missing UE-CloudPublishHost environment variable; unable to determine cloud host. Specify the environment variable or define it in your StorageServer section in your engine ini");
+				optionsOverride.OverrideHost = cloudHostEnvVar;
 			}
 
-			string ns;
 			string? cloudDefaultNamespaceEnv = Environment.GetEnvironmentVariable(OperatingSystem.IsWindows() ? "UE-CloudPublishNamespace" : "UE_CloudPublishNamespace");
 				
 			if (!String.IsNullOrEmpty(_parameters.Namespace))
 			{
-				ns = _parameters.Namespace;
+				optionsOverride.NamespaceId = new NamespaceId(_parameters.Namespace);
 			}
 			else if (!String.IsNullOrEmpty(cloudDefaultNamespaceEnv))
 			{
-				ns = cloudDefaultNamespaceEnv;
-			}
-			else if (foundConfig && cloudConfig != null && !String.IsNullOrEmpty(cloudConfig.Value.BuildsNamespace))
-			{
-				ns = cloudConfig.Value.BuildsNamespace;
-			}
-			else
-			{
-				throw new AutomationException("Missing UE-CloudPublishNamespace environment variable; unable to default namespace please specify it in the task or define it in your StorageServer section in your engine ini.");
+				optionsOverride.NamespaceId = new NamespaceId(cloudDefaultNamespaceEnv);
 			}
+
 			string httpVersion = _parameters.HttpVersion;
 			if (String.IsNullOrEmpty(httpVersion) || httpVersion == "None")
 			{
@@ -256,21 +254,16 @@ namespace AutomationTool.Tasks
 				}
 			}
 
-			string accessToken;
 			if (!String.IsNullOrEmpty(_parameters.AccessToken))
 			{
-				accessToken = _parameters.AccessToken;
+				optionsOverride.AccessToken = _parameters.AccessToken;
 			}
 			else
 			{
 				string? cloudAccessTokenEnvVar = Environment.GetEnvironmentVariable(OperatingSystem.IsWindows() ? "UE-CloudDataCacheAccessToken" : "UE_CloudDataCacheAccessToken");
 				if (!String.IsNullOrEmpty(cloudAccessTokenEnvVar))
 				{
-					accessToken = cloudAccessTokenEnvVar;
-				}
-				else
-				{
-					throw new AutomationException("Missing UE-CloudDataCacheAccessToken environment variable; unable to find access token to use.");
+					optionsOverride.AccessToken = cloudAccessTokenEnvVar;
 				}
 			}
 
@@ -300,6 +293,17 @@ namespace AutomationTool.Tasks
 				}
 			}
 
+			await cloudStorage.Reconfigure(options =>
+			{
+				if (projectPath != null && File.Exists(projectPath.FullName))
+				{
+					// initialize from the explicitly defined project name property for backwards compatibility reasons, this is usually read by the -ProjectFile argument to UAT
+					options.ReadFromIni(projectPath);
+				}
+			});
+
+			await cloudStorage.LoginAsync(unattended: true, optionsOverride);
+
 			string? name = null;
 			if (!String.IsNullOrEmpty(_parameters.Name))
 			{
@@ -309,122 +313,92 @@ namespace AutomationTool.Tasks
 			ArtifactType type = new ArtifactType(_parameters.Type);
 			List<string> keys = (_parameters.Keys ?? String.Empty).Split(';', StringSplitOptions.RemoveEmptyEntries | StringSplitOptions.TrimEntries).ToList();
 
-			CbWriter queryWriter = new CbWriter();
-			queryWriter.BeginObject(); // root
-			queryWriter.BeginObject("query"); // query
-			
-			queryWriter.BeginObject("stream");
-			queryWriter.WriteString("$eq", CreateCloudArtifactTask.SanitizeBucketValue(streamId.ToString()));
-			queryWriter.EndObject();
+			CloudStorageQueryBuilder queryBuilder = new();
 
+			string branch = ICloudStorage.SanitizeBucketValue(streamId.ToString());
+			queryBuilder.Equals("platform", ICloudStorage.SanitizeBucketValue(platform));
+			
 			if (requiredCommitId != null)
 			{
-				queryWriter.BeginObject("commit");
-				queryWriter.WriteString("$eq", requiredCommitId.ToString());
-				queryWriter.EndObject();
+				queryBuilder.Equals("commit", requiredCommitId.ToString());
 			}
 
 			if (name != null)
 			{
-				queryWriter.BeginObject("name");
-				queryWriter.WriteString("$eq", name);
-				queryWriter.EndObject();
+				queryBuilder.Equals("name", name);
 			}
 
 			if (!String.IsNullOrEmpty(_parameters.BuildGroup))
 			{
-				queryWriter.BeginObject("buildgroup");
-				queryWriter.WriteString("$eq", _parameters.BuildGroup);
-				queryWriter.EndObject();
+				queryBuilder.Equals("buildgroup", _parameters.BuildGroup);
 			}
 
-			queryWriter.BeginObject("type");
-			queryWriter.WriteString("$eq", type.ToString());
-			queryWriter.EndObject();
+			queryBuilder.Equals("type", type.ToString());
 
 			if (keys.Count > 0)
 			{
-				queryWriter.BeginObject("keys");
-				queryWriter.BeginUniformArray("$in", CbFieldType.String);
-				foreach (string key in keys)
-				{
-					queryWriter.WriteStringValue(key);
-				}
-				queryWriter.EndUniformArray();
-				queryWriter.EndObject();
+				queryBuilder.In("keys", keys);
 			}
-			
-			queryWriter.EndObject(); // end query
-			queryWriter.BeginObject("options");
-			queryWriter.WriteInteger("limit", 1); // we only care about the first match
-			queryWriter.EndObject(); // options
-			queryWriter.EndObject(); // end root object
 
-			CbObject queryObject = queryWriter.ToObject();
 			DirectoryReference outputDir = ResolveDirectory(_parameters.OutputDir);
-			string bucketId = $"{CreateCloudArtifactTask.SanitizeBucketValue(project)}.{CreateCloudArtifactTask.SanitizeBucketValue(type.ToString())}.{CreateCloudArtifactTask.SanitizeBucketValue(streamId.ToString())}.{CreateCloudArtifactTask.SanitizeBucketValue(platform)}".ToLowerInvariant();
-			using HttpClient httpClient = BuildHttpClient(host, ns, bucketId, accessToken);
 
-			List<CbObjectId> artifactIds = await SearchArtifactsAsync(httpClient, queryObject);
-
-			if (artifactIds.Count == 0)
+			List<FilePattern>? includeFilter = null;
+			if (!String.IsNullOrEmpty(_parameters.Include))
 			{
-				throw new AutomationException($"Unable to find any artifact matching given criteria in namespace {ns} and bucket {bucketId} criteria: {queryObject.ToJson()}");
+				includeFilter = new List<FilePattern>();
+
+				foreach (string filter in SplitDelimitedList(_parameters.Include))
+				{
+					includeFilter.Add(new FilePattern(outputDir, filter));
+				}
 			}
 
-			if (!_parameters.AllowMultipleMatches && artifactIds.Count != 1)
+			List<FilePattern>? excludeFilter = null;
+			if (!String.IsNullOrEmpty(_parameters.Exclude))
 			{
-				throw new AutomationException("More then one matching artifact given criteria, set \"AllowMultipleMatches\" option if you want to use the newest matching artifact or refine the search");
+				excludeFilter = new List<FilePattern>();
+				foreach (string filter in SplitDelimitedList(_parameters.Exclude))
+				{
+					excludeFilter.Add(new FilePattern(outputDir, filter));
+				}
 			}
-			
-			CbObjectId artifact = artifactIds.First();
-			Logger.LogInformation("Found artifact {ArtifactId}", artifact);
 
-			FileReference? traceFile = null;
-			if (_parameters.EnableTracing)
+			CbObjectId artifact;
+			FindBuildResponse foundBuild;
+			
 			{
-				// this is automatically uploaded as it's in the Saved folder
-				traceFile = FileReference.Combine(Unreal.RootDirectory, $"Engine/Programs/AutomationTool/Saved/Logs/build-download-{artifact}.utrace");
-			}
+				CbObject queryObject = queryBuilder.Done();
+				List<FindBuildResponse> matchingBuilds = await cloudStorage.FindBuildsAsync(project.ToLower(), type: type.ToString(), branch: branch, queryObject, searchOptions: new CloudBuildSearchOptions {Limit  = 2}, optionsOverride: optionsOverride, cancellationToken: CancellationToken.None).ToListAsync();
 
-			BuildDownload.DownloadBuild(host, outputDir, ns, bucketId, artifact.ToString(), accessToken, allowMultipart: _parameters.AllowMultipart, assumeHttp2: String.Equals(httpVersion, "http2-only", StringComparison.OrdinalIgnoreCase), boostWorkers: _parameters.BoostWorkers, traceFile: traceFile);
-		}
+				if (matchingBuilds.Count == 0)
+				{
+					throw new AutomationException($"Unable to find any artifact matching given criteria in namespace {cloudStorage.GetCurrentNamespace(optionsOverride)} for project \"{project}\" and type \"{type}\" using criteria: {queryObject.ToJson()}");
+				}
 
-		static HttpClient BuildHttpClient(string host, string ns, string bucketId, string accessToken)
-		{
-			HttpClient httpClient = new HttpClient();
-			string url = $"{host}/api/v2/builds/{ns}/{bucketId}/";
-			try
-			{
-				httpClient.BaseAddress = new Uri(url);
+				if (!_parameters.AllowMultipleMatches && matchingBuilds.Count != 1)
+				{
+					throw new AutomationException("More then one matching artifact given criteria, set \"AllowMultipleMatches\" option if you want to use the newest matching artifact or refine the search");
+				}
+
+				foundBuild = matchingBuilds.First();
+				artifact = foundBuild.BuildId;
 			}
-			catch (UriFormatException)
+
+			FileReference? traceFile = null;
+			if (_parameters.EnableTracing)
 			{
-				throw new AutomationException($"{url} is not a valid url. Make sure host, namespace and bucket have been specified");
+				// this is automatically uploaded by Horde as it's in the Saved folder
+				traceFile = FileReference.Combine(Unreal.RootDirectory, $"Engine/Programs/AutomationTool/Saved/Logs/build-download-{artifact}.utrace");
 			}
 
-			CbObjectId sessionId = CbObjectId.NewObjectId();
-			Logger.LogInformation("Using SessionId {SessionId} please include this in any error reports.", sessionId);
-			string authScheme = "Bearer"; // TODO: for a proper implementation this should be configurable
-			httpClient.DefaultRequestHeaders.Add("Authorization", $"{authScheme} {accessToken}");
-			httpClient.DefaultRequestHeaders.Add("UE-Session", sessionId.ToString());
-			httpClient.DefaultRequestHeaders.Add("Accept", CustomMediaTypeNames.UnrealCompactBinary);
-			httpClient.Timeout = TimeSpan.FromMinutes(5.0); // bump timeout for each request as we sometimes have larger files
-			return httpClient;
-		}
-		static async Task<List<CbObjectId>> SearchArtifactsAsync(HttpClient httpClient, CbObject queryObject)
-		{
-			using StringContent content = new StringContent(queryObject.ToJson());
-			content.Headers.ContentType = new MediaTypeHeaderValue(MediaTypeNames.Application.Json);
-
-			using HttpRequestMessage request = new HttpRequestMessage(HttpMethod.Post, new Uri($"search", UriKind.Relative));
-			request.Content = content;
-			request.Headers.Add("Accept", MediaTypeNames.Application.Json);
-			using HttpResponseMessage result = await httpClient.SendAsync(request);
-			result.EnsureSuccessStatusCode();
-
-			SearchResult? searchResult = await result.Content.ReadFromJsonAsync<SearchResult>();
-			return searchResult!.Results.Select(r => r.BuildId).ToList();
+			Logger.LogInformation("Downloading artifact {BuildId} from bucket {BucketId}", foundBuild.BuildId, foundBuild.BucketId);
+			await cloudStorage.DownloadBuildAsync(foundBuild.BucketId, foundBuild.BuildId, outputDir, includeFilter, excludeFilter, new CloudBuildDownloadOptions
+				{
+					AllowMultipart = _parameters.AllowMultipart, 
+					AssumeHttp2 = String.Equals(httpVersion, "http2-only", StringComparison.OrdinalIgnoreCase),
+					BoostWorkers = _parameters.BoostWorkers,
+					TraceFile = traceFile
+				}, optionsOverride: optionsOverride);
 		}
 
 		/// <inheritdoc/>
@@ -439,41 +413,4 @@ namespace AutomationTool.Tasks
 		public override IEnumerable<string> FindProducedTagNames()
 			=> Enumerable.Empty<string>();
 	}
-
-	class BuildSearchResult
-	{
-		[CbField("buildId")]
-		public CbObjectId BuildId { get; set; }
-	}
-
-	class SearchResult
-	{
-		[CbField("results")] public List<BuildSearchResult> Results { get; set; } = new List<BuildSearchResult>();
-
-		[CbField("partialResult")] public bool PartialResult { get; set; }
-
-		public SearchResult()
-		{
-			Results = new List<BuildSearchResult>();
-			PartialResult = false;
-		}
-	}
-
-	static class BuildDownload
-	{
-		public static void DownloadBuild(string host, DirectoryReference targetDir, string namespaceId, string bucketId, string buildId, string accessToken, bool allowMultipart = false, bool assumeHttp2 = false, bool boostWorkers = false, FileReference? traceFile = null)
-		{
-			FileInfo zenExe = new FileInfo("Engine/Binaries/Win64/zen.exe");
-
-			string http2Options = assumeHttp2 ? "--assume-http2" : "";
-			string boostWorkersOption = boostWorkers ? "--boost-workers" : "";
-			// trace file option needs to be passed before other arguments including the verbs
-			string traceOption = traceFile != null ? $"--tracefile={traceFile.FullName} " : "";
-			// pass the access token via environment variable to avoid showing it on the cli
-			string cmdline =
-				$"{traceOption}builds download --url {host} --namespace {namespaceId} --bucket {bucketId} --access-token-env UE-CloudDataCacheAccessToken --plain-progress --allow-multipart={allowMultipart} {http2Options} {boostWorkersOption} \"{targetDir}\" {buildId}";
-			CommandUtils.RunAndLog(CommandUtils.CmdEnv, zenExe.FullName, cmdline, Options: CommandUtils.ERunOptions.Default,
-				EnvVars: new Dictionary<string, string> { { "UE-CloudDataCacheAccessToken", accessToken } });
-		}
-	}
 }
diff --git a/Engine/Source/Programs/Shared/EpicGames.Core/ManagedProcess.cs b/Engine/Source/Programs/Shared/EpicGames.Core/ManagedProcess.cs
--- a/Engine/Source/Programs/Shared/EpicGames.Core/ManagedProcess.cs
+++ b/Engine/Source/Programs/Shared/EpicGames.Core/ManagedProcess.cs
@@ -1400,6 +1400,37 @@ namespace EpicGames.Core
 			}
 		}
 
+		/// <summary>
+		/// Returns true if the process has exited
+		/// </summary>
+		public bool HasExited
+		{
+			get
+			{
+				if (_frameworkProcess != null)
+				{
+					try
+					{
+						// throws exception if the process is still running
+						return _frameworkProcess.HasExited;
+					}
+					catch (Exception)
+					{
+						return false;
+					}
+				}
+				else
+				{
+					// returns non-zero on success and success means the process has exited
+					if(GetExitCodeProcess(_processHandle!, out _) != 0)
+					{
+						return true;
+					}
+					return false;
+				}
+			}
+		}
+
 		/// <summary>
 		/// Gets all the current environment variables
 		/// </summary>
@@ -1418,5 +1449,38 @@ namespace EpicGames.Core
 			}
 			return newEnvironment;
 		}
+
+		/// <summary>
+		/// Kill the process if it's still running, optionally kill child processes
+		/// </summary>
+		/// <param name="killDescendants"></param>
+		public void Kill(bool killDescendants = true)
+		{
+			if(_frameworkProcess != null)
+			{
+				if (_frameworkProcess.HasExited)
+				{
+					return;
+				}
+
+				// Kill entire process tree
+				try
+				{
+					_frameworkProcess.Kill(killDescendants);
+				}
+				catch { }
+
+				try
+				{
+					_frameworkProcess.WaitForExit();
+				}
+				catch { }
+			}
+			else
+			{
+				_ = TerminateProcess(_processHandle!, 0);
+				_ = WaitForSingleObject(_processHandle!, INFINITE);
+			}
+		}
 	}
 }
diff --git a/Engine/Source/Programs/Shared/EpicGames.ProjectStore/CloudStorage/CloudStorage.cs b/Engine/Source/Programs/Shared/EpicGames.ProjectStore/CloudStorage/CloudStorage.cs
new file mode 100644
--- /dev/null
+++ b/Engine/Source/Programs/Shared/EpicGames.ProjectStore/CloudStorage/CloudStorage.cs
@@ -0,0 +1,640 @@
+// Copyright Epic Games, Inc. All Rights Reserved.
+
+using System;
+using System.Collections;
+using System.Collections.Generic;
+using System.ComponentModel.DataAnnotations;
+using System.Diagnostics;
+using System.IO;
+using System.Linq;
+using System.Runtime.CompilerServices;
+using System.Text.RegularExpressions;
+using System.Threading;
+using System.Threading.Tasks;
+using EpicGames.Core;
+using EpicGames.Horde;
+using EpicGames.Horde.Storage;
+using EpicGames.OIDC;
+using EpicGames.Serialization;
+using Microsoft.Extensions.Configuration;
+using Microsoft.Extensions.DependencyInjection;
+using Microsoft.Extensions.Logging;
+using Microsoft.Extensions.Options;
+
+namespace EpicGames.ProjectStore
+{
+	/// <inheritdoc />
+	public class CloudStorage : ICloudStorage
+	{
+		private readonly IOptionsSnapshot<CloudStorageOptions> _options;
+		private readonly ILogger<CloudStorage> _logger;
+		private readonly IServiceProvider _provider;
+
+		/// <summary>
+		/// Zens documented exit codes
+		/// </summary>
+		private enum ZenExitCodes
+		{
+			Success = 0,
+			OtherError = 1,
+			BadInput = 2,
+			OutOfMemory = 16,
+			OutOfDisk = 17,
+			AssertError = 70,
+			HttpOtherClientError = 80,
+			HttpCantConnectError = 81,
+			HttpNotFound = 66,
+			HttpUnauthorized = 77,
+			HttpSLLError = 82,
+			HttpForbidden = 83,
+			HttpTimeout = 84,
+			HttpConflict = 85,
+			HttpNoHost = 86,
+			HttpOtherServerError = 90,
+			HttpInternalServerError = 91,
+			HttpServiceUnavailable = 69,
+			HttpBadGateway = 92,
+			HttpGatewayTimeout = 93
+		}
+
+		static readonly Regex s_logLinePattern = new Regex(@"\[.*\]\s?\[(?<severity>warning|error|info)\]\s?(\[(?<file>\S+)\:(?<line>\d+)\])?(?<message>.*)");
+		/// <summary>
+		/// CTOR
+		/// </summary>
+		public CloudStorage(IOptionsSnapshot<CloudStorageOptions> options, ILogger<CloudStorage> logger, IServiceProvider provider)
+		{
+			_options = options;
+			_logger = logger;
+			_provider = provider;
+		}
+
+		/// <inheritdoc />
+		public async Task<ICloudStorage.LoggedInStatus> LoginAsync(bool unattended = true, CloudBuildOptionsOverride? optionsOverride = null)
+		{
+			if (_options.Value.OidcTokenExePath == null && _options.Value.AccessToken == null)
+			{
+				return ICloudStorage.LoggedInStatus.LoginNotConfigured;
+			}
+
+			if (_options.Value.AccessToken != null)
+			{
+				// if the access token is set we assume its valid and thus say we are logged in
+				return ICloudStorage.LoggedInStatus.LoggedIn;
+			}
+
+			try
+			{
+				string? _ = await GetAccessTokenAsync(unattended, optionsOverride);
+			}
+			catch (NotLoggedInException)
+			{
+				return ICloudStorage.LoggedInStatus.InteractiveLoginRequired;
+			}
+			return ICloudStorage.LoggedInStatus.LoggedIn;
+		}
+
+		private async Task<string?> GetAccessTokenAsync(bool unattended = true, CloudBuildOptionsOverride? optionsOverride = null)
+		{
+			if (_options.Value.AccessToken != null)
+			{
+				return _options.Value.AccessToken;
+			}
+
+			if (optionsOverride?.AccessToken != null)
+			{
+				return optionsOverride.AccessToken;
+			}
+
+			if (optionsOverride?.DefaultHost == null && optionsOverride?.OverrideHost == null && _options.Value.Host == null)
+			{
+				throw new CloudStorageException("No host configured");
+			}
+
+			Uri host = new Uri(optionsOverride?.OverrideHost ?? optionsOverride?.DefaultHost ?? _options.Value.Host!);
+
+			// TODO: Should we support local file configuration?
+			ClientAuthConfigurationV1? authConfig = await ProviderConfigurationFactory.ReadRemoteAuthConfigurationAsync(host, ProviderConfigurationFactory.DefaultEncryptionKey);
+			if (authConfig?.DefaultProvider == null)
+			{
+				throw new CloudStorageException($"Failed to determine default auth provider from {host}");
+			}
+
+			IConfiguration config = ProviderConfigurationFactory.BindOptions(authConfig);
+			using ITokenStore tokenStore = _provider.GetService<ITokenStore>() ?? TokenStoreFactory.CreateTokenStore();
+			// TODO: Should we fetch the token manager from DI as well? Then we can not reconfigure it
+			OidcTokenManager tokenManager = OidcTokenManager.CreateTokenManager(config, tokenStore);
+
+			bool allowInteractiveLogin = !unattended || _options.Value.AllowAuthPrompt;
+			try
+			{
+				OidcTokenInfo tokenInfo = await tokenManager.GetAccessToken(authConfig.DefaultProvider);
+
+				return tokenInfo.AccessToken;
+			}
+			catch (NotLoggedInException)
+			{
+				if (!allowInteractiveLogin)
+				{
+					throw;
+				}
+
+				OidcTokenInfo tokenInfo = await tokenManager.LoginAsync(authConfig.DefaultProvider);
+
+				return tokenInfo.AccessToken;
+			}
+		}
+
+		/// <inheritdoc />
+		public Task DownloadBuildAsync(StringId bucketId, CbObjectId buildId, DirectoryReference outputDir, List<FilePattern>? includeFiles = null, List<FilePattern>? excludeFiles = null, CloudBuildDownloadOptions? downloadOptions = null, IProgress<(float, string)>? progress = null, CloudBuildOptionsOverride? optionsOverride = null, CancellationToken cancellationToken = default(CancellationToken))
+		{
+			NamespaceId? ns = _options.Value.NamespaceId != null ? new NamespaceId(_options.Value.NamespaceId) : null;
+			return DownloadBuildAsync(ns, bucketId, buildId, outputDir, includeFiles, excludeFiles, downloadOptions, progress, optionsOverride, cancellationToken);
+		}
+
+		/// <inheritdoc />
+		public Task DownloadBuildFromUriAsync(Uri uri, DirectoryReference outputDir, List<FilePattern>? includeFiles = null, List<FilePattern>? excludeFiles = null, CloudBuildDownloadOptions? downloadOptions = null, IProgress<(float, string)>? progress = null, CloudBuildOptionsOverride? optionsOverride = null, CancellationToken cancellationToken = default(CancellationToken))
+		{
+			string path = uri.AbsolutePath;
+			// remove api path if specified
+			path = path.Replace("api/v2/builds/", "", StringComparison.OrdinalIgnoreCase);
+			string[] components = path.Split('/', StringSplitOptions.RemoveEmptyEntries);
+			if (components.Length != 3)
+			{
+				_logger.LogError("Failed to base uri path {Uri}, did not match expected format of namespace/bucket/buildId", path);
+				throw new ArgumentException("Uri did not match expected format of namespace/bucket/buildId", nameof(uri));
+			}
+
+			string namespaceId = components[0];
+			string bucketId = components[1];	
+			string buildId = components[2];	
+
+			optionsOverride ??= new CloudBuildOptionsOverride();
+			optionsOverride.DefaultHost = uri.GetLeftPart(UriPartial.Authority);
+			return DownloadBuildAsync(new NamespaceId(namespaceId), new StringId(bucketId), CbObjectId.Parse(buildId), outputDir, includeFiles, excludeFiles, downloadOptions, progress, optionsOverride, cancellationToken);
+		}
+
+		private async Task DownloadBuildAsync(NamespaceId? namespaceId, StringId bucketId, CbObjectId buildId, DirectoryReference outputDir, List<FilePattern>? includeFiles, List<FilePattern>? excludeFiles, CloudBuildDownloadOptions? downloadOptions = null, IProgress<(float, string)>? progress = null, CloudBuildOptionsOverride? optionsOverride = null, CancellationToken cancellationToken = default(CancellationToken))
+		{
+			NamespaceId? ns = optionsOverride?.NamespaceId ?? namespaceId;
+			if (ns == null)
+			{
+				throw new ArgumentNullException(nameof(namespaceId), "No namespace specified when downloading build, make sure to set default namespace or override it for this download");
+			}
+			string http2Options = downloadOptions?.AssumeHttp2 ?? false ? "--assume-http2" : "";
+			string boostWorkersOption = downloadOptions?.BoostWorkers ?? false ? "--boost-workers" : "";
+			string allowMultipartOption = downloadOptions?.AllowMultipart != null ? $"--allow-multipart={downloadOptions.AllowMultipart}" : "";
+			// trace file option needs to be passed before other arguments including the verbs
+			string traceOption = downloadOptions?.TraceFile != null ? $"--tracefile={downloadOptions.TraceFile.FullName} " : "";
+			string zenFolderPath = downloadOptions?.ZenStateFolder != null ? $"--zen-folder-path={downloadOptions.ZenStateFolder.FullName} " : "";
+
+			string wildcardOption = includeFiles == null ? "" : $"--wildcard \"{String.Join(';', includeFiles.Select(p => ToZenWildcard(outputDir, p)))}\"";
+			string excludeOption =  excludeFiles == null ? "" : $"--exclude-wildcard \"{String.Join(';', excludeFiles.Select(p => ToZenWildcard(outputDir, p)))}\"";
+			// pass the access token via environment variable to avoid showing it on the cli
+			string cmdline = $"{traceOption}builds download {_options.Value.GetHostCliArgs(optionsOverride)} --namespace {ns} --bucket {bucketId} {wildcardOption} {excludeOption} {allowMultipartOption} {http2Options} {boostWorkersOption} {zenFolderPath} \"{outputDir}\" {buildId}";
+
+			int exitCode = await RunZenCliAsync(cmdline, progress: progress, optionsOverride: optionsOverride, cancellationToken: cancellationToken);
+			EnsureZenExitCodeSuccess(exitCode);
+		}
+
+		private static string ToZenWildcard(DirectoryReference outputDirectory, FilePattern filePattern)
+		{
+			IEnumerable<string> GetTokens(FilePattern pattern)
+			{
+				foreach (string token in pattern.Tokens)
+				{
+					// convert directory wildcards (...) to * as zen only has general wildcards
+					yield return token == "..." ? "*" : token;
+				}
+			}
+
+			return Path.Combine(filePattern.BaseDirectory.MakeRelativeTo(outputDirectory), String.Join(String.Empty, GetTokens(filePattern)));
+		}
+
+		/// <inheritdoc />
+		public async IAsyncEnumerable<FindBuildResponse> FindBuildsAsync(Regex bucketRegex, CbObject query, CloudBuildSearchOptions? searchOptions = null, CloudBuildOptionsOverride? optionsOverride = null, [EnumeratorCancellation] CancellationToken cancellationToken = default(CancellationToken))
+		{
+			string? ns = GetCurrentNamespace(optionsOverride);
+			if (ns == null)
+			{
+				throw new CloudStorageException("Unable to determine namespace to find builds for, make sure to set the default namespace or override it in the options override");
+			}
+
+			string tempInputFile = Path.Combine(Path.GetTempPath(), Path.GetRandomFileName());
+			string tempOutputFile = Path.Combine(Path.GetTempPath(), Path.GetRandomFileName());
+			tempInputFile = Path.ChangeExtension(tempInputFile, "cbo");
+			tempOutputFile = Path.ChangeExtension(tempOutputFile, "cbo");
+			try
+			{
+				// copy the query and append the bucket regex filter
+				CbWriter writer = new CbWriter();
+				writer.BeginObject();
+				writer.WriteObject("query", query);
+				writer.WriteString("bucketRegex", bucketRegex.ToString());
+				if (searchOptions != null)
+				{
+					CloudBuildSearchOptions options = searchOptions;
+					writer.BeginObject("options");
+					if (options.Skip.HasValue)
+					{
+						writer.WriteInteger("skip", options.Skip.Value);
+					}
+					if (options.Max.HasValue)
+					{
+						writer.WriteInteger("max", options.Max.Value);
+					}
+					if (options.Limit.HasValue)
+					{
+						writer.WriteInteger("limit", options.Limit.Value);
+					}
+					if (options.IncludeTTL.HasValue)
+					{
+						writer.WriteBool("includeTtl", options.IncludeTTL.Value);
+					}
+					writer.EndObject();
+				}
+				writer.EndObject();
+
+				await File.WriteAllBytesAsync(tempInputFile, writer.ToByteArray(), cancellationToken);
+
+				string cmdline = $"builds list {_options.Value.GetHostCliArgs(optionsOverride)} --namespace {ns} --query-path \"{tempInputFile}\" --result-path \"{tempOutputFile}\"";
+
+				int exitCode = await RunZenCliAsync(cmdline, optionsOverride: optionsOverride, cancellationToken: cancellationToken);
+
+				EnsureZenExitCodeSuccess(exitCode);
+				byte[] b = await File.ReadAllBytesAsync(tempOutputFile, cancellationToken);
+				CbObject o = new CbObject(b);
+
+				bool partialResult = o["partialResult"].AsBool();
+				if (partialResult)
+				{
+					// TODO: Use a much smaller options then default to trigger pagination and return reasonable page sizes always so we do not have this warning
+					_logger.LogWarning("Partial request result, not returning all possible values. Try tweaking your search options to allow for scanning more objects, reducing the wide your search goes or implementing pagination");
+				}
+
+				foreach (CbField result in o["results"].AsArray())
+				{
+					CbObject metadata = result["metadata"].AsObject();
+					if (metadata.Equals(CbObject.Empty))
+					{
+						// we require metadata to be able to know which commit this build belongs to
+						continue;
+					}
+
+					string buildName = metadata["name"].AsString();
+					string bucketIdString = result["bucketId"].AsString();
+					StringId.TryParse(new Utf8String(bucketIdString), 256, out StringId bucketId, out string? _);
+					string buildIdString = result["buildId"].AsString();
+					CbObjectId buildId = CbObjectId.Parse(buildIdString);
+					DateTime createdAt = result["createdAt"].AsDateTime();
+					CbField commitField = metadata["commit"];
+
+					string? commit = null;
+					if (commitField.Equals(CbField.Empty))
+					{
+						// commit field is missing
+						commit = null;
+					}
+					else if (commitField.IsInteger())
+					{
+						commit = commitField.AsInt64().ToString();
+					}
+					else if (commitField.IsFloat())
+					{
+						// due to jsons problematic number types we sometimes get doubles for integers, these are likely changelists and thus whole numbers so we just cast them
+						commit = ((long)commitField.AsDouble()).ToString();
+					}
+					else if (commitField.IsString())
+					{
+						commit = commitField.AsString();
+					}
+					
+					yield return new FindBuildResponse(buildName, buildId, bucketId, commit, createdAt, metadata);
+				}
+			}
+			finally
+			{
+				if (File.Exists(tempInputFile))
+				{
+					File.Delete(tempInputFile);
+				}
+
+				if (File.Exists(tempOutputFile))
+				{
+					File.Delete(tempOutputFile);
+				}
+			}
+		}
+
+		/// <inheritdoc />
+		public IAsyncEnumerable<FindBuildResponse> FindBuildsAsync(string project, string? type, string? branch, CbObject query, CloudBuildSearchOptions? searchOptions = null,
+			CloudBuildOptionsOverride? optionsOverride = null, CancellationToken cancellationToken = default(CancellationToken))
+		{
+			string regex = ICloudStorage.SanitizeBucketValue(project);
+			if (!String.IsNullOrEmpty(type))
+			{
+				regex += $"\\.{ICloudStorage.SanitizeBucketValue(type)}";
+			}
+			if (!String.IsNullOrEmpty(branch))
+			{
+				regex += $"\\.{ICloudStorage.SanitizeBucketValue(branch)}";
+			}
+			// wildcard after this prefix
+			regex += "\\.*";
+
+			return FindBuildsAsync(new Regex(regex), query, searchOptions, optionsOverride, cancellationToken);
+		}
+
+		private static ZenExitCodes AsZenExitCodeValue(int exitCode)
+		{
+			if (!Enum.IsDefined(typeof(ZenExitCodes), exitCode))
+			{
+				throw new NotImplementedException($"Unknown zen exit code: {exitCode}");
+			}
+
+			return (ZenExitCodes)exitCode;
+		}
+
+		/// <inheritdoc />
+		public async Task<List<BuildContents>> ListBuildContentsAsync(StringId bucketId, CbObjectId buildId, CloudBuildOptionsOverride? optionsOverride = null, CancellationToken cancellationToken = default(CancellationToken))
+		{
+			Uri uri = new Uri($"{_options.Value.Host}/api/v2/builds/{GetCurrentNamespace(optionsOverride)}/{bucketId}/{buildId}");
+			string cmdline = $"builds ls --cloud-url={uri} --quiet";
+			// TODO: we need a way to ask this to get written to file
+			int exitCode = await RunZenCliAsync(cmdline, cancellationToken: cancellationToken);
+
+			EnsureZenExitCodeSuccess(exitCode);
+
+			// TODO: we need structured output from zen for this listing
+
+			Dictionary<string, FileDescription> files = new Dictionary<string, FileDescription>();
+			/*foreach (string line in stdout)
+			{
+				string[] segments = line.Split('\t');
+				string fileName = segments[0];
+				ulong rawSize = UInt64.Parse(segments[1]);
+				IoHash rawHash = IoHash.Parse(segments[2]);
+
+				// TODO: Attributes are missing from zen output
+				files.Add(fileName, new FileDescription
+				{
+					RawSize = rawSize,
+					RawHash = rawHash,
+					Attributes = 0,
+					Mode = 0
+				});
+			}*/
+			List<BuildContents> contents =
+			[
+				new BuildContents
+				{
+					Parts = new Dictionary<string, PartContents>
+					{
+						{ "default", new PartContents() { PartName = "default", TotalSize = 0, Files = files } }
+					}
+				}
+			];
+			return contents;
+		}
+
+		private void EnsureZenExitCodeSuccess(int exitCode)
+		{
+			if (exitCode != 0)
+			{
+				string msg = AsZenExitCodeValue(exitCode).ToString();
+				_logger.LogError("Zen non-zero exit code {ExitCode} ({Message})", exitCode, msg);
+				throw new CloudStorageException($"Zen exited with non-success exit code: {exitCode} ({msg})");
+			}
+		}
+
+		/// <inheritdoc />
+		public Task Reconfigure(Action<CloudStorageOptions> action)
+		{
+			action(_options.Value);
+
+			// throws validation exceptions if there are any issues
+			Validator.ValidateObject(_options.Value, new ValidationContext(_options.Value), validateAllProperties: true);
+
+			return Task.CompletedTask;
+		}
+
+		/// <inheritdoc />
+		public string? GetCurrentNamespace(CloudBuildOptionsOverride? optionsOverride)
+		{
+			return optionsOverride?.NamespaceId?.ToString() ?? _options.Value.NamespaceId;
+		}
+
+		/// <inheritdoc />
+		public async Task ImportSnapshotAsync(string localFileName, string localRootPath, string engineDir, string cookedPlatformDir, bool forceImport = false, bool asyncImport = false, CancellationToken cancellationToken = default(CancellationToken))
+		{
+			/*string projectId = GetProjectID(projectFile);
+			CreateProject(projectFile, rootDirectory, engineDirectory);
+
+			platformCookedDirectory = new DirectoryReference(platformCookedDirectory.FullName.Replace("{Platform}", _buildData.Metadata.CookPlatform, StringComparison.Ordinal));
+			FileReference projectStoreFile = WriteProjectStoreFile(projectId, _buildData.Metadata.CookPlatform, platformCookedDirectory);
+
+			string oplogName = _buildData.Metadata.CookPlatform;
+
+			StringBuilder oplogCreateCommandLine = new StringBuilder();
+			oplogCreateCommandLine.AppendFormat("oplog-create --force-update {0} {1}", projectId, oplogName);
+
+			StringBuilder oplogImportCommandLine = new StringBuilder();
+			oplogImportCommandLine.AppendFormat("oplog-import {0} {1} {2} --ignore-missing-attachments --clean", projectId, oplogName, projectStoreFile.FullName);
+
+			if (forceImport)
+			{
+				oplogImportCommandLine.AppendFormat(" --force");
+			}
+
+			if (asyncImport)
+			{
+				oplogImportCommandLine.AppendFormat(" --async");
+			}
+
+			oplogImportCommandLine.AppendFormat(" --namespace {0}", _namespace);
+			oplogImportCommandLine.AppendFormat(" --access-token {0}", _accessToken);
+			oplogImportCommandLine.AppendFormat(" --bucket {0}", _bucket);
+			oplogImportCommandLine.AppendFormat(" --builds {0}", _host);
+			oplogImportCommandLine.AppendFormat(" --builds-id {0}", _buildData.BuildId);
+
+			await RunZenCliAsync(oplogCreateCommandLine.ToString());
+			await RunZenCliAsync(oplogImportCommandLine.ToString());*/
+
+			await Task.CompletedTask;
+			throw new NotImplementedException();
+		}
+
+		private async Task<int> RunZenCliAsync(string cmdline, IProgress<(float, string)>? progress = null, CloudBuildOptionsOverride? optionsOverride = null, CancellationToken cancellationToken = default)
+		{
+			FileReference zenExe = _options.Value.ZenExePath!;
+
+			string? accessToken = _options.Value.AccessToken;
+			if (_options.Value.FetchAccessTokenInProcess)
+			{
+				accessToken = await GetAccessTokenAsync(unattended: true, optionsOverride);
+			}
+
+			// append auth arguments
+			Dictionary<string, string>? env = null;
+			if (!String.IsNullOrEmpty(accessToken))
+			{
+				cmdline += " --access-token-env UE_CloudDataCacheAccessToken ";
+
+				env = new Dictionary<string, string>();
+				// forward the current environment and then override the token
+				foreach (DictionaryEntry entry in Environment.GetEnvironmentVariables())
+				{
+					if (entry is { Key: string k, Value: string v })
+					{
+						env[k] = v;
+					}
+				}
+
+				env.Add("UE_CloudDataCacheAccessToken", accessToken);
+			}
+			else if (_options.Value.OidcTokenExePath != null)
+			{
+				string unattendedOption = _options.Value.AllowAuthPrompt ? "--oidctoken-exe-unattended=true" : "";
+				cmdline += $" --oidctoken-exe-path=\"{_options.Value.OidcTokenExePath}\" {unattendedOption}";
+			}
+
+			// if we have a progress callback with to the log progress mode
+			if (progress != null)
+			{
+				cmdline += " --log-progress";
+			}
+			else
+			{
+				// when no progress reporting is requested we use the plain progress which simply outputs console output about the progress in a way that works well with redirected output
+				cmdline += " --plain-progress";
+			}
+
+			_logger.Log(LogLevel.Information, "Running: {App} with {CommandLine}", zenExe.FullName, cmdline);
+
+			ILogger? zenLogger = null;
+
+			if (_options.Value.EchoZenCliToLog && optionsOverride?.EchoLogger == null)
+			{
+				ILoggerFactory loggerFactory = _provider.GetRequiredService<ILoggerFactory>();
+				zenLogger = loggerFactory.CreateLogger("ZenCli");
+			}
+			else if (_options.Value.EchoZenCliToLog)
+			{
+				zenLogger = optionsOverride?.EchoLogger;
+			}
+
+			const int MaxRetryAttempts = 3;
+			const int WarnAtAttempts = 2;
+			for (int retryAttempt = 0; retryAttempt < MaxRetryAttempts; retryAttempt++)
+			{
+				int exitCode = await RunProcessAsync(zenExe.FullName, cmdline, env: env, echoStdout: _options.Value.EchoZenCliToStdout, echoLogger: zenLogger, progress: progress, cancellationToken: cancellationToken);
+
+				// increasingly long time to wait before retrying
+				int[] waitDurationSeconds = [60, 300, 900];
+				int waitDuration = waitDurationSeconds[retryAttempt];
+				
+				ZenExitCodes zenExitCode = AsZenExitCodeValue(exitCode);
+				// retry on zen exit codes that indicate transient issues
+				if (zenExitCode is ZenExitCodes.HttpServiceUnavailable or ZenExitCodes.HttpBadGateway or ZenExitCodes.HttpGatewayTimeout)
+				{
+					if (retryAttempt + 1 == MaxRetryAttempts)
+					{
+						// give up
+						_logger.LogError("Zen returned http status code {ExitCode}({ExitCodeMessage}) after {RetryAttempts} retries. Giving up.", exitCode, AsZenExitCodeValue(exitCode), retryAttempt);
+
+						// we return the exit code even though it's a failure to allow the caller to handle this issue that isn't fixed by retrying
+					}
+					else
+					{
+						if (retryAttempt + 1 == WarnAtAttempts)
+						{
+							_logger.LogWarning("Zen returned http status code {ExitCode}({ExitCodeMessage}) after {RetryAttempts} retries. This is a indication of an infrastructure issue.", exitCode, AsZenExitCodeValue(exitCode), retryAttempt);
+						}
+						else
+						{
+							_logger.LogInformation("Http Status code {ExitCode}({ExitCodeMessage}) returned from zen. Waiting {WaitDuration}s before trying.", exitCode, AsZenExitCodeValue(exitCode), waitDuration);
+						}
+						await Task.Delay(TimeSpan.FromSeconds(waitDuration), cancellationToken);
+						continue; // retry
+					}
+				}
+
+				return exitCode;
+			}
+
+			throw new CloudStorageException($"Failed to run zen cli, gave up after {MaxRetryAttempts} attempts");
+		}
+
+		private async Task<int> RunProcessAsync(string cmd, string arguments, string? workingDirectory = null, Dictionary<string, string>? env = null, IProgress<(float, string)>? progress = null, bool echoStdout = false, ILogger? echoLogger = null, CancellationToken cancellationToken = default)
+		{
+			using ManagedProcessGroup managedProcessGroup = new ManagedProcessGroup();
+			using ManagedProcess managedProcess = new(managedProcessGroup, cmd, arguments, workingDirectory, env, ProcessPriorityClass.Normal, flags: ManagedProcessFlags.MergeOutputPipes);
+
+			_options.Value.OnProcessCreated?.Invoke(managedProcess);
+
+			ProgressValue progressValue = new ProgressValue();
+			while (managedProcess.TryReadLine(out string? line, cancellationToken))
+			{
+				if (progress != null && line.StartsWith(ProgressTextReader.DirectivePrefix, StringComparison.Ordinal))
+				{
+					line = ProgressTextReader.ParseLine(line, progressValue);
+					progress.Report((progressValue.Current.Item2, progressValue.Current.Item1));
+
+					if (String.IsNullOrEmpty(line))
+					{
+						// line only had progress reporting, move on
+						continue;
+					}
+				}
+
+				if (echoLogger != null)
+				{
+					EchoToLogger(echoLogger, line);
+				}
+				if (echoStdout)
+				{
+					Console.WriteLine(line);
+				}
+			}
+	
+			cancellationToken.ThrowIfCancellationRequested();
+
+			await managedProcess.WaitForExitAsync(cancellationToken);
+			_options.Value.OnProcessExited?.Invoke(managedProcess);
+
+			cancellationToken.ThrowIfCancellationRequested();
+
+			return managedProcess.ExitCode;
+		}
+
+		private static void EchoToLogger(ILogger zenLogger, string line)
+		{
+			// TODO: Ideally zen cli has structured logging here that we can parse and output back as a structured logging event, for now we just text match
+			Match match = s_logLinePattern.Match(line);
+			if (match.Success)
+			{
+				// we ignore the file matching for now
+				//string file = match.Groups["file"].Value;
+				//string lineNumber = match.Groups["line"].Value;
+				string message = match.Groups["message"].Value;
+
+				LogLevel level = match.Groups["severity"].Value switch
+				{
+					"error" => LogLevel.Error,
+					"warning" => LogLevel.Warning,
+					_ => LogLevel.Information,
+				};
+
+#pragma warning disable CA2254
+				zenLogger.Log(level, message);
+#pragma warning restore CA2254
+			}
+			else
+			{
+				// didn't match the structured output
+#pragma warning disable CA2254
+				zenLogger.LogInformation(line);
+#pragma warning restore CA2254
+			}
+		}
+	}
+}
diff --git a/Engine/Source/Programs/Shared/EpicGames.ProjectStore/CloudStorage/CloudStorageExtensions.cs b/Engine/Source/Programs/Shared/EpicGames.ProjectStore/CloudStorage/CloudStorageExtensions.cs
new file mode 100644
--- /dev/null
+++ b/Engine/Source/Programs/Shared/EpicGames.ProjectStore/CloudStorage/CloudStorageExtensions.cs
@@ -0,0 +1,38 @@
+// Copyright Epic Games, Inc. All Rights Reserved.
+
+using System;
+using Microsoft.Extensions.DependencyInjection;
+
+namespace EpicGames.ProjectStore
+{
+	/// <summary>
+	/// Extension methods for Cloud Storage
+	/// </summary>
+	public static class CloudStorageExtensions
+	{
+		/// <summary>
+		/// Adds Cloud Storage-related services with the default settings
+		/// </summary>
+		/// <param name="serviceCollection">Collection to register services with</param>
+		public static void AddCloudStorage(this IServiceCollection serviceCollection)
+		{
+			// This is registered as transient before of the Reconfigure call which will mutate the options of a ICloudStorage
+			// That option is required due to some cases where configuration is read after all our services would normally be configured
+			// This is not a very big deal as CloudStorage itself has limited local state
+			serviceCollection.AddTransient<ICloudStorage, CloudStorage>();
+
+			serviceCollection.AddOptions<CloudStorageOptions>();
+		}
+
+		/// <summary>
+		/// Adds Cloud Storage-related services with configuration
+		/// </summary>
+		/// <param name="serviceCollection">Collection to register services with</param>
+		/// <param name="configureCloudStorage">Callback to configure options</param>
+		public static void AddCloudStorage(this IServiceCollection serviceCollection, Action<CloudStorageOptions> configureCloudStorage)
+		{
+			serviceCollection.Configure(configureCloudStorage);
+			AddCloudStorage(serviceCollection);
+		}
+	}
+}
diff --git a/Engine/Source/Programs/Shared/EpicGames.ProjectStore/CloudStorage/CloudStorageQueryBuilder.cs b/Engine/Source/Programs/Shared/EpicGames.ProjectStore/CloudStorage/CloudStorageQueryBuilder.cs
new file mode 100644
--- /dev/null
+++ b/Engine/Source/Programs/Shared/EpicGames.ProjectStore/CloudStorage/CloudStorageQueryBuilder.cs
@@ -0,0 +1,200 @@
+// Copyright Epic Games, Inc. All Rights Reserved.
+
+using System;
+using System.Collections.Generic;
+using EpicGames.Serialization;
+
+namespace EpicGames.ProjectStore
+{
+	/// <summary>
+	/// Helper class to build a query for finding builds using cloud storage
+	/// </summary>
+	public class CloudStorageQueryBuilder
+	{
+		private readonly CbWriter _writer = new CbWriter();
+		private CbObject? _generatedObject = null;
+
+		/// <summary>
+		/// CTOR
+		/// </summary>
+		public CloudStorageQueryBuilder()
+		{
+			// open the writer so that we can start appending request to it
+			_writer.BeginObject();
+		}
+
+		/// <summary>
+		/// Search for exact value in a field
+		/// </summary>
+		/// <typeparam name="T">Type</typeparam>
+		/// <param name="fieldName">The field name</param>
+		/// <param name="value">The value to filter for</param>
+		public void Equals<T>(string fieldName, T value)
+		{
+			ThrowOnLocked();
+
+			_writer.BeginObject(fieldName);
+			WriteValue("$eq", value);
+			_writer.EndObject();
+		}
+
+		/// <summary>
+		/// Search for any value except value in a field
+		/// </summary>
+		/// <typeparam name="T">Type</typeparam>
+		/// <param name="fieldName">The field name</param>
+		/// <param name="value">The value to filter for</param>
+		public void NotEquals<T>(string fieldName, T value)
+		{
+			ThrowOnLocked();
+
+			_writer.BeginObject(fieldName);
+			WriteValue("$neq", value);
+			_writer.EndObject();
+		}
+
+		/// <summary>
+		/// Search for any value less than, type must be able to be compared
+		/// </summary>
+		/// <typeparam name="T">Type</typeparam>
+		/// <param name="fieldName">The field name</param>
+		/// <param name="value">The value to filter for</param>
+		public void LessThen<T>(string fieldName, T value)
+		{
+			ThrowOnLocked();
+
+			_writer.BeginObject(fieldName);
+			WriteValue("$lt", value);
+			_writer.EndObject();
+		}
+
+		/// <summary>
+		/// Search for any value less than or equal, type must be able to be compared
+		/// </summary>
+		/// <typeparam name="T">Type</typeparam>
+		/// <param name="fieldName">The field name</param>
+		/// <param name="value">The value to filter for</param>
+		public void LessThenOrEqual<T>(string fieldName, T value)
+		{
+			ThrowOnLocked();
+
+			_writer.BeginObject(fieldName);
+			WriteValue("$lte", value);
+			_writer.EndObject();
+		}
+
+		/// <summary>
+		/// Search for any value greater than, type must be able to be compared
+		/// </summary>
+		/// <typeparam name="T">Type</typeparam>
+		/// <param name="fieldName">The field name</param>
+		/// <param name="value">The value to filter for</param>
+		public void GreaterThen<T>(string fieldName, T value)
+		{
+			ThrowOnLocked();
+
+			_writer.BeginObject(fieldName);
+			WriteValue("$gt", value);
+			_writer.EndObject();
+		}
+
+		/// <summary>
+		/// Search for any value greater than or equal, type must be able to be compared
+		/// </summary>
+		/// <typeparam name="T">Type</typeparam>
+		/// <param name="fieldName">The field name</param>
+		/// <param name="value">The value to filter for</param>
+		public void GreaterThenOrEqual<T>(string fieldName, T value)
+		{
+			ThrowOnLocked();
+
+			_writer.BeginObject(fieldName);
+			WriteValue("$gte", value);
+			_writer.EndObject();
+		}
+
+		/// <summary>
+		/// Search for any of the exact values in a list
+		/// </summary>
+		/// <typeparam name="T">Type</typeparam>
+		/// <param name="fieldName">The field name</param>
+		/// <param name="values">The values to filter for</param>
+		public void In<T>(string fieldName, IEnumerable<T> values)
+		{
+			ThrowOnLocked();
+
+			_writer.BeginObject(fieldName);
+			_writer.BeginArray("$in");
+			foreach (T value in values)
+			{
+				WriteValue(null, value);
+			}
+			_writer.EndArray();
+			_writer.EndObject();
+		}
+
+		/// <summary>
+		/// Search for any of the exact values to not be in this list
+		/// </summary>
+		/// <typeparam name="T">Type</typeparam>
+		/// <param name="fieldName">The field name</param>
+		/// <param name="values">The values to filter for</param>
+		public void NotIn<T>(string fieldName, IEnumerable<T> values)
+		{
+			ThrowOnLocked();
+
+			_writer.BeginObject(fieldName);
+			_writer.BeginArray("$nin");
+			foreach (T value in values)
+			{
+				WriteValue(null, value);
+			}
+			_writer.EndArray();
+			_writer.EndObject();
+		}
+
+		/// <summary>
+		/// Returns the built query and locks the builder to prevent further appended filters
+		/// </summary>
+		/// <returns></returns>
+		public CbObject Done()
+		{
+			if (_generatedObject != null)
+			{
+				return _generatedObject;
+			}
+
+			_writer.EndObject();
+			_generatedObject = _writer.ToObject();
+
+			return _generatedObject;
+		}
+		
+		private void WriteValue<T>(string? fieldName, T value)
+		{
+			CbConverter<T> converter = CbConverter.GetConverter<T>();
+
+			if (converter == null)
+			{
+				throw new NotImplementedException($"Failed to find a converter for type \"{typeof(T)}\", you may want to try to use a basic type like string");
+			}
+
+			if (fieldName != null)
+			{
+				converter.WriteNamed(_writer, fieldName, value);
+			}
+			else
+			{
+				converter.Write(_writer, value);
+			}
+		}
+
+		private void ThrowOnLocked()
+		{
+			if (_generatedObject != null)
+			{
+				throw new Exception("Query builder is locked, more filters can not be added after the object has been fetched");
+			}
+		}
+	}
+}
diff --git a/Engine/Source/Programs/Shared/EpicGames.ProjectStore/CloudStorage/ICloudStorage.cs b/Engine/Source/Programs/Shared/EpicGames.ProjectStore/CloudStorage/ICloudStorage.cs
new file mode 100644
--- /dev/null
+++ b/Engine/Source/Programs/Shared/EpicGames.ProjectStore/CloudStorage/ICloudStorage.cs
@@ -0,0 +1,504 @@
+// Copyright Epic Games, Inc. All Rights Reserved.
+
+using System;
+using System.Collections.Generic;
+using System.ComponentModel.DataAnnotations;
+using System.IO;
+using System.Text.RegularExpressions;
+using System.Threading;
+using System.Threading.Tasks;
+using EpicGames.Core;
+using EpicGames.Horde;
+using EpicGames.Horde.Storage;
+using EpicGames.Serialization;
+using Microsoft.Extensions.Logging;
+using UnrealBuildBase;
+
+namespace EpicGames.ProjectStore
+{
+	/// <summary>
+	/// A build found during search
+	/// </summary>
+	/// <param name="Name">The name of the build</param>
+	/// <param name="BuildId">The unique id of the build</param>
+	/// <param name="BucketId">The bucket the build belongs to</param>
+	/// <param name="Commit">The commit of the build, optional and can be a changelist, hash or other global unique value</param>
+	/// <param name="CreatedAt">The time this build was created</param>
+	/// <param name="Metadata">Metadata associated with the build</param>
+	public readonly record struct FindBuildResponse(
+		string Name,
+		CbObjectId BuildId,
+		StringId BucketId,
+		string? Commit,
+		DateTime CreatedAt,
+		CbObject Metadata)
+	{
+		/// <summary>
+		/// The globally unique id of the build
+		/// </summary>
+		public CbObjectId BuildId { get; init; } = BuildId;
+
+		/// <summary>
+		/// Which bucket the build is in
+		/// </summary>
+		public StringId BucketId { get; init; } = BucketId;
+
+		/// <summary>
+		/// The commit identifier as defined in source control typically an incrementing number (changelist) or hash (git commit), can be null if missing
+		/// </summary>
+		public string? Commit { get; init; } = Commit;
+
+		/// <summary>
+		/// The name of the build
+		/// </summary>
+		public string Name { get; init; } = Name;
+
+		/// <summary>
+		/// When the build was created
+		/// </summary>
+		public DateTime CreatedAt { get; init; } = CreatedAt;
+
+		/// <summary>
+		/// The metadata of the build
+		/// </summary>
+		public CbObject Metadata { get; init; } = Metadata;
+	}
+
+	/// <summary>
+	/// The contents available for a part
+	/// </summary>
+	public readonly record struct PartContents(string PartName, ulong TotalSize, Dictionary<string, FileDescription> Files)
+	{
+		/// <summary>
+		/// The name of the part
+		/// </summary>
+		public required string PartName { get; init; } = PartName;
+
+		/// <summary>
+		/// The total uncompressed size of the part
+		/// </summary>
+		public required ulong TotalSize { get; init; } = TotalSize;
+
+		/// <summary>
+		/// All the files in the part
+		/// </summary>
+		public required Dictionary<string, FileDescription> Files { get; init; } = Files;
+	}
+
+	/// <summary>
+	/// Description of an individual file
+	/// </summary>
+	public readonly record struct FileDescription(IoHash RawHash, ulong RawSize, int Attributes, int Mode)
+	{
+		/// <summary>
+		/// Hash of uncompressed content
+		/// </summary>
+		public required IoHash RawHash { get; init; } = RawHash;
+
+		/// <summary>
+		/// Size of the file before compression
+		/// </summary>
+		public required ulong RawSize { get; init; } = RawSize;
+
+		/// <summary>
+		/// Windows file attributes applied to the file
+		/// </summary>
+		public required int Attributes { get; init; } = Attributes;
+
+		/// <summary>
+		/// Linux/Mac attributes applied to the file
+		/// </summary>
+		public required int Mode { get; init; } = Mode;
+	}
+
+	/// <summary>
+	/// The contents available for a build
+	/// </summary>
+	public record struct BuildContents
+	{
+		/// <summary>
+		/// Each part of the build
+		/// </summary>
+		public required Dictionary<string, PartContents> Parts { get; init; }
+	}
+
+	/// <summary>
+	/// Options passed to the search
+	/// </summary>
+	public record CloudBuildSearchOptions
+	{
+		/// <summary>
+		/// Max number of builds returned
+		/// </summary>
+		public int? Limit { get; set; }
+
+		/// <summary>
+		/// Number of builds to skip before returning result, useful to paginate results
+		/// </summary>
+		public int? Skip { get; set; }
+
+		/// <summary>
+		/// Max number of builds considered in the search request, you will usually want Limit rather then Max
+		/// </summary>
+		public int? Max { get; set; }
+
+		/// <summary>
+		/// Set to true to include how long the build will live, makes the search slower
+		/// </summary>
+		public bool? IncludeTTL { get; set; }
+	}
+
+	/// <summary>
+	/// Per command overrides of common options
+	/// </summary>
+	public record CloudBuildOptionsOverride
+	{
+		/// <summary>
+		/// Override which host is being used by default for this command, this host is used with service discovery to determine optimal host for this invocation
+		/// </summary>
+		public string? DefaultHost { get; set; }
+
+		/// <summary>
+		/// Explicitly define which host to use ignoring any service discovery, only use this is you explicitly want to hit a specific host
+		/// </summary>
+		public string? OverrideHost { get; set; }
+		
+		/// <summary>
+		/// Override which namespace is being used
+		/// </summary>
+		public NamespaceId? NamespaceId { get; set; }
+
+		/// <summary>
+		/// Override the access token used for auth
+		/// </summary>
+		public string? AccessToken { get; set; }
+
+		/// <summary>
+		/// Logger to echo output to, echo log to output option must also be set. If this logger is not provided but that option is set a global ZenCli logger is used instead
+		/// </summary>
+		public ILogger? EchoLogger { get; set; }
+	}
+	
+	/// <summary>
+	/// Advanced options for downloading builds
+	/// </summary>
+	public record CloudBuildDownloadOptions
+	{
+		/// <summary>
+		/// Set to override if we allow usage of multipart downloads
+		/// </summary>
+		public bool? AllowMultipart { get; set; }
+
+		/// <summary>
+		/// Set to opt out of http upgrade - needed if using http2 without TLS
+		/// </summary>
+		public bool AssumeHttp2 { get; set; }
+
+		/// <summary>
+		/// Set to boost the number of worker threads used by zen. Will usually speed up downloads if there is CPU capacity available but can cause you machine to become unresponsive
+		/// </summary>
+		public bool BoostWorkers { get; set; }
+
+		/// <summary>
+		/// Optional path to an utrace file, if specified Zen will output its utrace file here
+		/// </summary>
+		public FileReference? TraceFile { get; set; }
+
+		/// <summary>
+		/// Override directory location for the zen state folder
+		/// </summary>
+		public DirectoryReference? ZenStateFolder { get; set; }
+	}
+
+	/// <summary>
+	/// General exception thrown from Cloud Storage
+	/// </summary>
+	/// <param name="msg">Error message</param>
+	public class CloudStorageException(string msg) : Exception(msg);
+
+	/// <summary>
+	/// Interface for interacting with UE Cloud Storage
+	/// </summary>
+	public interface ICloudStorage
+	{
+		/// <summary>
+		/// The possible login states
+		/// </summary>
+		enum LoggedInStatus
+		{
+			/// <summary>
+			/// Has an active access token
+			/// </summary>
+			LoggedIn, 
+			/// <summary>
+			/// No access token found, will require interactive login
+			/// </summary>
+			InteractiveLoginRequired,
+
+			/// <summary>
+			/// Unattended login specified but no login configuration given
+			/// </summary>
+			LoginNotConfigured
+		}
+
+		/// <summary>
+		/// Login against cloud storage, will trigger an interactive login unless set to unattended
+		/// </summary>
+		/// <param name="unattended">Set as false to allow interactive logins</param>
+		/// <param name="optionsOverride">Optional per invocation overrides of options</param>
+		/// <returns></returns>
+		Task<LoggedInStatus> LoginAsync(bool unattended = true, CloudBuildOptionsOverride? optionsOverride = null);
+
+		/// <summary>
+		/// Download a build using all the required fields
+		/// </summary>
+		/// <param name="bucketId">The bucket</param>
+		/// <param name="buildId">The id of the build</param>
+		/// <param name="outputDir">The directory to write the build to</param>
+		/// <param name="includeFiles">List of file patterns of files to download, if specified only files matching these patterns are downloaded</param>
+		/// <param name="excludeFiles">List of file patterns to exclude. If combined with include files these apply after the include files</param>
+		/// <param name="downloadOptions">Advanced options for reconfiguring the download</param>
+		/// <param name="progress">Progress reporting</param>
+		/// <param name="optionsOverride">Optional per invocation overrides of options</param>
+		/// <param name="cancellationToken">A cancellation token</param>
+		/// <returns></returns>
+		Task DownloadBuildAsync(StringId bucketId, CbObjectId buildId, DirectoryReference outputDir, List<FilePattern>? includeFiles = null, List<FilePattern>? excludeFiles = null, CloudBuildDownloadOptions? downloadOptions = null, IProgress<(float, string)>? progress = null, CloudBuildOptionsOverride? optionsOverride = null, CancellationToken cancellationToken = default(CancellationToken));
+
+		/// <summary>
+		/// Download a build using a single uri value that encodes everything needed to locate the build
+		/// </summary>
+		/// <param name="uri">The uri to download from</param>
+		/// <param name="outputDir">The directory to write the build to</param>
+		/// <param name="includeFiles">List of file patterns of files to download, if specified only files matching these patterns are downloaded</param>
+		/// <param name="excludeFiles">List of file patterns to exclude. If combined with include files these apply after the include files</param>
+		/// <param name="downloadOptions">Advanced options for reconfiguring the download</param>
+		/// <param name="progress">Progress reporting</param>
+		/// <param name="optionsOverride">Optional per invocation overrides of options</param>		
+		/// <param name="cancellationToken">A cancellation token</param>
+		/// <returns></returns>
+		Task DownloadBuildFromUriAsync(Uri uri, DirectoryReference outputDir, List<FilePattern>? includeFiles = null, List<FilePattern>? excludeFiles = null,  CloudBuildDownloadOptions? downloadOptions = null,  IProgress<(float, string)>? progress = null, CloudBuildOptionsOverride? optionsOverride = null, CancellationToken cancellationToken = default(CancellationToken));
+
+		/// <summary>
+		/// Search for builds in UE Cloud Storage using a regular expression to describe which buckets the build might be in
+		/// </summary>
+		/// <param name="bucketRegex">A regex that describe which buckets the build might be in</param>
+		/// <param name="query">CbObject that describes the query, see UE Cloud Storage documentation</param>
+		/// <param name="searchOptions">Optional set of options for the search</param>
+		/// <param name="optionsOverride">Optional per invocation overrides of options</param>		
+		/// <param name="cancellationToken">A cancellation token</param>
+		/// <returns></returns>
+		IAsyncEnumerable<FindBuildResponse> FindBuildsAsync(Regex bucketRegex, CbObject query, CloudBuildSearchOptions? searchOptions = null, CloudBuildOptionsOverride? optionsOverride = null, CancellationToken cancellationToken = default(CancellationToken));
+
+		/// <summary>
+		/// Search for builds in UE Cloud Storage using common components
+		/// </summary>
+		/// <param name="project">The project you are searching builds for</param>
+		/// <param name="type">The type of build you are searching for, if omitted all types are considered</param>
+		/// <param name="branch">The branch you which to search, if omitted all branches are searched</param>
+		/// <param name="query">CbObject that describes the query, see UE Cloud Storage documentation</param>
+		/// <param name="searchOptions">Optional set of options for the search</param>
+		/// <param name="optionsOverride">Optional per invocation overrides of options</param>		
+		/// <param name="cancellationToken">A cancellation token</param>
+		/// <returns></returns>
+		IAsyncEnumerable<FindBuildResponse> FindBuildsAsync(string project, string? type, string? branch, CbObject query, CloudBuildSearchOptions? searchOptions = null, CloudBuildOptionsOverride? optionsOverride = null, CancellationToken cancellationToken = default(CancellationToken));
+
+		/// <summary>
+		/// List the contents of a build
+		/// </summary>
+		/// <param name="bucketId">The bucket the build is in</param>
+		/// <param name="buildId">The unique id of the build</param>
+		/// <param name="optionsOverride">Optional per invocation overrides of options</param>		
+		/// <param name="cancellationToken">A cancellation token</param>
+		/// <returns></returns>
+		Task<List<BuildContents>> ListBuildContentsAsync(StringId bucketId, CbObjectId buildId, CloudBuildOptionsOverride? optionsOverride = null, CancellationToken cancellationToken = default(CancellationToken));
+
+		/// <summary>
+		/// Override the configuration of this instance of cloud storage, mostly useful if you get configuration from another source then the default inis / environment variables
+		/// </summary>
+		/// <param name="action">Action that reconfigures the options</param>
+		/// <returns></returns>
+		Task Reconfigure(Action<CloudStorageOptions> action);
+
+		/// <summary>
+		/// Helper method to turn a bucket component into a value that is compatible with cloud storage bucket names
+		/// </summary>
+		/// <param name="s">The value to clean up</param>
+		/// <returns></returns>
+		static string SanitizeBucketValue(string s)
+		{
+			// bucket components are lower cased
+			s = s.ToLower();
+			// trim starting + and then replace any + and . with -
+			s = s.TrimStart('+');
+			s = s.Replace("+", "-", StringComparison.OrdinalIgnoreCase);
+			// . have special meaning in the bucket names, its used as a separator between components, as such component can not themselves contain .
+			return s.Replace(".", "-", StringComparison.OrdinalIgnoreCase);
+		}
+
+		/// <summary>
+		/// Returns the current namespace used
+		/// </summary>
+		/// <param name="optionsOverride">Optional per invocation overrides of options</param>	
+		/// <returns></returns>
+		string? GetCurrentNamespace(CloudBuildOptionsOverride? optionsOverride);
+
+		/// <summary>
+		/// Import a snapshot to the local zen
+		/// </summary>
+		/// <param name="localFileName"></param>
+		/// <param name="localRootPath"></param>
+		/// <param name="engineDir"></param>
+		/// <param name="cookedPlatformDir"></param>
+		/// <param name="forceImport"></param>
+		/// <param name="asyncImport"></param>
+		/// <returns></returns>
+		Task ImportSnapshotAsync(string localFileName, string localRootPath, string engineDir, string cookedPlatformDir, bool forceImport = false, bool asyncImport = false, CancellationToken cancellationToken = default(CancellationToken));
+	}
+
+	/// <summary>
+	/// Options for cloud storage
+	/// </summary>
+	public record CloudStorageOptions : IValidatableObject
+	{
+		/// <summary>
+		/// Set too true to allow auth prompts, recommended to set to false for un-interactive scenarios
+		/// </summary>
+		public bool AllowAuthPrompt { get; set; } = false;
+
+		/// <summary>
+		/// The host name to use to connect to cloud storage
+		/// </summary>
+		[Url]
+		[Required]
+		public string? Host { get; set; }
+
+		/// <summary>
+		///  The namespace to use when connecting to cloud storage
+		/// </summary>
+		[Required]
+		public string? NamespaceId { get; set; }
+
+		/// <summary>
+		/// The Baseline branch to use for builds within cloud storage
+		/// </summary>
+		[Required]
+		public string? BaselineBranch { get; set; }
+
+		/// <summary>
+		/// The path to the zen cli executable
+		/// </summary>
+		[Required]
+		public FileReference? ZenExePath { get; set; }
+
+		/// <summary>
+		/// The path to the OidcToken executable
+		/// </summary>
+		public FileReference? OidcTokenExePath { get; set; }
+
+		/// <summary>
+		/// The explicit access token to use when communicating with Cloud Storage, it's usually best to not set this and specify the OidcTokenExePath instead to allow for refreshing of tokens if the background task end up long-running
+		/// </summary>
+		public string? AccessToken { get; set; } = null;
+
+		/// <summary>
+		/// Set to fetch access token via EpicGames.Oidc instead of using the OidcToken executable. This lets you build a tool that doesn't need to ship the token executable but prevents zen from refreshing the access token as needed so long-running operations can fail
+		/// </summary>
+		public bool FetchAccessTokenInProcess { get; set; } = false;
+
+		/// <summary>
+		/// Callback when a process has been created
+		/// </summary>
+		public Action<ManagedProcess>? OnProcessCreated { get; set; } = null;
+
+		/// <summary>
+		/// Callback when a process we created has exited
+		/// </summary>
+		public Action<ManagedProcess>? OnProcessExited { get; set; } = null;
+
+		/// <summary>
+		/// Set to true to echo all zencli output to stdout
+		/// </summary>
+		public bool EchoZenCliToStdout { get; set; } = false;
+
+		/// <summary>
+		/// Set to true to echo all zen cli logs to a Microsoft.Extensions.Logger called "ZenCli"
+		/// </summary>
+		public bool EchoZenCliToLog { get; set; } = false;
+
+		public IEnumerable<ValidationResult> Validate(ValidationContext validationContext)
+		{
+			List<ValidationResult> results = new List<ValidationResult>();
+			if (String.IsNullOrEmpty(AccessToken) && OidcTokenExePath == null && FetchAccessTokenInProcess == false)
+			{
+				results.Add(new ValidationResult("AccessToken or OidcTokenExePath needs to be defined for auth to function"));
+			}
+
+			if (FetchAccessTokenInProcess == false && OidcTokenExePath != null && !FileReference.Exists(OidcTokenExePath))
+			{
+				results.Add(new ValidationResult($"Unable to find OidcToken at path \"{OidcTokenExePath}\""));
+			}
+
+			if (ZenExePath != null && !FileReference.Exists(ZenExePath))
+			{
+				results.Add(new ValidationResult($"Unable to find Zen executable at path \"{ZenExePath}\""));
+			}
+
+			return results;
+		}
+
+		public void FindExecutablesFromUnrealRoot()
+		{
+			string binaryType = "Win64";
+			if (OperatingSystem.IsMacOS())
+			{
+				binaryType = "Mac";
+			}
+			else if (OperatingSystem.IsLinux())
+			{
+				binaryType = "Linux";
+			}
+
+			DirectoryReference engineBinariesDir = DirectoryReference.Combine(Unreal.RootDirectory, "Engine", "Binaries", binaryType);
+			FileReference zenExe = new FileReference(Path.Combine(engineBinariesDir.FullName, "zen" + RuntimePlatform.ExeExtension));
+			ZenExePath = zenExe;
+
+			string binaryTypeDotnet = "win-x64";
+			if (OperatingSystem.IsMacOS())
+			{
+				binaryTypeDotnet = "osx-x64";
+			}
+			else if (OperatingSystem.IsLinux())
+			{
+				binaryTypeDotnet = "linux-x64";
+			}
+
+			DirectoryReference oidcDir = DirectoryReference.Combine(Unreal.RootDirectory, "Engine", "Binaries", "DotNET", "OidcToken", binaryTypeDotnet);
+			FileReference oidcExe = new FileReference(Path.Combine(oidcDir.FullName, "OidcToken" + RuntimePlatform.ExeExtension));
+			
+			OidcTokenExePath = oidcExe;
+		}
+
+		/// <summary>
+		/// Returns the zen cli arguments to use for this set of host arguments
+		/// </summary>
+		/// <param name="optionsOverride">The options overrides used</param>
+		/// <returns></returns>
+		internal string GetHostCliArgs(CloudBuildOptionsOverride? optionsOverride)
+		{
+			string? defaultHost = optionsOverride?.DefaultHost ?? Host;
+			string args = "";
+
+			if (!String.IsNullOrEmpty(defaultHost))
+			{
+				args += $"--host {defaultHost}";
+			}
+			
+			if (optionsOverride?.OverrideHost != null)
+			{
+				args += $" --override-host {optionsOverride.OverrideHost}";
+			}
+
+			return args;
+		}
+	}
+}
