Cloud Storage - Fixed static code analysis warnings

diff --git a/Engine/Source/Programs/Shared/EpicGames.ProjectStore/CloudStorage/ICloudStorage.cs b/Engine/Source/Programs/Shared/EpicGames.ProjectStore/CloudStorage/ICloudStorage.cs
--- a/Engine/Source/Programs/Shared/EpicGames.ProjectStore/CloudStorage/ICloudStorage.cs
+++ b/Engine/Source/Programs/Shared/EpicGames.ProjectStore/CloudStorage/ICloudStorage.cs
@@ -351,6 +351,7 @@ namespace EpicGames.ProjectStore
 		/// <param name="cookedPlatformDir"></param>
 		/// <param name="forceImport"></param>
 		/// <param name="asyncImport"></param>
+		/// <param name="cancellationToken"></param>
 		/// <returns></returns>
 		Task ImportSnapshotAsync(string localFileName, string localRootPath, string engineDir, string cookedPlatformDir, bool forceImport = false, bool asyncImport = false, CancellationToken cancellationToken = default(CancellationToken));
 	}
@@ -425,6 +426,11 @@ namespace EpicGames.ProjectStore
 		/// </summary>
 		public bool EchoZenCliToLog { get; set; } = false;
 
+		/// <summary>
+		/// Implements verification of the settings object for IValidatableObject
+		/// </summary>
+		/// <param name="validationContext"></param>
+		/// <returns></returns>
 		public IEnumerable<ValidationResult> Validate(ValidationContext validationContext)
 		{
 			List<ValidationResult> results = new List<ValidationResult>();
@@ -446,6 +452,9 @@ namespace EpicGames.ProjectStore
 			return results;
 		}
 
+		/// <summary>
+		/// Finds paths for executables based on unreal root directory
+		/// </summary>
 		public void FindExecutablesFromUnrealRoot()
 		{
 			string binaryType = "Win64";
