Added AllowRetry flag to build graph nodes.

diff --git a/Engine/Source/Programs/AutomationTool/BuildGraph/BgSchema.cs b/Engine/Source/Programs/AutomationTool/BuildGraph/BgSchema.cs
--- a/Engine/Source/Programs/AutomationTool/BuildGraph/BgSchema.cs
+++ b/Engine/Source/Programs/AutomationTool/BuildGraph/BgSchema.cs
@@ -571,6 +571,7 @@ namespace AutomationTool
 			extension.Attributes.Add(CreateSchemaAttribute("NotifyOnWarnings", ScriptSchemaStandardType.Boolean, XmlSchemaUse.Optional));
 			extension.Attributes.Add(CreateSchemaAttribute("Annotations", ScriptSchemaStandardType.BalancedString, XmlSchemaUse.Optional));
 			extension.Attributes.Add(CreateSchemaAttribute("IgnoreModified", ScriptSchemaStandardType.BalancedString, XmlSchemaUse.Optional));
+			extension.Attributes.Add(CreateSchemaAttribute("AllowRetry", ScriptSchemaStandardType.Boolean, XmlSchemaUse.Optional));
 
 			XmlSchemaComplexContent contentModel = new XmlSchemaComplexContent();
 			contentModel.Content = extension;
diff --git a/Engine/Source/Programs/AutomationTool/BuildGraph/BgScriptReader.cs b/Engine/Source/Programs/AutomationTool/BuildGraph/BgScriptReader.cs
--- a/Engine/Source/Programs/AutomationTool/BuildGraph/BgScriptReader.cs
+++ b/Engine/Source/Programs/AutomationTool/BuildGraph/BgScriptReader.cs
@@ -1461,6 +1461,7 @@ namespace AutomationTool
 				string[] afterNames = ReadListAttribute(element, "After");
 				string[] tokenFileNames = ReadListAttribute(element, "Token");
 				bool runEarly = ReadBooleanAttribute(element, "RunEarly", false);
+				bool allowRetry = ReadBooleanAttribute(element, "AllowRetry", true);
 				bool notifyOnWarnings = ReadBooleanAttribute(element, "NotifyOnWarnings", true);
 				Dictionary<string, string> annotations = ReadAnnotationsAttribute(element, "Annotations");
 				string[] ignoreModified = ReadListAttribute(element, "IgnoreModified");
@@ -1530,6 +1531,7 @@ namespace AutomationTool
 					// Add it to the node lookup
 					BgScriptNode newNode = new BgScriptNode(name, inputs.ToArray(), validOutputNames.ToArray(), inputDependencies.ToArray(), orderDependencies.ToArray(), requiredTokens.ToArray(), ignoreModified.ToArray());
 					newNode.RunEarly = runEarly;
+					newNode.AllowRetry = allowRetry;
 					newNode.NotifyOnWarnings = notifyOnWarnings;
 					foreach ((string key, string value) in annotations)
 					{
diff --git a/Engine/Source/Programs/AutomationTool/BuildGraph/BgScriptWriter.cs b/Engine/Source/Programs/AutomationTool/BuildGraph/BgScriptWriter.cs
--- a/Engine/Source/Programs/AutomationTool/BuildGraph/BgScriptWriter.cs
+++ b/Engine/Source/Programs/AutomationTool/BuildGraph/BgScriptWriter.cs
@@ -168,6 +168,11 @@ namespace AutomationTool
 				writer.WriteAttributeString("RunEarly", node.RunEarly.ToString());
 			}
 
+			if (!node.AllowRetry)
+			{
+				writer.WriteAttributeString("AllowRetry", node.AllowRetry.ToString());
+			}
+
 			if (node.IgnoreModified.Count > 0)
 			{
 				writer.WriteAttributeString("IgnoreModified", String.Join(";", node.IgnoreModified));
diff --git a/Engine/Source/Programs/Horde/Drivers/JobDriver/Execution/JobExecutor.cs b/Engine/Source/Programs/Horde/Drivers/JobDriver/Execution/JobExecutor.cs
--- a/Engine/Source/Programs/Horde/Drivers/JobDriver/Execution/JobExecutor.cs
+++ b/Engine/Source/Programs/Horde/Drivers/JobDriver/Execution/JobExecutor.cs
@@ -77,6 +77,7 @@ namespace JobDriver.Execution
 		{
 			public string Name { get; set; } = String.Empty;
 			public bool RunEarly { get; set; }
+			public bool? AllowRetry { get; set; }
 			public bool? Warnings { get; set; }
 			public List<string> Inputs { get; set; } = new List<string>();
 			public List<string> Outputs { get; set; } = new List<string>();
@@ -950,6 +951,7 @@ namespace JobDriver.Execution
 						createNode.OrderDependencies.Add(exportedNode.OrderDependencies);
 					}
 					createNode.RunEarly = exportedNode.RunEarly;
+					createNode.AllowRetry = exportedNode.AllowRetry;
 					createNode.Warnings = exportedNode.Warnings;
 					createNode.Priority = (int)Priority.Normal;
 					createNode.Annotations.Add(exportedNode.Annotations);
diff --git a/Engine/Source/Programs/Shared/EpicGames.BuildGraph/BgGraphDef.cs b/Engine/Source/Programs/Shared/EpicGames.BuildGraph/BgGraphDef.cs
--- a/Engine/Source/Programs/Shared/EpicGames.BuildGraph/BgGraphDef.cs
+++ b/Engine/Source/Programs/Shared/EpicGames.BuildGraph/BgGraphDef.cs
@@ -359,6 +359,7 @@ namespace EpicGames.BuildGraph
 							jsonWriter.WriteValue("Name", node.Name);
 							jsonWriter.WriteValue("DependsOn", String.Join(";", node.GetDirectOrderDependencies().Where(x => nodesToExecute.Contains(x))));
 							jsonWriter.WriteValue("RunEarly", node.RunEarly);
+							jsonWriter.WriteValue("AllowRetry", node.AllowRetry);
 							jsonWriter.WriteObjectStart("Notify");
 							jsonWriter.WriteValue("Default", String.Join(";", node.NotifyUsers));
 							jsonWriter.WriteValue("Submitters", String.Join(";", node.NotifySubmitters));
@@ -510,6 +511,7 @@ namespace EpicGames.BuildGraph
 						jsonWriter.WriteObjectStart();
 						jsonWriter.WriteValue("Name", node.Name);
 						jsonWriter.WriteValue("RunEarly", node.RunEarly);
+						jsonWriter.WriteValue("AllowRetry", node.AllowRetry);
 						jsonWriter.WriteValue("Warnings", node.NotifyOnWarnings);
 
 						jsonWriter.WriteArrayStart("Inputs");
@@ -709,7 +711,7 @@ namespace EpicGames.BuildGraph
 				logger.LogInformation("        Agent: {Name} ({Types})", agent.Name, String.Join(";", agent.PossibleTypes));
 				foreach (BgNodeDef node in agent.Nodes)
 				{
-					logger.LogInformation("            Node: {Name}{Type}", node.Name, completedNodes.Contains(node) ? " (completed)" : node.RunEarly ? " (early)" : "");
+					logger.LogInformation("            Node: {Name}{Type}{Retry}", node.Name, completedNodes.Contains(node) ? " (completed)" : node.RunEarly ? " (early)" : "", !node.AllowRetry ? " (retry disabled)" : "");
 					if (printOptions.HasFlag(GraphPrintOptions.ShowDependencies))
 					{
 						HashSet<BgNodeDef> inputDependencies = new HashSet<BgNodeDef>(node.GetDirectInputDependencies());
diff --git a/Engine/Source/Programs/Shared/EpicGames.BuildGraph/BgNodeDef.cs b/Engine/Source/Programs/Shared/EpicGames.BuildGraph/BgNodeDef.cs
--- a/Engine/Source/Programs/Shared/EpicGames.BuildGraph/BgNodeDef.cs
+++ b/Engine/Source/Programs/Shared/EpicGames.BuildGraph/BgNodeDef.cs
@@ -152,6 +152,11 @@ namespace EpicGames.BuildGraph
 		/// </summary>
 		public bool RunEarly { get; set; } = false;
 
+		/// <summary>
+		/// Whether this node can be retried.
+		/// </summary>
+		public bool AllowRetry { get; set; } = true;
+
 		/// <summary>
 		/// Whether to ignore warnings produced by this node
 		/// </summary>
