Fix free space calculation. This change allows the use of drives mounted as paths, not just drive letters.

diff --git a/Engine/Source/Programs/Shared/EpicGames.Perforce.Managed/ManagedWorkspace.cs b/Engine/Source/Programs/Shared/EpicGames.Perforce.Managed/ManagedWorkspace.cs
--- a/Engine/Source/Programs/Shared/EpicGames.Perforce.Managed/ManagedWorkspace.cs
+++ b/Engine/Source/Programs/Shared/EpicGames.Perforce.Managed/ManagedWorkspace.cs
@@ -2074,6 +2074,15 @@ namespace EpicGames.Perforce.Managed
 			}
 		}
 
+		[DllImport("kernel32.dll", SetLastError = true, CharSet = CharSet.Unicode)]
+		[return: MarshalAs(UnmanagedType.Bool)]
+		static extern bool GetDiskFreeSpaceEx(
+			string lpDirectoryName,
+			out ulong lpFreeBytesAvailable,
+			out ulong lpTotalNumberOfBytes,
+			out ulong lpTotalNumberOfFreeBytes
+		);
+
 		/// <summary>
 		/// Update the workspace to match the given stream, syncing files and moving to/from the cache as necessary.
 		/// </summary>
@@ -2148,9 +2157,14 @@ namespace EpicGames.Perforce.Managed
 				// Make sure there's enough space on this drive
 				if (RuntimeInformation.IsOSPlatform(OSPlatform.Windows))
 				{
-					long freeSpace = new DriveInfo(Path.GetPathRoot(_baseDir.FullName)!).AvailableFreeSpace;
+					ulong freeSpace;
+            		bool success = GetDiskFreeSpaceEx(_baseDir.FullName, out _, out _, out freeSpace);
+					if (!success)
+					{
+						throw new InsufficientSpaceException($"Can't get free space for {_baseDir.FullName}");
+					}
 					long minScratchSpaceBytes = _options.MinScratchSpace * 1024 * 1024;
-					if (freeSpace - syncSize < minScratchSpaceBytes)
+					if ((long)freeSpace - syncSize < minScratchSpaceBytes)
 					{
 						throw new InsufficientSpaceException($"Not enough space to sync new files (free space: {freeSpace / (1024.0 * 1024.0):n1}mb, sync size: {syncSize / (1024.0 * 1024.0):n1}mb, min scratch space: {minScratchSpaceBytes / (1024.0 * 1024.0):n1}mb)");
 					}
