Emit an error in UBT for bad allowlist values for plugin / module target values type and configuration values.

diff --git a/Engine/Plugins/Experimental/ChaosOutfitAsset/ChaosOutfitAsset.uplugin b/Engine/Plugins/Experimental/ChaosOutfitAsset/ChaosOutfitAsset.uplugin
--- a/Engine/Plugins/Experimental/ChaosOutfitAsset/ChaosOutfitAsset.uplugin
+++ b/Engine/Plugins/Experimental/ChaosOutfitAsset/ChaosOutfitAsset.uplugin
@@ -49,7 +49,7 @@
 		{
 			"Name": "ChaosClothAssetEditor",
 			"Enabled": true,
-			"TargetAllowList" : "Editor"
+			"TargetAllowList" : [ "Editor" ]
 		}
 	]
 }
diff --git a/Engine/Source/Programs/Shared/EpicGames.Core/JsonObject.cs b/Engine/Source/Programs/Shared/EpicGames.Core/JsonObject.cs
--- a/Engine/Source/Programs/Shared/EpicGames.Core/JsonObject.cs
+++ b/Engine/Source/Programs/Shared/EpicGames.Core/JsonObject.cs
@@ -605,30 +605,61 @@ namespace EpicGames.Core
 		/// </summary>
 		/// <param name="fieldName">Name of the field to get</param>
 		/// <param name="result">On success, receives the field value</param>
+		/// <param name="errors">On failure, details the error(s) that occurred</param>
 		/// <returns>True if the field could be read, false otherwise</returns>
-		public bool TryGetEnumArrayField<T>(string fieldName, [NotNullWhen(true)] out T[]? result) where T : struct
+		public bool TryGetEnumArrayField<T>(string fieldName, [NotNullWhen(true)] out T[]? result, [NotNullWhen(false)] out string? errors) where T : struct
 		{
+			if (!ContainsField(fieldName))
+			{
+				result = null;
+				errors = $"Field '{fieldName}' not found";
+				return false;
+			}
+
 			string[]? stringValues;
 			if (!TryGetStringArrayField(fieldName, out stringValues))
 			{
 				result = null;
+				errors = $"Field '{fieldName}' has an incorrect type - it must be an array of strings";
 				return false;
 			}
 
 			T[] enumValues = new T[stringValues.Length];
+			bool success = true;
+			StringBuilder? errorBuilder = null;
 			for (int idx = 0; idx < stringValues.Length; idx++)
 			{
-				if (!Enum.TryParse<T>(stringValues[idx], true, out enumValues[idx]))
+				if (!Enum.TryParse(stringValues[idx], true, out enumValues[idx]))
 				{
-					result = null;
-					return false;
+					success = false;
+					errorBuilder ??= new StringBuilder();
+					errorBuilder.AppendLine($"Field '{fieldName}' entry at index {idx} '{stringValues[idx]}' is not a valid {typeof(T).Name}");
 				}
 			}
 
+			if (!success)
+			{
+				result = null;
+#pragma warning disable CS8602 // If !success then errorBuilder != null
+				errors = errorBuilder.ToString();
+#pragma warning restore CS8602
+				return false;
+			}
+
 			result = enumValues;
+			errors = null;
 			return true;
 		}
 
+		/// <summary>
+		/// Tries to read an enum array field by the given name from the object
+		/// </summary>
+		/// <param name="fieldName">Name of the field to get</param>
+		/// <param name="result">On success, receives the field value</param>
+		/// <returns>True if the field could be read, false otherwise</returns>
+		public bool TryGetEnumArrayField<T>(string fieldName, [NotNullWhen(true)] out T[]? result) where T : struct
+			=> TryGetEnumArrayField(fieldName, out result, out _);
+
 		/// <summary>
 		/// Gets an object field by the given name from the object, throwing an exception if it is not there or cannot be parsed.
 		/// </summary>
diff --git a/Engine/Source/Programs/Shared/UnrealBuildTool.Tests/CppCompileWarningsTests.cs b/Engine/Source/Programs/Shared/UnrealBuildTool.Tests/CppCompileWarningsTests.cs
--- a/Engine/Source/Programs/Shared/UnrealBuildTool.Tests/CppCompileWarningsTests.cs
+++ b/Engine/Source/Programs/Shared/UnrealBuildTool.Tests/CppCompileWarningsTests.cs
@@ -240,11 +240,9 @@ namespace UnrealBuildTool.Tests
 		[TestMethod]
 		public void TestGenerateWarningCommandLineAgsOrder()
 		{
-			DirectoryReference programsDirectory = DirectoryReference.Combine(Unreal.EngineSourceDirectory, "Programs");
-			FileReference testStubsDirectory = FileReference.Combine(programsDirectory, "Shared", "UnrealBuildTool.Tests", "UBT");
-
+			DirectoryReference testStubsDirectory = TestDirs.Stubs;
 			// Generate ephemeral TargetRules & uproject files
-			using (SafeTestDirectory testRootDirectory = SafeTestDirectory.CreateTestDirectory(Path.Combine(StaticTargetRulesInitializer.UBTTestFolderRoot, $"{StaticTargetRulesInitializer.UBTTestFolderPrefix}{Guid.NewGuid().ToString()}")))
+			using (SafeTestDirectory testRootDirectory = TestDirs.Create())
 			using (SafeTestDirectory testSourceDirectory = SafeTestDirectory.CreateTestDirectory(Path.Combine(testRootDirectory.TemporaryDirectory, "Source")))
 			using (SafeTestFile testEditorTargetTestFile = new SafeTestFile(File.ReadAllText(Path.Combine(testStubsDirectory.FullName, "TargetStubs", "TestStubEditor.Target.ubttest")), "TestStubEditor.Target.cs", testSourceDirectory.TemporaryDirectory))
 			using (SafeTestFile testTargetTestFile = new SafeTestFile(File.ReadAllText(Path.Combine(testStubsDirectory.FullName, "TargetStubs", "TestStub.Target.ubttest")), "TestStub.Target.cs", testSourceDirectory.TemporaryDirectory))
@@ -280,7 +278,7 @@ namespace UnrealBuildTool.Tests
 			FileReference testStubsDirectory = FileReference.Combine(programsDirectory, "Shared", "UnrealBuildTool.Tests", "UBT");
 
 			// Generate ephemeral TargetRules & uproject files
-			using (SafeTestDirectory testRootDirectory = SafeTestDirectory.CreateTestDirectory(Path.Combine(StaticTargetRulesInitializer.UBTTestFolderRoot, $"{StaticTargetRulesInitializer.UBTTestFolderPrefix}{Guid.NewGuid().ToString()}")))
+			using (SafeTestDirectory testRootDirectory = TestDirs.Create())
 			using (SafeTestDirectory testSourceDirectory = SafeTestDirectory.CreateTestDirectory(Path.Combine(testRootDirectory.TemporaryDirectory, "Source")))
 			using (SafeTestFile testEditorTargetTestFile = new SafeTestFile(File.ReadAllText(Path.Combine(testStubsDirectory.FullName, "TargetStubs", "TestStubEditor.Target.ubttest")), "TestStubEditor.Target.cs", testSourceDirectory.TemporaryDirectory))
 			using (SafeTestFile testTargetTestFile = new SafeTestFile(File.ReadAllText(Path.Combine(testStubsDirectory.FullName, "TargetStubs", "TestStub.Target.ubttest")), "TestStub.Target.cs", testSourceDirectory.TemporaryDirectory))
@@ -364,7 +362,7 @@ namespace UnrealBuildTool.Tests
 			FileReference testStubsDirectory = FileReference.Combine(programsDirectory, "Shared", "UnrealBuildTool.Tests", "UBT");
 
 			// Generate ephemeral TargetRules & uproject files
-			using (SafeTestDirectory testRootDirectory = SafeTestDirectory.CreateTestDirectory(Path.Combine(StaticTargetRulesInitializer.UBTTestFolderRoot, $"{StaticTargetRulesInitializer.UBTTestFolderPrefix}{Guid.NewGuid().ToString()}")))
+			using (SafeTestDirectory testRootDirectory = TestDirs.Create())
 			using (SafeTestDirectory testSourceDirectory = SafeTestDirectory.CreateTestDirectory(Path.Combine(testRootDirectory.TemporaryDirectory, "Source")))
 			using (SafeTestFile testEditorTargetTestFile = new SafeTestFile(File.ReadAllText(Path.Combine(testStubsDirectory.FullName, "TargetStubs", "TestStubEditor.Target.ubttest")), "TestStubEditor.Target.cs", testSourceDirectory.TemporaryDirectory))
 			using (SafeTestFile testTargetTestFile = new SafeTestFile(File.ReadAllText(Path.Combine(testStubsDirectory.FullName, "TargetStubs", "TestStub.Target.ubttest")), "TestStub.Target.cs", testSourceDirectory.TemporaryDirectory))
@@ -541,7 +539,7 @@ namespace UnrealBuildTool.Tests
 			FileReference testStubsDirectory = FileReference.Combine(programsDirectory, "Shared", "UnrealBuildTool.Tests", "UBT");
 
 			// Generate ephemeral TargetRules & uproject files
-			using (SafeTestDirectory testRootDirectory = SafeTestDirectory.CreateTestDirectory(Path.Combine(StaticTargetRulesInitializer.UBTTestFolderRoot, $"{StaticTargetRulesInitializer.UBTTestFolderPrefix}{Guid.NewGuid().ToString()}")))
+			using (SafeTestDirectory testRootDirectory = TestDirs.Create())
 			using (SafeTestDirectory testSourceDirectory = SafeTestDirectory.CreateTestDirectory(Path.Combine(testRootDirectory.TemporaryDirectory, "Source")))
 			using (SafeTestFile testEditorTargetTestFile = new SafeTestFile(File.ReadAllText(Path.Combine(testStubsDirectory.FullName, "TargetStubs", "TestStubEditor.Target.ubttest")), "TestStubEditor.Target.cs", testSourceDirectory.TemporaryDirectory))
 			using (SafeTestFile testTargetTestFile = new SafeTestFile(File.ReadAllText(Path.Combine(testStubsDirectory.FullName, "TargetStubs", "TestStub.Target.ubttest")), "TestStub.Target.cs", testSourceDirectory.TemporaryDirectory))
@@ -584,7 +582,7 @@ namespace UnrealBuildTool.Tests
 			FileReference testStubsDirectory = FileReference.Combine(programsDirectory, "Shared", "UnrealBuildTool.Tests", "UBT");
 
 			// Generate ephemeral TargetRules & uproject files
-			using (SafeTestDirectory testRootDirectory = SafeTestDirectory.CreateTestDirectory(Path.Combine(StaticTargetRulesInitializer.UBTTestFolderRoot, $"{StaticTargetRulesInitializer.UBTTestFolderPrefix}{Guid.NewGuid().ToString()}")))
+			using (SafeTestDirectory testRootDirectory = TestDirs.Create())
 			using (SafeTestDirectory testSourceDirectory = SafeTestDirectory.CreateTestDirectory(Path.Combine(testRootDirectory.TemporaryDirectory, "Source")))
 			using (SafeTestFile testEditorTargetTestFile = new SafeTestFile(File.ReadAllText(Path.Combine(testStubsDirectory.FullName, "TargetStubs", "TestStubEditor.Target.ubttest")), "TestStubEditor.Target.cs", testSourceDirectory.TemporaryDirectory))
 			using (SafeTestFile testTargetTestFile = new SafeTestFile(File.ReadAllText(Path.Combine(testStubsDirectory.FullName, "TargetStubs", "TestStub.Target.ubttest")), "TestStub.Target.cs", testSourceDirectory.TemporaryDirectory))
@@ -635,7 +633,7 @@ namespace UnrealBuildTool.Tests
 			FileReference testStubsDirectory = FileReference.Combine(programsDirectory, "Shared", "UnrealBuildTool.Tests", "UBT");
 
 			// Generate ephemeral TargetRules & uproject files
-			using (SafeTestDirectory testRootDirectory = SafeTestDirectory.CreateTestDirectory(Path.Combine(StaticTargetRulesInitializer.UBTTestFolderRoot, $"{StaticTargetRulesInitializer.UBTTestFolderPrefix}{Guid.NewGuid().ToString()}")))
+			using (SafeTestDirectory testRootDirectory = TestDirs.Create())
 			using (SafeTestDirectory testSourceDirectory = SafeTestDirectory.CreateTestDirectory(Path.Combine(testRootDirectory.TemporaryDirectory, "Source")))
 			using (SafeTestFile testEditorTargetTestFile = new SafeTestFile(File.ReadAllText(Path.Combine(testStubsDirectory.FullName, "TargetStubs", "TestStubEditor.Target.ubttest")), "TestStubEditor.Target.cs", testSourceDirectory.TemporaryDirectory))
 			using (SafeTestFile testTargetTestFile = new SafeTestFile(File.ReadAllText(Path.Combine(testStubsDirectory.FullName, "TargetStubs", "TestStub.Target.ubttest")), "TestStub.Target.cs", testSourceDirectory.TemporaryDirectory))
@@ -745,7 +743,7 @@ namespace UnrealBuildTool.Tests
 			FileReference testStubsDirectory = FileReference.Combine(programsDirectory, "Shared", "UnrealBuildTool.Tests", "UBT");
 
 			// Generate ephemeral TargetRules & uproject files
-			using (SafeTestDirectory testRootDirectory = SafeTestDirectory.CreateTestDirectory(Path.Combine(StaticTargetRulesInitializer.UBTTestFolderRoot, $"{StaticTargetRulesInitializer.UBTTestFolderPrefix}{Guid.NewGuid().ToString()}")))
+			using (SafeTestDirectory testRootDirectory = TestDirs.Create())
 			using (SafeTestDirectory testSourceDirectory = SafeTestDirectory.CreateTestDirectory(Path.Combine(testRootDirectory.TemporaryDirectory, "Source")))
 			using (SafeTestFile testEditorTargetTestFile = new SafeTestFile(File.ReadAllText(Path.Combine(testStubsDirectory.FullName, "TargetStubs", "TestStubEditor.Target.ubttest")), "TestStubEditor.Target.cs", testSourceDirectory.TemporaryDirectory))
 			using (SafeTestFile testTargetTestFile = new SafeTestFile(File.ReadAllText(Path.Combine(testStubsDirectory.FullName, "TargetStubs", "TestStub.Target.ubttest")), "TestStub.Target.cs", testSourceDirectory.TemporaryDirectory))
diff --git a/Engine/Source/Programs/Shared/UnrealBuildTool.Tests/JsonExtensionsTests.cs b/Engine/Source/Programs/Shared/UnrealBuildTool.Tests/JsonExtensionsTests.cs
new file mode 100644
--- /dev/null
+++ b/Engine/Source/Programs/Shared/UnrealBuildTool.Tests/JsonExtensionsTests.cs
@@ -0,0 +1,93 @@
+// Copyright Epic Games, Inc. All Rights Reserved.
+
+using System;
+using EpicGames.Core;
+using JsonExtensions;
+using Microsoft.VisualStudio.TestTools.UnitTesting;
+using UnrealBuildBase;
+
+namespace UnrealBuildTool.Tests
+{
+	[TestClass]
+	public class JsonExtensionsTests
+	{
+		public enum TestEnum
+		{
+			A,
+			B,
+			C,
+		}
+
+		[TestMethod]
+		public void TestGetEnumArrayFieldWithDeprecatedFallbackReturnsNullWhenBothAbsent()
+		{
+			JsonObject json = new JsonObject();
+			TestEnum[]? values = json.GetEnumArrayFieldWithDeprecatedFallback<TestEnum>("New", "Old");
+			Assert.IsNull(values);
+		}
+
+		[TestMethod]
+		public void TestGetEnumArrayFieldWithDeprecatedFallbackThrowsWhenBothPresent()
+		{
+			JsonObject json = new JsonObject();
+			json.AddOrSetFieldValue("New", [TestEnum.A, TestEnum.B]);
+			json.AddOrSetFieldValue("Old", [TestEnum.A, TestEnum.B]);
+			BuildLogEventException ex = Assert.ThrowsException<BuildLogEventException>(
+				() => json.GetEnumArrayFieldWithDeprecatedFallback<TestEnum>("New", "Old"));
+			StringAssert.Contains(ex.Message, "New", StringComparison.Ordinal);
+			StringAssert.Contains(ex.Message, "Old", StringComparison.Ordinal);
+		}
+
+		[TestMethod]
+		public void TestGetEnumArrayFieldWithDeprecatedFallbackPicksUpNewValues()
+		{
+			JsonObject json = new JsonObject();
+			json.AddOrSetFieldValue("New", [TestEnum.A, TestEnum.B]);
+			TestEnum[]? values = json.GetEnumArrayFieldWithDeprecatedFallback<TestEnum>("New", "Old");
+			CollectionAssert.AreEquivalent(values, new[] { TestEnum.A, TestEnum.B });
+		}
+
+		[TestMethod]
+		public void TestGetEnumArrayFieldWithDeprecatedFallbackPicksUpFallbackValues()
+		{
+			JsonObject json = new JsonObject();
+			json.AddOrSetFieldValue("Old", [TestEnum.A, TestEnum.B]);
+			TestEnum[]? values = json.GetEnumArrayFieldWithDeprecatedFallback<TestEnum>("New", "Old");
+			CollectionAssert.AreEquivalent(values, new[] { TestEnum.A, TestEnum.B });
+		}
+
+		[TestMethod]
+		public void TestGetEnumArrayFieldWithDeprecatedFallbackFailsOnInvalidNewValue()
+		{
+			JsonObject json = new JsonObject();
+			json.AddOrSetFieldValue("New", ["A", "B", "C", "BadItem1", "BadItem2"]);
+			BuildLogEventException ex = Assert.ThrowsException<BuildLogEventException>(
+				() => json.GetEnumArrayFieldWithDeprecatedFallback<TestEnum>("New", "Old"));
+			StringAssert.Contains(ex.Message, "New", StringComparison.Ordinal);
+			StringAssert.Contains(ex.Message, "BadItem1", StringComparison.Ordinal);
+			StringAssert.Contains(ex.Message, "BadItem2", StringComparison.Ordinal);
+		}
+
+		[TestMethod]
+		public void TestGetEnumArrayFieldWithDeprecatedFallbackFailsOnInvalidFallbackValue()
+		{
+			JsonObject json = new JsonObject();
+			json.AddOrSetFieldValue("Old", ["A", "B", "C", "BadItem1", "BadItem2"]);
+			BuildLogEventException ex = Assert.ThrowsException<BuildLogEventException>(
+				() => json.GetEnumArrayFieldWithDeprecatedFallback<TestEnum>("New", "Old"));
+			StringAssert.Contains(ex.Message, "Old", StringComparison.Ordinal);
+			StringAssert.Contains(ex.Message, "BadItem1", StringComparison.Ordinal);
+			StringAssert.Contains(ex.Message, "BadItem2", StringComparison.Ordinal);
+		}
+
+		[TestMethod]
+		public void TestGetEnumArrayFieldWithDeprecatedFallbackFailsOnInvalidDataType()
+		{
+			JsonObject json = new JsonObject();
+			json.AddOrSetFieldValue("New", "A");
+			BuildLogEventException ex = Assert.ThrowsException<BuildLogEventException>(
+				() => json.GetEnumArrayFieldWithDeprecatedFallback<TestEnum>("New", "Old"));
+			StringAssert.Contains(ex.Message, "New", StringComparison.Ordinal);
+		}
+	}
+}
diff --git a/Engine/Source/Programs/Shared/UnrealBuildTool.Tests/ProjectDescriptorTests.cs b/Engine/Source/Programs/Shared/UnrealBuildTool.Tests/ProjectDescriptorTests.cs
new file mode 100644
--- /dev/null
+++ b/Engine/Source/Programs/Shared/UnrealBuildTool.Tests/ProjectDescriptorTests.cs
@@ -0,0 +1,49 @@
+// Copyright Epic Games, Inc. All Rights Reserved.
+
+using System;
+using System.IO;
+using Microsoft.VisualStudio.TestTools.UnitTesting;
+using UnrealBuildBase;
+using UnrealBuildTool.Tests.TestUtilities;
+
+namespace UnrealBuildTool.Tests
+{
+	[TestClass]
+	public class ProjectDescriptorTests
+	{
+		[TestMethod]
+		public void TestLoadingProjectWithInvalidModuleAllowListFails()
+		{
+			string uprojectPath = Path.Combine(TestDirs.Stubs.FullName, "UprojectStubs", "BadModuleAllowList.uproject.ubttest");
+
+			using SafeTestDirectory testRootDirectory = TestDirs.Create();
+			using SafeTestFile uprojectFile = new SafeTestFile(File.ReadAllText(uprojectPath), "TestStub.uproject", testRootDirectory.TemporaryDirectory);
+
+			CheckExceptionContainsError(uprojectFile, "Editro", "TargetAllowList");
+		}
+
+		[TestMethod]
+		public void TestLoadingProjectWithInvalidPluginAllowListFails()
+		{
+			string uprojectPath = Path.Combine(TestDirs.Stubs.FullName, "UprojectStubs", "BadPluginAllowList.uproject.ubttest");
+
+			using SafeTestDirectory testRootDirectory = TestDirs.Create();
+			using SafeTestFile uprojectFile = new SafeTestFile(File.ReadAllText(uprojectPath), "TestStub.uproject", testRootDirectory.TemporaryDirectory);
+
+			CheckExceptionContainsError(uprojectFile, "Shpipping", "TargetConfigurationAllowList");
+		}
+
+		private static void CheckExceptionContainsError(SafeTestFile uprojectFile, params string[] requiredTexts)
+		{
+			Exception ex = Assert.ThrowsException<Exception>(() => ProjectDescriptor.FromFile(uprojectFile));
+			Assert.IsNotNull(ex.InnerException);
+			Assert.AreEqual(ex.InnerException.GetType(), typeof(BuildLogEventException));
+			BuildLogEventException inner = (BuildLogEventException)ex.InnerException;
+
+			foreach (string text in requiredTexts)
+			{
+				StringAssert.Contains(inner.Message, text, StringComparison.Ordinal);
+			}
+		}
+	}
+}
diff --git a/Engine/Source/Programs/Shared/UnrealBuildTool.Tests/TargetRulesTests.cs b/Engine/Source/Programs/Shared/UnrealBuildTool.Tests/TargetRulesTests.cs
--- a/Engine/Source/Programs/Shared/UnrealBuildTool.Tests/TargetRulesTests.cs
+++ b/Engine/Source/Programs/Shared/UnrealBuildTool.Tests/TargetRulesTests.cs
@@ -20,7 +20,7 @@ namespace UnrealBuildTool.Tests
 			DirectoryReference testStubsDirectory = DirectoryReference.Combine(programsDirectory, "Shared", "UnrealBuildTool.Tests", "UBT");
 
 			// Generate ephemeral TargetRules & uproject files
-			using (SafeTestDirectory testRootDirectory = SafeTestDirectory.CreateTestDirectory(Path.Combine(StaticTargetRulesInitializer.UBTTestFolderRoot, $"{StaticTargetRulesInitializer.UBTTestFolderPrefix}{Guid.NewGuid().ToString()}")))
+			using (SafeTestDirectory testRootDirectory = TestDirs.Create())
 			using (SafeTestDirectory testSourceDirectory = SafeTestDirectory.CreateTestDirectory(Path.Combine(testRootDirectory.TemporaryDirectory, "Source")))
 			using (SafeTestFile testEditorTargetTestFile = new SafeTestFile(File.ReadAllText(Path.Combine(testStubsDirectory.FullName, "TargetStubs", "TestStubEditor.Target.ubttest")), "TestStubEditor.Target.cs", testSourceDirectory.TemporaryDirectory))
 			using (SafeTestFile testUprojectTestFile = new SafeTestFile(File.ReadAllText(Path.Combine(testStubsDirectory.FullName, "UprojectStubs", "TestStub.uproject.ubttest")), "TestStub.uproject", testRootDirectory.TemporaryDirectory))
diff --git a/Engine/Source/Programs/Shared/UnrealBuildTool.Tests/TestUtilities/SafeTestFile.cs b/Engine/Source/Programs/Shared/UnrealBuildTool.Tests/TestUtilities/SafeTestFile.cs
--- a/Engine/Source/Programs/Shared/UnrealBuildTool.Tests/TestUtilities/SafeTestFile.cs
+++ b/Engine/Source/Programs/Shared/UnrealBuildTool.Tests/TestUtilities/SafeTestFile.cs
@@ -2,6 +2,7 @@
 
 using System;
 using System.IO;
+using EpicGames.Core;
 
 namespace UnrealBuildTool.Tests.TestUtilities
 {
@@ -33,5 +34,8 @@ namespace UnrealBuildTool.Tests.TestUtilities
 				File.Delete(TemporaryFile);
 			}
 		}
+
+		public static implicit operator FileReference(SafeTestFile file)
+			=> new FileReference(file.TemporaryFile);
 	}
 }
diff --git a/Engine/Source/Programs/Shared/UnrealBuildTool.Tests/TestUtilities/StaticTargetRulesInitializer.cs b/Engine/Source/Programs/Shared/UnrealBuildTool.Tests/TestUtilities/StaticTargetRulesInitializer.cs
--- a/Engine/Source/Programs/Shared/UnrealBuildTool.Tests/TestUtilities/StaticTargetRulesInitializer.cs
+++ b/Engine/Source/Programs/Shared/UnrealBuildTool.Tests/TestUtilities/StaticTargetRulesInitializer.cs
@@ -10,15 +10,12 @@ namespace UnrealBuildTool.Tests.TestUtilities
 	[TestClass]
 	public static class StaticTargetRulesInitializer
 	{
-		internal const string UBTTestFolderPrefix = "UBT-Test-";
-		internal static readonly string UBTTestFolderRoot = Path.GetTempPath();
-
 		[AssemblyInitialize]
 		public static void InitializeBuildEnvironment(TestContext context)
 		{
 			// Clean any previous runs linger temp folder
 			// Remarks: this is due to target rules modules dll being held by the test process; full cleanup not feasible on existing run
-			string[] subfolders = Directory.GetDirectories(UBTTestFolderRoot, $"{UBTTestFolderPrefix}*", SearchOption.TopDirectoryOnly);
+			string[] subfolders = Directory.GetDirectories(TestDirs.UBTTestFolderRoot, $"{TestDirs.UBTTestFolderPrefix}*", SearchOption.TopDirectoryOnly);
 
 			foreach (string folder in subfolders)
 			{
diff --git a/Engine/Source/Programs/Shared/UnrealBuildTool.Tests/TestUtilities/TestDirs.cs b/Engine/Source/Programs/Shared/UnrealBuildTool.Tests/TestUtilities/TestDirs.cs
new file mode 100644
--- /dev/null
+++ b/Engine/Source/Programs/Shared/UnrealBuildTool.Tests/TestUtilities/TestDirs.cs
@@ -0,0 +1,31 @@
+// Copyright Epic Games, Inc. All Rights Reserved.
+
+using System;
+using System.IO;
+using EpicGames.Core;
+using UnrealBuildBase;
+
+namespace UnrealBuildTool.Tests.TestUtilities
+{
+	internal class TestDirs
+	{
+		internal static DirectoryReference Stubs
+		{
+			get
+			{
+				DirectoryReference programsDirectory = DirectoryReference.Combine(Unreal.EngineSourceDirectory, "Programs");
+				return DirectoryReference.Combine(programsDirectory, "Shared", "UnrealBuildTool.Tests", "UBT");
+			}
+		}
+
+		internal const string UBTTestFolderPrefix = "UBT-Test-";
+		internal static readonly string UBTTestFolderRoot = Path.GetTempPath();
+
+		internal static SafeTestDirectory Create()
+		{
+			string tempDirName = $"{UBTTestFolderPrefix}{Guid.NewGuid()}";
+			string tempDirPath = Path.Combine(UBTTestFolderRoot, tempDirName);
+			return SafeTestDirectory.CreateTestDirectory(tempDirPath);
+		}
+	}
+}
diff --git a/Engine/Source/Programs/UnrealBuildTool/System/JsonExtensions.cs b/Engine/Source/Programs/UnrealBuildTool/System/JsonExtensions.cs
--- a/Engine/Source/Programs/UnrealBuildTool/System/JsonExtensions.cs
+++ b/Engine/Source/Programs/UnrealBuildTool/System/JsonExtensions.cs
@@ -2,6 +2,8 @@
 
 using System.Diagnostics.CodeAnalysis;
 using EpicGames.Core;
+using UnrealBuildBase;
+using UnrealBuildTool;
 
 namespace JsonExtensions
 {
@@ -36,28 +38,46 @@ namespace JsonExtensions
 		}
 
 		/// <summary>
-		/// Tries to read an enum array field by the given name from the object; if that's not found, checks for an older deprecated name
+		/// Reads an enum array field by the given name from the object; if that's not found, checks for an older deprecated name.
+		/// Throws if any invalid values are found.
 		/// </summary>
 		/// <param name="Obj">JSON object to check</param>
 		/// <param name="FieldName">Name of the field to get</param>
 		/// <param name="DeprecatedFieldName">Backup field name to check</param>
-		/// <param name="Result">On success, receives the field value</param>
-		/// <returns>True if the field could be read, false otherwise</returns>
-		public static bool TryGetEnumArrayFieldWithDeprecatedFallback<T>(this JsonObject Obj, string FieldName, string DeprecatedFieldName, [NotNullWhen(true)] out T[]? Result) where T : struct
+		/// <returns>A parsed T[] if one exists, or null if the field is absent</returns>
+		/// <exception cref="BuildLogEventException">Thrown if invalid values are found.</exception>
+		public static T[]? GetEnumArrayFieldWithDeprecatedFallback<T>(this JsonObject Obj, string FieldName, string DeprecatedFieldName) where T : struct
 		{
-			if (Obj.TryGetEnumArrayField<T>(FieldName, out Result))
+			T[]? Result;
+			string? Errors;
+
+			if (!Obj.ContainsField(FieldName) && !Obj.ContainsField(DeprecatedFieldName))
 			{
-				return true;
+				return null;
 			}
-			else if (Obj.TryGetEnumArrayField<T>(DeprecatedFieldName, out Result))
+
+			if (Obj.ContainsField(FieldName) && Obj.ContainsField(DeprecatedFieldName))
 			{
-				//@TODO: Warn here?
-				return true;
+				throw new BuildLogEventException("Field '{DeprecatedFieldName}' must be removed, as '{FieldName}' is already defined", DeprecatedFieldName, FieldName);
 			}
-			else
+
+			if (Obj.ContainsField(FieldName))
 			{
-				return false;
+				if (Obj.TryGetEnumArrayField(FieldName, out Result, out Errors))
+				{
+					return Result;
+				}
+
+				throw new BuildLogEventException("{JsonErrors}", Errors);
+			}
+
+			if (Obj.TryGetEnumArrayField(DeprecatedFieldName, out Result, out Errors))
+			{
+				//@TODO: Warn here?
+				return Result;
 			}
+
+			throw new BuildLogEventException("{JsonErrors}", Errors);
 		}
 	}
-}
\ No newline at end of file
+}
diff --git a/Engine/Source/Programs/UnrealBuildTool/System/ModuleDescriptor.cs b/Engine/Source/Programs/UnrealBuildTool/System/ModuleDescriptor.cs
--- a/Engine/Source/Programs/UnrealBuildTool/System/ModuleDescriptor.cs
+++ b/Engine/Source/Programs/UnrealBuildTool/System/ModuleDescriptor.cs
@@ -368,6 +368,11 @@ namespace UnrealBuildTool
 				{
 					Module.PlatformArchitectureDenyList = ParsePlatformArchitectureList(PlatformArchitectureDenyList);
 				}
+
+				Module.TargetAllowList = InObject.GetEnumArrayFieldWithDeprecatedFallback<TargetType>("TargetAllowList", "WhitelistTargets");
+				Module.TargetDenyList = InObject.GetEnumArrayFieldWithDeprecatedFallback<TargetType>("TargetDenyList", "BlacklistTargets");
+				Module.TargetConfigurationAllowList = InObject.GetEnumArrayFieldWithDeprecatedFallback<UnrealTargetConfiguration>("TargetConfigurationAllowList", "WhitelistTargetConfigurations");
+				Module.TargetConfigurationDenyList = InObject.GetEnumArrayFieldWithDeprecatedFallback<UnrealTargetConfiguration>("TargetConfigurationDenyList", "BlacklistTargetConfigurations");
 			}
 			catch (BuildException Ex)
 			{
@@ -375,32 +380,6 @@ namespace UnrealBuildTool
 				throw;
 			}
 
-
-
-			TargetType[]? TargetAllowList;
-			if (InObject.TryGetEnumArrayFieldWithDeprecatedFallback<TargetType>("TargetAllowList", "WhitelistTargets", out TargetAllowList))
-			{
-				Module.TargetAllowList = TargetAllowList;
-			}
-
-			TargetType[]? TargetDenyList;
-			if (InObject.TryGetEnumArrayFieldWithDeprecatedFallback<TargetType>("TargetDenyList", "BlacklistTargets", out TargetDenyList))
-			{
-				Module.TargetDenyList = TargetDenyList;
-			}
-
-			UnrealTargetConfiguration[]? TargetConfigurationAllowList;
-			if (InObject.TryGetEnumArrayFieldWithDeprecatedFallback<UnrealTargetConfiguration>("TargetConfigurationAllowList", "WhitelistTargetConfigurations", out TargetConfigurationAllowList))
-			{
-				Module.TargetConfigurationAllowList = TargetConfigurationAllowList;
-			}
-
-			UnrealTargetConfiguration[]? TargetConfigurationDenyList;
-			if (InObject.TryGetEnumArrayFieldWithDeprecatedFallback<UnrealTargetConfiguration>("TargetConfigurationDenyList", "BlacklistTargetConfigurations", out TargetConfigurationDenyList))
-			{
-				Module.TargetConfigurationDenyList = TargetConfigurationDenyList;
-			}
-
 			string[]? ProgramAllowList;
 			if (InObject.TryGetStringArrayFieldWithDeprecatedFallback("ProgramAllowList", "WhitelistPrograms", out ProgramAllowList))
 			{
diff --git a/Engine/Source/Programs/UnrealBuildTool/System/PluginReferenceDescriptor.cs b/Engine/Source/Programs/UnrealBuildTool/System/PluginReferenceDescriptor.cs
--- a/Engine/Source/Programs/UnrealBuildTool/System/PluginReferenceDescriptor.cs
+++ b/Engine/Source/Programs/UnrealBuildTool/System/PluginReferenceDescriptor.cs
@@ -286,36 +286,49 @@ namespace UnrealBuildTool
 		/// <returns>New PluginReferenceDescriptor object</returns>
 		public static PluginReferenceDescriptor FromJsonObject(JsonObject RawObject)
 		{
-			PluginReferenceDescriptor Descriptor = new PluginReferenceDescriptor(RawObject.GetStringField("Name"), null, RawObject.GetBoolField("Enabled"));
-			RawObject.TryGetBoolField("Optional", out Descriptor.bOptional);
-			RawObject.TryGetBoolField("Activate", out Descriptor.bActivate);
-			RawObject.TryGetStringField("Description", out Descriptor.Description);
-			RawObject.TryGetStringField("MarketplaceURL", out Descriptor.MarketplaceURL);
-
-			// Only parse platform information if enabled
-			if (Descriptor.bEnabled)
-			{
-				RawObject.TryGetStringArrayFieldWithDeprecatedFallback("PlatformAllowList", "WhitelistPlatforms", out Descriptor.PlatformAllowList);
-				RawObject.TryGetStringArrayFieldWithDeprecatedFallback("PlatformDenyList", "BlacklistPlatforms", out Descriptor.PlatformDenyList);
-				RawObject.TryGetEnumArrayFieldWithDeprecatedFallback<UnrealTargetConfiguration>("TargetConfigurationAllowList", "WhitelistTargetConfigurations", out Descriptor.TargetConfigurationAllowList);
-				RawObject.TryGetEnumArrayFieldWithDeprecatedFallback<UnrealTargetConfiguration>("TargetConfigurationDenyList", "BlacklistTargetConfigurations", out Descriptor.TargetConfigurationDenyList);
-				RawObject.TryGetEnumArrayFieldWithDeprecatedFallback<TargetType>("TargetAllowList", "WhitelistTargets", out Descriptor.TargetAllowList);
-				RawObject.TryGetEnumArrayFieldWithDeprecatedFallback<TargetType>("TargetDenyList", "BlacklistTargets", out Descriptor.TargetDenyList);
-				RawObject.TryGetStringArrayField("SupportedTargetPlatforms", out Descriptor.SupportedTargetPlatforms);
-				RawObject.TryGetBoolField("HasExplicitPlatforms", out Descriptor.bHasExplicitPlatforms);
-
-				int RequestedVersion = -1;
-				if (RawObject.TryGetIntegerField("Version", out RequestedVersion))
+			string? PluginReferenceName = null;
+			try
+			{
+				PluginReferenceName = RawObject.GetStringField("Name");
+
+				PluginReferenceDescriptor Descriptor = new PluginReferenceDescriptor(PluginReferenceName, null, RawObject.GetBoolField("Enabled"));
+				RawObject.TryGetBoolField("Optional", out Descriptor.bOptional);
+				RawObject.TryGetBoolField("Activate", out Descriptor.bActivate);
+				RawObject.TryGetStringField("Description", out Descriptor.Description);
+				RawObject.TryGetStringField("MarketplaceURL", out Descriptor.MarketplaceURL);
+
+				// Only parse platform information if enabled
+				if (Descriptor.bEnabled)
 				{
-					Descriptor.RequestedVersion = RequestedVersion;
+					RawObject.TryGetStringArrayFieldWithDeprecatedFallback("PlatformAllowList", "WhitelistPlatforms", out Descriptor.PlatformAllowList);
+					RawObject.TryGetStringArrayFieldWithDeprecatedFallback("PlatformDenyList", "BlacklistPlatforms", out Descriptor.PlatformDenyList);
+
+					Descriptor.TargetAllowList = RawObject.GetEnumArrayFieldWithDeprecatedFallback<TargetType>("TargetAllowList", "WhitelistTargets");
+					Descriptor.TargetDenyList = RawObject.GetEnumArrayFieldWithDeprecatedFallback<TargetType>("TargetDenyList", "BlacklistTargets");
+					Descriptor.TargetConfigurationAllowList = RawObject.GetEnumArrayFieldWithDeprecatedFallback<UnrealTargetConfiguration>("TargetConfigurationAllowList", "WhitelistTargetConfigurations");
+					Descriptor.TargetConfigurationDenyList = RawObject.GetEnumArrayFieldWithDeprecatedFallback<UnrealTargetConfiguration>("TargetConfigurationDenyList", "BlacklistTargetConfigurations");
+
+					RawObject.TryGetStringArrayField("SupportedTargetPlatforms", out Descriptor.SupportedTargetPlatforms);
+					RawObject.TryGetBoolField("HasExplicitPlatforms", out Descriptor.bHasExplicitPlatforms);
+
+					int RequestedVersion = -1;
+					if (RawObject.TryGetIntegerField("Version", out RequestedVersion))
+					{
+						Descriptor.RequestedVersion = RequestedVersion;
+					}
+
+					Descriptor.VerifyPlatformNames(Descriptor.PlatformAllowList);
+					Descriptor.VerifyPlatformNames(Descriptor.PlatformDenyList);
+					Descriptor.VerifyPlatformNames(Descriptor.SupportedTargetPlatforms);
 				}
 
-				Descriptor.VerifyPlatformNames(Descriptor.PlatformAllowList);
-				Descriptor.VerifyPlatformNames(Descriptor.PlatformDenyList);
-				Descriptor.VerifyPlatformNames(Descriptor.SupportedTargetPlatforms);
+				return Descriptor;
+			}
+			catch (Exception Ex)
+			{
+				ExceptionUtils.AddContext(Ex, "while parsing plugin reference descriptor '{0}'", PluginReferenceName ?? "<unknown>");
+				throw;
 			}
-
-			return Descriptor;
 		}
 
 		/// <summary>
