Filter out preflights from existing methods that query snapshots.  Whether to include/exclude preflights is something that can be added as an option at a later date if needed.

diff --git a/Engine/Source/Developer/Zen/Private/BuildListRetriever.cpp b/Engine/Source/Developer/Zen/Private/BuildListRetriever.cpp
--- a/Engine/Source/Developer/Zen/Private/BuildListRetriever.cpp
+++ b/Engine/Source/Developer/Zen/Private/BuildListRetriever.cpp
@@ -4,6 +4,7 @@
 #include "Misc/App.h"
 #include "Misc/DateTime.h"
 #include "Misc/EngineVersion.h"
+#include "Serialization/CompactBinaryWriter.h"
 
 namespace UE::Zen::Build
 {
@@ -153,7 +154,16 @@ void FBuildListRetriever::IssueBuildQueries()
 
 		for (const FNamespacePlatformBucketTuple& NamespacePlatformBucket : NamespacePlatformBucketTuples)
 		{
-			ServiceInstance->ListBuilds(NamespacePlatformBucket.Namespace, NamespacePlatformBucket.Bucket, FCbObject(),
+			FCbWriter QueryWriter;
+			QueryWriter.BeginObject();
+			QueryWriter.BeginObject("query");
+			QueryWriter.BeginObject("isPreflight");
+			QueryWriter.AddString("$eq", "false");
+			QueryWriter.EndObject();// isPreflight
+			QueryWriter.EndObject();// query
+			QueryWriter.EndObject();
+
+			ServiceInstance->ListBuilds(NamespacePlatformBucket.Namespace, NamespacePlatformBucket.Bucket, QueryWriter.Save().AsObject(),
 				[this, QueryIndex, Namespace = NamespacePlatformBucket.Namespace, Platform = NamespacePlatformBucket.Platform,
 				PendingQueryState]
 				(TArray<FBuildServiceInstance::FBuildRecord>&& Results) mutable
diff --git a/Engine/Source/Programs/Shared/EpicGames.ProjectStore/ProjectStoreTypes.cs b/Engine/Source/Programs/Shared/EpicGames.ProjectStore/ProjectStoreTypes.cs
--- a/Engine/Source/Programs/Shared/EpicGames.ProjectStore/ProjectStoreTypes.cs
+++ b/Engine/Source/Programs/Shared/EpicGames.ProjectStore/ProjectStoreTypes.cs
@@ -1016,11 +1016,11 @@ namespace EpicGames.ProjectStore
 
 					if (changelist > 0)
 					{
-						request.Content = new StringContent(String.Format("{{\"query\": {{\"runtime\": {{\"$eq\": \"{0}\"}}, \"changelist\": {{\"$lte\": {1}}}, \"branch\": {{\"$eq\": \"{2}\"}}}}}}", runtime, changelist, branch), MediaTypeHeaderValue.Parse("application/json"));
+						request.Content = new StringContent(String.Format("{{\"query\": {{\"isPreflight\": {{\"$eq\": \"false\"}}, {{\"runtime\": {{\"$eq\": \"{0}\"}}, \"changelist\": {{\"$lte\": {1}}}, \"branch\": {{\"$eq\": \"{2}\"}}}}}}", runtime, changelist, branch), MediaTypeHeaderValue.Parse("application/json"));
 					}
 					else
 					{
-						request.Content = new StringContent(String.Format("{{\"query\": {{\"runtime\": {{\"$eq\": \"{0}\"}}, \"branch\": {{\"$eq\": \"{1}\"}}}}}}", runtime, branch), MediaTypeHeaderValue.Parse("application/json"));
+						request.Content = new StringContent(String.Format("{{\"query\": {{\"isPreflight\": {{\"$eq\": \"false\"}}, {{\"runtime\": {{\"$eq\": \"{0}\"}}, \"branch\": {{\"$eq\": \"{1}\"}}}}}}", runtime, branch), MediaTypeHeaderValue.Parse("application/json"));
 					}
 
 					request.Headers.Authorization = new AuthenticationHeaderValue("Bearer", _accessToken);
