[AutoRTFM] Ignore AutoRTFM attributes when parsing a class in UHT.

This change allows UHT to properly parse `class AUTORTFM_DISABLE UMyClass`. Previously, it would misidentify "AUTORTFM_DISABLE" as the name of the class.

diff --git a/Engine/Source/Programs/Shared/EpicGames.UHT/Parsers/UhtSpecifierParser.cs b/Engine/Source/Programs/Shared/EpicGames.UHT/Parsers/UhtSpecifierParser.cs
--- a/Engine/Source/Programs/Shared/EpicGames.UHT/Parsers/UhtSpecifierParser.cs
+++ b/Engine/Source/Programs/Shared/EpicGames.UHT/Parsers/UhtSpecifierParser.cs
@@ -103,7 +103,7 @@ namespace EpicGames.UHT.Parsers
 		}
 
 		/// <summary>
-		/// Perform the specify parsing
+		/// Perform the specifier parsing
 		/// </summary>
 		/// <returns>The parser</returns>
 		public UhtSpecifierParser ParseSpecifiers()
diff --git a/Engine/Source/Programs/Shared/EpicGames.UHT/Tokenizer/UhtTokenReaderGeneralExtensions.cs b/Engine/Source/Programs/Shared/EpicGames.UHT/Tokenizer/UhtTokenReaderGeneralExtensions.cs
--- a/Engine/Source/Programs/Shared/EpicGames.UHT/Tokenizer/UhtTokenReaderGeneralExtensions.cs
+++ b/Engine/Source/Programs/Shared/EpicGames.UHT/Tokenizer/UhtTokenReaderGeneralExtensions.cs
@@ -281,6 +281,10 @@ namespace EpicGames.UHT.Tokenizer
 					tokenReader.ConsumeToken();
 					attributeAction?.Invoke("no_unique_address");
 				}
+				else if (token.IsIdentifier("AUTORTFM_DISABLE") || token.IsIdentifier("AUTORTFM_ENABLE") || token.IsIdentifier("AUTORTFM_INFER") || token.IsIdentifier("AUTORTFM_OPEN") || token.IsIdentifier("AUTORTFM_OPEN_NO_VALIDATION"))
+				{
+					tokenReader.ConsumeToken();
+				}
 				else if (token.IsSymbol("[["))
 				{
 					tokenReader.ConsumeToken();
