Horde: Ensure HttpClient usage in HttpStorageBackend uses a retry policy

Ensures a bit more resiliancy to networking issues

#rnx

diff --git a/Engine/Source/Programs/Shared/EpicGames.Horde/Storage/Backends/HttpStorageBackend.cs b/Engine/Source/Programs/Shared/EpicGames.Horde/Storage/Backends/HttpStorageBackend.cs
--- a/Engine/Source/Programs/Shared/EpicGames.Horde/Storage/Backends/HttpStorageBackend.cs
+++ b/Engine/Source/Programs/Shared/EpicGames.Horde/Storage/Backends/HttpStorageBackend.cs
@@ -97,8 +97,10 @@ namespace EpicGames.Horde.Storage.Backends
 						{
 							request.Headers.Range = new RangeHeaderValue(offset, (length == null) ? null : (offset + (length - 1)));
 						}
+						
+						IAsyncPolicy<HttpResponseMessage> retryPolicy = HordeHttpMessageHandler.CreateDefaultTransientErrorPolicy();
 
-						HttpResponseMessage response = await httpClient.SendAsync(request, cancellationToken);
+						HttpResponseMessage response = await retryPolicy.ExecuteAsync(() => httpClient.SendAsync(request, cancellationToken));
 						await EnsureSuccessAsync(response, cancellationToken);
 
 						Stream stream = await response.Content.ReadAsStreamAsync(cancellationToken);
@@ -210,7 +212,8 @@ namespace EpicGames.Horde.Storage.Backends
 #pragma warning restore CA2000
 
 					request.Content = form;
-					using (HttpResponseMessage response = await httpClient.SendAsync(request, cancellationToken))
+					IAsyncPolicy<HttpResponseMessage> retryPolicy = HordeHttpMessageHandler.CreateDefaultTransientErrorPolicy();
+					using (HttpResponseMessage response = await retryPolicy.ExecuteAsync(() => httpClient.SendAsync(request, cancellationToken)))
 					{
 						if (!response.IsSuccessStatusCode)
 						{
@@ -355,7 +358,8 @@ namespace EpicGames.Horde.Storage.Backends
 
 			using (HttpClient httpClient = _createClient())
 			{
-				using HttpResponseMessage response = await httpClient.PostAsJsonAsync($"{_basePath}", request, cancellationToken: cancellationToken);
+				IAsyncPolicy<HttpResponseMessage> retryPolicy = HordeHttpMessageHandler.CreateDefaultTransientErrorPolicy();
+				using HttpResponseMessage response = await retryPolicy.ExecuteAsync(() => httpClient.PostAsJsonAsync($"{_basePath}", request, cancellationToken: cancellationToken));
 				await EnsureSuccessAsync(response, cancellationToken);
 			}
 		}
