[UBT]
* Fixed so ubt does not crash if UbaHost.dll is not found

diff --git a/Engine/Source/Programs/Shared/EpicGames.UBA/ILogger.cs b/Engine/Source/Programs/Shared/EpicGames.UBA/ILogger.cs
--- a/Engine/Source/Programs/Shared/EpicGames.UBA/ILogger.cs
+++ b/Engine/Source/Programs/Shared/EpicGames.UBA/ILogger.cs
@@ -65,7 +65,14 @@ namespace EpicGames.UBA
 		/// <returns>The ILogger</returns>
 		public static ILogger CreateLogger(Microsoft.Extensions.Logging.ILogger logger, bool makeGlobal)
 		{
-			return new LoggerImpl(logger, makeGlobal);
+			try
+			{
+				return new LoggerImpl(logger, makeGlobal);
+			}
+			catch
+			{
+				return new NullLogger();
+			}
 		}
 
 		/// <summary>
