AutomationTool: Fix matching of script build errors

#rnx

diff --git a/Engine/Source/Programs/AutomationTool/Program.cs b/Engine/Source/Programs/AutomationTool/Program.cs
--- a/Engine/Source/Programs/AutomationTool/Program.cs
+++ b/Engine/Source/Programs/AutomationTool/Program.cs
@@ -466,9 +466,9 @@ namespace AutomationToolDriver
 				// If there was an error before a log writer was attached, do so now
 				if (!Log.HasFileWriter() && ReturnCode != ExitCode.Success)
 				{
-					FileReference errorLogFile = FileReference.Combine(Unreal.WritableEngineDirectory, "Programs/AutomationTool/Saved/Logs/ErrorLog.txt");
+					FileReference errorLogFile = FileReference.Combine(Unreal.WritableEngineDirectory, "Programs", "AutomationTool", "Saved", "Logs", "ErrorLog.txt");
 					Log.AddFileWriterWithoutBackup("UnhandledFileLogListener", errorLogFile);
-					Logger.LogError("AutomationTool encountered an error before a log file was set, see '{File}' for more details", errorLogFile);
+					Logger.LogInformation("AutomationTool encountered an error before a log file was set, see '{File}' for more details", errorLogFile.FullName);
 				}
 
 				// Write the exit code
diff --git a/Engine/Source/Programs/Shared/EpicGames.MsBuild/CsProjBuilder.cs b/Engine/Source/Programs/Shared/EpicGames.MsBuild/CsProjBuilder.cs
--- a/Engine/Source/Programs/Shared/EpicGames.MsBuild/CsProjBuilder.cs
+++ b/Engine/Source/Programs/Shared/EpicGames.MsBuild/CsProjBuilder.cs
@@ -431,8 +431,10 @@ namespace EpicGames.MsBuild
 			};
 			startInfo.EnvironmentVariables["DOTNET_MULTILEVEL_LOOKUP"] = "0"; // use only the bundled dotnet installation - ignore any other/system dotnet install
 
+			using LogEventParser parser = new LogEventParser(logger);
+			parser.Matchers.AddRange(Log.EventParser.Matchers);
+
 			logger.LogDebug("Running: {Application} {Arguments}", startInfo.FileName, startInfo.Arguments);
-			PrefixLogger prefixLogger = new($"[dotnet {arguments.First()}]", logger);
 			using Process dotnetProcess = new Process();
 			dotnetProcess.StartInfo = startInfo;
 			dotnetProcess.OutputDataReceived += (sender, args) =>
@@ -440,7 +442,7 @@ namespace EpicGames.MsBuild
 				string output = args.Data?.TrimEnd() ?? String.Empty;
 				if (!String.IsNullOrEmpty(output))
 				{
-					prefixLogger.LogInformation("{Message}", output);
+					parser.WriteLine(output);
 				}
 			};
 			dotnetProcess.ErrorDataReceived += (sender, args) =>
@@ -448,7 +450,7 @@ namespace EpicGames.MsBuild
 				string output = args.Data?.TrimEnd() ?? String.Empty;
 				if (!String.IsNullOrEmpty(output))
 				{
-					prefixLogger.LogError("{Message}", output);
+					parser.WriteLine(output);
 				}
 			};
 			dotnetProcess.Start();
@@ -456,7 +458,6 @@ namespace EpicGames.MsBuild
 			dotnetProcess.BeginErrorReadLine();
 			await dotnetProcess.WaitForExitAsync(cancellationToken);
 			return dotnetProcess.ExitCode == 0;
-
 		}
 
 		private static async Task<bool> RunDotnetBuildAsync(IEnumerable<FileReference> projects, Dictionary<string, string> globalProperties, CsProjBuildHook hook, ILogger logger, CancellationToken cancellationToken)
