Mark the hash code field of FileSystemReference as non-serializable, so any uses of this class with a serializer that uses the Serializable attribute (e.g. DataContractSerializer) will not persist this field. This is required, as the dotnet documentation makes it clear that hash codes are not guaranteed to be stable across program runs, and must not be persisted, as it can introduce subtle bugs in e.g. collection classes that make use of the hash value.

diff --git a/Engine/Source/Programs/Shared/EpicGames.Core/FileSystemReference.cs b/Engine/Source/Programs/Shared/EpicGames.Core/FileSystemReference.cs
--- a/Engine/Source/Programs/Shared/EpicGames.Core/FileSystemReference.cs
+++ b/Engine/Source/Programs/Shared/EpicGames.Core/FileSystemReference.cs
@@ -61,6 +61,12 @@ namespace EpicGames.Core
 		/// <summary>
 		/// Lazy initialized hash code
 		/// </summary>
+		/// <remarks>
+		/// This must not be (de)serialized, as newly constructed instances may produce a different hash code to
+		/// deserialized results from previous program runs that are otherwise equal.
+		/// </remarks>
+		/// <see href="https://learn.microsoft.com/en-us/dotnet/api/system.string.gethashcode"/>
+		[NonSerialized]
 		private int _hashCode;
 
 		/// <summary>
