Fixed issue where a boolean contained within a namespace would fail to generate a proper SetBit method when building with const init.

#rnx

diff --git a/Engine/Source/Programs/Shared/EpicGames.UHT/Types/Properties/UhtBoolProperty.cs b/Engine/Source/Programs/Shared/EpicGames.UHT/Types/Properties/UhtBoolProperty.cs
--- a/Engine/Source/Programs/Shared/EpicGames.UHT/Types/Properties/UhtBoolProperty.cs
+++ b/Engine/Source/Programs/Shared/EpicGames.UHT/Types/Properties/UhtBoolProperty.cs
@@ -192,14 +192,7 @@ namespace EpicGames.UHT.Types
 		/// <inheritdoc/>
 		public override StringBuilder AppendMemberDef(StringBuilder builder, IUhtPropertyMemberContext context, string name, string nameSuffix, string? offset, int tabs)
 		{
-			if (Outer == context.OuterStruct)
-			{
-				builder.AppendTabs(tabs).Append("void ").AppendNameDef(context, name, nameSuffix).Append("_SetBit(void* Obj)\r\n");
-				builder.AppendTabs(tabs).Append("{\r\n");
-				builder.AppendTabs(tabs + 1).Append("((").Append(context.OuterStructNamespaceName).Append(context.OuterStructSourceName).Append("*)Obj)->").Append(SourceName).Append(" = 1;\r\n");
-				builder.AppendTabs(tabs).Append("}\r\n");
-			}
-
+			AppendSetBit(builder, context, name, nameSuffix, tabs);
 			AppendMemberDefStart(builder, context, name, nameSuffix, offset, tabs, "FBoolPropertyParams",
 				IsNativeBool ?
 				"UECodeGen_Private::EPropertyGenFlags::Bool | UECodeGen_Private::EPropertyGenFlags::NativeBool" :
@@ -230,13 +223,7 @@ namespace EpicGames.UHT.Types
 			{
 				throw new ArgumentException("Cannot create bool property initialization for a bitfield at a specified offset");
 			}
-			if (Outer == context.OuterStruct)
-			{
-				builder.AppendTabs(tabs).Append("void ").AppendNameDef(context, name, nameSuffix).Append("_SetBit(void* Obj)\r\n");
-				builder.AppendTabs(tabs).Append("{\r\n");
-				builder.AppendTabs(tabs + 1).Append("((").Append(context.OuterStructSourceName).Append("*)Obj)->").Append(SourceName).Append(" = 1;\r\n");
-				builder.AppendTabs(tabs).Append("}\r\n");
-			}
+			AppendSetBit(builder, context, name, nameSuffix, tabs);
 			AppendConstInitMemberDefStart(builder, context, name, nameSuffix, outerFunc, IsBitfield ? "0" : offset, tabs);
 			builder.Append("(int32)sizeof(").Append(CppTypeText).Append("), ");
 			builder.Append(IsNativeBool ? "true, " : "false, ");
@@ -313,6 +300,18 @@ namespace EpicGames.UHT.Types
 			return other is UhtBoolProperty;
 		}
 
+		private StringBuilder AppendSetBit(StringBuilder builder, IUhtPropertyMemberContext context, string name, string nameSuffix, int tabs)
+		{
+			if (Outer == context.OuterStruct)
+			{
+				builder.AppendTabs(tabs).Append("void ").AppendNameDef(context, name, nameSuffix).Append("_SetBit(void* Obj)\r\n");
+				builder.AppendTabs(tabs).Append("{\r\n");
+				builder.AppendTabs(tabs + 1).Append("((").Append(context.OuterStructNamespaceName).Append(context.OuterStructSourceName).Append("*)Obj)->").Append(SourceName).Append(" = 1;\r\n");
+				builder.AppendTabs(tabs).Append("}\r\n");
+			}
+			return builder;
+		}
+
 		#region Keyword
 		[UhtPropertyType(Keyword = "bool", Options = UhtPropertyTypeOptions.Simple | UhtPropertyTypeOptions.Immediate)]
 		[SuppressMessage("CodeQuality", "IDE0051:Remove unused private members", Justification = "Attribute accessed method")]
