Horde: New workspace option, maxFileScanConcurrency, to control file scanning and writing independently

diff --git a/Engine/Source/Programs/Horde/Drivers/JobDriver/Utility/WorkspaceInfo.cs b/Engine/Source/Programs/Horde/Drivers/JobDriver/Utility/WorkspaceInfo.cs
--- a/Engine/Source/Programs/Horde/Drivers/JobDriver/Utility/WorkspaceInfo.cs
+++ b/Engine/Source/Programs/Horde/Drivers/JobDriver/Utility/WorkspaceInfo.cs
@@ -279,6 +279,7 @@ namespace JobDriver.Utility
 			const string ManagedWorkspaceValue = "managedWorkspace";
 			const string NumParallelSyncThreadsKey = "numParallelSyncThreads";
 			const string MaxFileConcurrencyKey = "maxFileConcurrency";
+			const string MaxFileScanConcurrencyKey = "maxFileScanConcurrency";
 			const string MinScratchSpaceKey = "minScratchSpace";
 			const string UseHaveTableKey = "useHaveTable";
 			const string UseModTimeKey = "useModTime";
@@ -306,6 +307,10 @@ namespace JobDriver.Utility
 					{
 						options = options with { MaxFileConcurrency = v };
 					}
+					if (Int32.TryParse(MaxFileScanConcurrencyKey, out v))
+					{
+						options = options with { MaxFileScanConcurrency = v };
+					}
 					if (Int32.TryParse(nameValues[MinScratchSpaceKey], out v))
 					{
 						options = options with { MinScratchSpace = v };
diff --git a/Engine/Source/Programs/Shared/EpicGames.Perforce.Managed/ManagedWorkspace.cs b/Engine/Source/Programs/Shared/EpicGames.Perforce.Managed/ManagedWorkspace.cs
--- a/Engine/Source/Programs/Shared/EpicGames.Perforce.Managed/ManagedWorkspace.cs
+++ b/Engine/Source/Programs/Shared/EpicGames.Perforce.Managed/ManagedWorkspace.cs
@@ -84,7 +84,8 @@ namespace EpicGames.Perforce.Managed
 	/// Extra options for configuring ManagedWorkspace
 	/// </summary>
 	/// <param name="NumParallelSyncThreads">Maximum number of threads to sync in parallel</param>
-	/// <param name="MaxFileConcurrency">Maximum number of concurrent file system operations (copying, moving, deleting etc)</param>
+	/// <param name="MaxFileConcurrency">Maximum number of concurrent file system operations (copying, moving, deleting)</param>
+	/// <param name="MaxFileScanConcurrency">Maximum number of current file system scaning operations (file system discovery, traversing)</param>
 	/// <param name="MinScratchSpace">Minimum amount of disk space that must be available after a sync (in megabytes)</param>
 	/// <param name="UseHaveTable">
 	///		Use the client's have table when syncing.
@@ -99,6 +100,7 @@ namespace EpicGames.Perforce.Managed
 	(
 		int NumParallelSyncThreads,
 		int MaxFileConcurrency,
+		int MaxFileScanConcurrency,
 		long MinScratchSpace,
 		bool UseHaveTable,
 		bool Partitioned,
@@ -106,11 +108,12 @@ namespace EpicGames.Perforce.Managed
 		bool UseModTime
 	)
 	{
-		/// <inheritdoc cref="ManagedWorkspaceOptions(Int32, Int32, Int64, Boolean, Boolean, Boolean, Boolean)"/>
-		public ManagedWorkspaceOptions(int? numParallelSyncThreads = null, int? maxFileConcurrency = null, long minScratchSpace = 50L * 1024, bool useHaveTable = true, bool partitioned = false, bool preferNativeClient = false, bool useModTime = false)
+		/// <inheritdoc cref="ManagedWorkspaceOptions(Int32, Int32, Int32, Int64, Boolean, Boolean, Boolean, Boolean)"/>
+		public ManagedWorkspaceOptions(int? numParallelSyncThreads = null, int? maxFileConcurrency = null, int? maxFileScanConcurrency = null, long minScratchSpace = 50L * 1024, bool useHaveTable = true, bool partitioned = false, bool preferNativeClient = false, bool useModTime = false)
 			: this(
 				  numParallelSyncThreads ?? GetDefaultThreadCount(4),
 				  maxFileConcurrency ?? GetDefaultThreadCount(4),
+				  maxFileScanConcurrency ?? GetDefaultThreadCount(4),
 				  minScratchSpace,
 				  useHaveTable,
 				  partitioned,
@@ -596,7 +599,7 @@ namespace EpicGames.Perforce.Managed
 				using TelemetrySpan findFilesSpan = _tracer.StartActiveSpan($"{nameof(ManagedWorkspace)}.FindFilesToClean");
 				Stopwatch timer = Stopwatch.StartNew();
 
-				(FileInfo[] filesToDelete, DirectoryInfo[] directoriesToDelete) refreshResult = await _workspace.RefreshAsync(removeUntracked, _options.MaxFileConcurrency);
+				(FileInfo[] filesToDelete, DirectoryInfo[] directoriesToDelete) refreshResult = await _workspace.RefreshAsync(removeUntracked, _options.MaxFileScanConcurrency);
 				filesToDelete = refreshResult.filesToDelete;
 				directoriesToDelete = refreshResult.directoriesToDelete;
 				findFilesSpan.SetAttribute("horde.mw.num_delete_files", filesToDelete.Length);
@@ -932,7 +935,7 @@ namespace EpicGames.Perforce.Managed
 			_logger.LogInformation("Stage contains {NumFiles:n0} files, {TotalSize:n1}mb", _workspace.GetFiles().Count, _workspace.GetFiles().Sum(x => x._length) / (1024.0 * 1024.0));
 
 			// Print the contents of the workspace
-			string[] differences = await _workspace.FindDifferencesAsync(_options.MaxFileConcurrency);
+			string[] differences = await _workspace.FindDifferencesAsync(_options.MaxFileScanConcurrency);
 			if (differences.Length > 0)
 			{
 				_logger.LogInformation("Local changes:");
@@ -1941,7 +1944,7 @@ namespace EpicGames.Perforce.Managed
 			{
 				Stopwatch timer = Stopwatch.StartNew();
 
-				transaction = await RemoveTransaction.CreateAsync(_workspace, contents, _contentIdToTrackedFile, _options.MaxFileConcurrency);
+				transaction = await RemoveTransaction.CreateAsync(_workspace, contents, _contentIdToTrackedFile, _options.MaxFileScanConcurrency);
 
 				scope.Progress = $"({timer.Elapsed.TotalSeconds:0.0}s)";
 			}
@@ -2103,7 +2106,7 @@ namespace EpicGames.Perforce.Managed
 			{
 				Stopwatch timer = Stopwatch.StartNew();
 
-				transaction = await AddTransaction.CreateAsync(_workspace, stream, _contentIdToTrackedFile, _options.MaxFileConcurrency);
+				transaction = await AddTransaction.CreateAsync(_workspace, stream, _contentIdToTrackedFile, _options.MaxFileScanConcurrency);
 				_workspace = transaction._newWorkspaceRootDir;
 				await SaveAsync(TransactionState.Dirty, cancellationToken);
 
@@ -2232,7 +2235,7 @@ namespace EpicGames.Perforce.Managed
 							(int beginIdx, int endIdx) = batches[batchIdx];
 
 							int[] indexesToUpdate = Enumerable.Range(beginIdx, endIdx - beginIdx).ToArray();
-							ParallelOptions options = new() { MaxDegreeOfParallelism = _options.MaxFileConcurrency, CancellationToken = cancellationToken };
+							ParallelOptions options = new() { MaxDegreeOfParallelism = _options.MaxFileScanConcurrency, CancellationToken = cancellationToken };
 							await Parallel.ForEachAsync(indexesToUpdate, options, (idx, ct) =>
 							{
 								filesToSync[idx]._workspaceFile.UpdateMetadata();
