Add missing flag IsUnderAudit for test audit post

#rnx

diff --git a/Engine/Source/Programs/Horde/Plugins/Build/HordeServer.Build/Jobs/TestData/ITestData.cs b/Engine/Source/Programs/Horde/Plugins/Build/HordeServer.Build/Jobs/TestData/ITestData.cs
--- a/Engine/Source/Programs/Horde/Plugins/Build/HordeServer.Build/Jobs/TestData/ITestData.cs
+++ b/Engine/Source/Programs/Horde/Plugins/Build/HordeServer.Build/Jobs/TestData/ITestData.cs
@@ -492,6 +492,11 @@ namespace HordeServer.Jobs.TestData
 		/// </summary>
 		DateTime? LastAuditDate { get; }
 
+		/// <summary>
+		/// Whether or not the associated test is under audit
+		/// </summary>
+		bool? IsUnderAudit { get; }
+
 		/// <summary>
 		/// The owner user id of the associated test
 		/// </summary>
diff --git a/Engine/Source/Programs/Horde/Plugins/Build/HordeServer.Build/Jobs/TestData/ITestDataCollection.cs b/Engine/Source/Programs/Horde/Plugins/Build/HordeServer.Build/Jobs/TestData/ITestDataCollection.cs
--- a/Engine/Source/Programs/Horde/Plugins/Build/HordeServer.Build/Jobs/TestData/ITestDataCollection.cs
+++ b/Engine/Source/Programs/Horde/Plugins/Build/HordeServer.Build/Jobs/TestData/ITestDataCollection.cs
@@ -317,6 +317,11 @@ namespace HordeServer.Jobs.TestData
 		/// </summary>
 		public DateTime? LastAuditDate { get; set; }
 
+		/// <summary>
+		/// Whether or not the test is under audit
+		/// </summary>
+		public bool? IsUnderAudit {  get; set; }
+
 		/// <summary>
 		/// The owner of the related test
 		/// </summary>
diff --git a/Engine/Source/Programs/Horde/Plugins/Build/HordeServer.Build/Jobs/TestData/TestDataCollectionV2.cs b/Engine/Source/Programs/Horde/Plugins/Build/HordeServer.Build/Jobs/TestData/TestDataCollectionV2.cs
--- a/Engine/Source/Programs/Horde/Plugins/Build/HordeServer.Build/Jobs/TestData/TestDataCollectionV2.cs
+++ b/Engine/Source/Programs/Horde/Plugins/Build/HordeServer.Build/Jobs/TestData/TestDataCollectionV2.cs
@@ -547,6 +547,9 @@ namespace HordeServer.Jobs.TestData
 			[BsonIgnoreIfNull, BsonElement("la")]
 			public DateTime? LastAuditDate { get; set; }
 
+			[BsonIgnoreIfNull, BsonElement("ua")]
+			public bool? IsUnderAudit { get; set; }
+
 			[BsonIgnoreIfNull, BsonElement("o")]
 			public UserId? Owner { get; set; }
 
@@ -588,6 +591,7 @@ namespace HordeServer.Jobs.TestData
 			{
 				EnableSchedule = options.EnableSchedule;
 				LastAuditDate = options.LastAuditDate;
+				IsUnderAudit = options.IsUnderAudit;
 				Owner = options.Owner;
 				Customers = options.Customers;
 				Notifications = options.Notifications?.Select(
diff --git a/Engine/Source/Programs/Horde/Plugins/Build/HordeServer.Build/Jobs/TestData/TestDataController.cs b/Engine/Source/Programs/Horde/Plugins/Build/HordeServer.Build/Jobs/TestData/TestDataController.cs
--- a/Engine/Source/Programs/Horde/Plugins/Build/HordeServer.Build/Jobs/TestData/TestDataController.cs
+++ b/Engine/Source/Programs/Horde/Plugins/Build/HordeServer.Build/Jobs/TestData/TestDataController.cs
@@ -939,6 +939,7 @@ namespace HordeServer.Jobs.TestData
 			TestAuditOptions options = new TestAuditOptions()
 			{
 				LastAuditDate = request.LastAuditDate?.UtcDateTime,
+				IsUnderAudit = request.IsUnderAudit,
 				EnableSchedule = request.EnableSchedule ?? true,
 				Owner = request.Owner == null || request.Owner.Length == 0? null : UserId.Parse(request.Owner),
 				Customers = request.Customers?.Where(u => u.Length > 0).Select(u => UserId.Parse(u)).ToList(),
@@ -989,6 +990,7 @@ namespace HordeServer.Jobs.TestData
 				UserId = audit.UserId.ToString(),
 				UserInputDate = audit.UserInputDate,
 				LastAuditDate = audit.LastAuditDate,
+				IsUnderAudit = audit.IsUnderAudit,
 				Owner = audit.Owner?.ToString(),
 				Customers = audit.Customers == null ? null : audit.Customers.Select(c => c.ToString()).ToList(),
 				Notifications = audit.Notifications?.Where(
diff --git a/Engine/Source/Programs/Shared/EpicGames.Horde/Jobs/TestData/TestDataMessages.cs b/Engine/Source/Programs/Shared/EpicGames.Horde/Jobs/TestData/TestDataMessages.cs
--- a/Engine/Source/Programs/Shared/EpicGames.Horde/Jobs/TestData/TestDataMessages.cs
+++ b/Engine/Source/Programs/Shared/EpicGames.Horde/Jobs/TestData/TestDataMessages.cs
@@ -416,7 +416,7 @@ namespace EpicGames.Horde.Jobs.TestData
 		public string TestId { get; set; } = String.Empty;
 
 		/// <summary>
-		/// Whether the test schedules should be enabled
+		/// Whether or not the test schedules should be enabled
 		/// </summary>
 		public bool? EnableSchedule { get; set; } = true;
 
@@ -425,6 +425,11 @@ namespace EpicGames.Horde.Jobs.TestData
 		/// </summary>
 		public DateTimeOffset? LastAuditDate { get; set; }
 
+		/// <summary>
+		/// Whether or not the test is under audit
+		/// </summary>
+		public bool? IsUnderAudit { get; set; }
+
 		/// <summary>
 		/// The owner id of the related test
 		/// </summary>
@@ -498,7 +503,7 @@ namespace EpicGames.Horde.Jobs.TestData
 		public DateTime UserInputDate { get; set; }
 
 		/// <summary>
-		/// Whether the test schedules should be enabled
+		/// Whether or not the test schedules should be enabled
 		/// </summary>
 		public bool? EnableSchedule { get; set; } = true;
 
@@ -507,6 +512,11 @@ namespace EpicGames.Horde.Jobs.TestData
 		/// </summary>
 		public DateTime? LastAuditDate { get; set; }
 
+		/// <summary>
+		/// Whether or not the test is under audit
+		/// </summary>
+		public bool? IsUnderAudit { get; set; }
+
 		/// <summary>
 		/// The owner id of the related test
 		/// </summary>
