Fix some static analysis errors from prior checkin

#rnx

diff --git a/Engine/Source/Programs/Shared/EpicGames.UHT/Types/Properties/UhtContainerProperty.cs b/Engine/Source/Programs/Shared/EpicGames.UHT/Types/Properties/UhtContainerProperty.cs
--- a/Engine/Source/Programs/Shared/EpicGames.UHT/Types/Properties/UhtContainerProperty.cs
+++ b/Engine/Source/Programs/Shared/EpicGames.UHT/Types/Properties/UhtContainerProperty.cs
@@ -33,7 +33,7 @@ namespace EpicGames.UHT.Types
 		/// </summary>
 		/// <param name="reader">Reader</param>
 		/// <param name="outer">Outer type</param>
-		public UhtContainerBaseProperty(UhtInputCacheReader reader, UhtType outer) : base(reader, outer)
+		protected UhtContainerBaseProperty(UhtInputCacheReader reader, UhtType outer) : base(reader, outer)
 		{
 			ValueProperty = (reader.ReadType() as UhtProperty)!;
 		}
diff --git a/Engine/Source/Programs/Shared/EpicGames.UHT/Types/Properties/UhtMulticastDelegateProperty.cs b/Engine/Source/Programs/Shared/EpicGames.UHT/Types/Properties/UhtMulticastDelegateProperty.cs
--- a/Engine/Source/Programs/Shared/EpicGames.UHT/Types/Properties/UhtMulticastDelegateProperty.cs
+++ b/Engine/Source/Programs/Shared/EpicGames.UHT/Types/Properties/UhtMulticastDelegateProperty.cs
@@ -53,7 +53,7 @@ namespace EpicGames.UHT.Types
 		/// </summary>
 		/// <param name="reader">Reader</param>
 		/// <param name="outer">Outer type</param>
-		public UhtMulticastDelegateProperty(UhtInputCacheReader reader, UhtType outer) : base(reader, outer)
+		protected UhtMulticastDelegateProperty(UhtInputCacheReader reader, UhtType outer) : base(reader, outer)
 		{
 			Function = (reader.ReadType() as UhtFunction)!;
 		}
diff --git a/Engine/Source/Programs/Shared/EpicGames.UHT/Types/Properties/UhtNumericProperty.cs b/Engine/Source/Programs/Shared/EpicGames.UHT/Types/Properties/UhtNumericProperty.cs
--- a/Engine/Source/Programs/Shared/EpicGames.UHT/Types/Properties/UhtNumericProperty.cs
+++ b/Engine/Source/Programs/Shared/EpicGames.UHT/Types/Properties/UhtNumericProperty.cs
@@ -32,7 +32,7 @@ namespace EpicGames.UHT.Types
 		/// </summary>
 		/// <param name="reader">Reader</param>
 		/// <param name="outer">Outer type</param>
-		public UhtNumericProperty(UhtInputCacheReader reader, UhtType outer) : base(reader, outer)
+		protected UhtNumericProperty(UhtInputCacheReader reader, UhtType outer) : base(reader, outer)
 		{
 		}
 
diff --git a/Engine/Source/Programs/Shared/EpicGames.UHT/Types/UhtField.cs b/Engine/Source/Programs/Shared/EpicGames.UHT/Types/UhtField.cs
--- a/Engine/Source/Programs/Shared/EpicGames.UHT/Types/UhtField.cs
+++ b/Engine/Source/Programs/Shared/EpicGames.UHT/Types/UhtField.cs
@@ -148,7 +148,7 @@ namespace EpicGames.UHT.Types
 		/// </summary>
 		/// <param name="reader">Reader</param>
 		/// <param name="outer">Outer type</param>
-		public UhtField(UhtInputCacheReader reader, UhtType outer) : base(reader, outer)
+		protected UhtField(UhtInputCacheReader reader, UhtType outer) : base(reader, outer)
 		{
 			Namespace = reader.ReadNamespace();
 			FieldExportFlags = (UhtFieldExportFlags)reader.ReadUInt64();
diff --git a/Engine/Source/Programs/Shared/EpicGames.UHT/Types/UhtHeaderFile.cs b/Engine/Source/Programs/Shared/EpicGames.UHT/Types/UhtHeaderFile.cs
--- a/Engine/Source/Programs/Shared/EpicGames.UHT/Types/UhtHeaderFile.cs
+++ b/Engine/Source/Programs/Shared/EpicGames.UHT/Types/UhtHeaderFile.cs
@@ -309,7 +309,7 @@ namespace EpicGames.UHT.Types
 			{
 				try
 				{
-					UhtInputCacheReader reader = new(this, entry);
+					using UhtInputCacheReader reader = new(this, entry);
 					reader.ReadHeader();
 					Session.IncrementTotalFilesCacheRead();
 				}
diff --git a/Engine/Source/Programs/Shared/EpicGames.UHT/Types/UhtObject.cs b/Engine/Source/Programs/Shared/EpicGames.UHT/Types/UhtObject.cs
--- a/Engine/Source/Programs/Shared/EpicGames.UHT/Types/UhtObject.cs
+++ b/Engine/Source/Programs/Shared/EpicGames.UHT/Types/UhtObject.cs
@@ -61,7 +61,7 @@ namespace EpicGames.UHT.Types
 		/// </summary>
 		/// <param name="reader">Reader</param>
 		/// <param name="outer">Outer type</param>
-		public UhtObject(UhtInputCacheReader reader, UhtType outer) : base(reader, outer)
+		protected UhtObject(UhtInputCacheReader reader, UhtType outer) : base(reader, outer)
 		{
 			ObjectTypeIndex = Session.GetNextObjectTypeIndex();
 		}
diff --git a/Engine/Source/Programs/Shared/EpicGames.UHT/Types/UhtProperty.cs b/Engine/Source/Programs/Shared/EpicGames.UHT/Types/UhtProperty.cs
--- a/Engine/Source/Programs/Shared/EpicGames.UHT/Types/UhtProperty.cs
+++ b/Engine/Source/Programs/Shared/EpicGames.UHT/Types/UhtProperty.cs
@@ -1264,7 +1264,7 @@ namespace EpicGames.UHT.Types
 		/// </summary>
 		/// <param name="reader">Reader</param>
 		/// <param name="outer">Outer type</param>
-		public UhtProperty(UhtInputCacheReader reader, UhtType outer) : base(reader, outer)
+		protected UhtProperty(UhtInputCacheReader reader, UhtType outer) : base(reader, outer)
 		{
 			VerseNameOverride = reader.ReadString();
 			PropertyCategory = (UhtPropertyCategory)reader.ReadUInt64();
diff --git a/Engine/Source/Programs/Shared/EpicGames.UHT/Types/UhtStruct.cs b/Engine/Source/Programs/Shared/EpicGames.UHT/Types/UhtStruct.cs
--- a/Engine/Source/Programs/Shared/EpicGames.UHT/Types/UhtStruct.cs
+++ b/Engine/Source/Programs/Shared/EpicGames.UHT/Types/UhtStruct.cs
@@ -153,7 +153,7 @@ namespace EpicGames.UHT.Types
 		/// </summary>
 		/// <param name="reader">Reader</param>
 		/// <param name="outer">Outer type</param>
-		public UhtStruct(UhtInputCacheReader reader, UhtType outer) : base(reader, outer)
+		protected UhtStruct(UhtInputCacheReader reader, UhtType outer) : base(reader, outer)
 		{
 			GeneratedCodeVersion = (EGeneratedCodeVersion)reader.ReadUInt8();
 			SuperIdentifier = reader.ReadTokenArray();
diff --git a/Engine/Source/Programs/Shared/EpicGames.UHT/Types/UhtType.cs b/Engine/Source/Programs/Shared/EpicGames.UHT/Types/UhtType.cs
--- a/Engine/Source/Programs/Shared/EpicGames.UHT/Types/UhtType.cs
+++ b/Engine/Source/Programs/Shared/EpicGames.UHT/Types/UhtType.cs
@@ -1022,7 +1022,7 @@ namespace EpicGames.UHT.Types
 		/// </summary>
 		/// <param name="reader">Reader</param>
 		/// <param name="outer">Outer type</param>
-		public UhtType(UhtInputCacheReader reader, UhtType outer)
+		protected UhtType(UhtInputCacheReader reader, UhtType outer)
 		{
 			reader.AddType(this);
 			_headerFile = reader.Header;
diff --git a/Engine/Source/Programs/Shared/EpicGames.UHT/Utils/UhtSession.cs b/Engine/Source/Programs/Shared/EpicGames.UHT/Utils/UhtSession.cs
--- a/Engine/Source/Programs/Shared/EpicGames.UHT/Utils/UhtSession.cs
+++ b/Engine/Source/Programs/Shared/EpicGames.UHT/Utils/UhtSession.cs
@@ -860,7 +860,7 @@ namespace EpicGames.UHT.Utils
 		/// <summary>
 		/// Latest date/time for the main UHT assembly and the plugins
 		/// </summary>
-		public DateTime AssemblyDateTime = DateTime.MaxValue;
+		public DateTime AssemblyDateTime { get; private set; } = DateTime.MaxValue;
 
 		#region Well known types (only avialable after symbol table is built
 		/// <summary>
