Fixed another static analysis issue that I missed

#rnx

diff --git a/Engine/Source/Programs/Shared/EpicGames.UHT/Types/Properties/UhtObjectPropertyBase.cs b/Engine/Source/Programs/Shared/EpicGames.UHT/Types/Properties/UhtObjectPropertyBase.cs
--- a/Engine/Source/Programs/Shared/EpicGames.UHT/Types/Properties/UhtObjectPropertyBase.cs
+++ b/Engine/Source/Programs/Shared/EpicGames.UHT/Types/Properties/UhtObjectPropertyBase.cs
@@ -500,7 +500,7 @@ namespace EpicGames.UHT.Types
 		/// </summary>
 		/// <param name="reader">Reader</param>
 		/// <param name="outer">Outer type</param>
-		public UhtObjectPropertyBase(UhtInputCacheReader reader, UhtType outer) : base(reader, outer)
+		protected UhtObjectPropertyBase(UhtInputCacheReader reader, UhtType outer) : base(reader, outer)
 		{
 			ReferencedClass = (reader.ReadType() as UhtClass)!;
 			Class = (reader.ReadType() as UhtClass)!;
