EpicGames.Core: Increase default connect timeout for HTTP handlers to 25 sec

Previous value of 5 seconds is too tight for many scenarios (e.g for build farm workloads)

diff --git a/Engine/Source/Programs/Shared/EpicGames.Core/HttpClientExtensions.cs b/Engine/Source/Programs/Shared/EpicGames.Core/HttpClientExtensions.cs
--- a/Engine/Source/Programs/Shared/EpicGames.Core/HttpClientExtensions.cs
+++ b/Engine/Source/Programs/Shared/EpicGames.Core/HttpClientExtensions.cs
@@ -249,11 +249,9 @@ namespace EpicGames.Core
 		{
 			return new SocketsHttpHandler
 			{
-				ConnectTimeout = TimeSpan.FromSeconds(5), // This needs to be higher even if there is Polly above because dns has to make it otherwise it will timeout over and over again
+				ConnectTimeout = TimeSpan.FromSeconds(25),
 				PooledConnectionLifetime = TimeSpan.FromMinutes(15),
 				PooledConnectionIdleTimeout = TimeSpan.FromMinutes(5),
-				//UseProxy = false,
-				//AllowAutoRedirect = false,
 			};
 		}
 
