Piping verse unmangled struct name through to UVerseStruct metadata the same way we do for UVerseClasses

diff --git a/Engine/Source/Programs/Shared/EpicGames.UHT/Parsers/UhtScriptStructParser.cs b/Engine/Source/Programs/Shared/EpicGames.UHT/Parsers/UhtScriptStructParser.cs
--- a/Engine/Source/Programs/Shared/EpicGames.UHT/Parsers/UhtScriptStructParser.cs
+++ b/Engine/Source/Programs/Shared/EpicGames.UHT/Parsers/UhtScriptStructParser.cs
@@ -139,6 +139,7 @@ namespace EpicGames.UHT.Parsers
 					if (scriptStruct.IsVerseField)
 					{
 						scriptStruct.MetaData.Add(UhtNames.IsBlueprintBase, "false");
+						scriptStruct.MetaData.Add(UhtNames.DisplayName, scriptStruct.VerseName!);
 					}
 
 					UhtCompilerDirective compilerDirective = topScope.HeaderParser.GetCurrentCompositeCompilerDirective();
