Fixed an error where we retry a read request but re-use the same request message resulting in a "The request message was already sent. Cannot send the same request message multiple times" error.

diff --git a/Engine/Source/Programs/Shared/EpicGames.Horde/Storage/Backends/HttpStorageBackend.cs b/Engine/Source/Programs/Shared/EpicGames.Horde/Storage/Backends/HttpStorageBackend.cs
--- a/Engine/Source/Programs/Shared/EpicGames.Horde/Storage/Backends/HttpStorageBackend.cs
+++ b/Engine/Source/Programs/Shared/EpicGames.Horde/Storage/Backends/HttpStorageBackend.cs
@@ -91,30 +91,36 @@ namespace EpicGames.Horde.Storage.Backends
 
 				using (HttpClient httpClient = _createClient())
 				{
-					using (HttpRequestMessage request = new HttpRequestMessage(HttpMethod.Get, $"{_basePath}/blobs/{locator}"))
+					IAsyncPolicy<HttpResponseMessage> retryPolicy = HordeHttpMessageHandler.CreateDefaultTransientErrorPolicy();
+
+					HttpResponseMessage response = await retryPolicy.ExecuteAsync(async () =>
 					{
-						if (offset != 0 || length != null)
+						using (HttpRequestMessage request = new HttpRequestMessage(HttpMethod.Get, $"{_basePath}/blobs/{locator}"))
 						{
-							request.Headers.Range = new RangeHeaderValue(offset, (length == null) ? null : (offset + (length - 1)));
-						}
-						
-						IAsyncPolicy<HttpResponseMessage> retryPolicy = HordeHttpMessageHandler.CreateDefaultTransientErrorPolicy();
+							if (offset != 0 || length != null)
+							{
+								request.Headers.Range = new RangeHeaderValue(offset, (length == null) ? null : (offset + (length - 1)));
+							}
 
-						HttpResponseMessage response = await retryPolicy.ExecuteAsync(() => httpClient.SendAsync(request, cancellationToken));
-						await EnsureSuccessAsync(response, cancellationToken);
+							HttpResponseMessage requestResponse = await httpClient.SendAsync(request, cancellationToken);
 
-						Stream stream = await response.Content.ReadAsStreamAsync(cancellationToken);
-
-						try
-						{
-							Interlocked.Add(ref _numBytes, response.Content.Headers.ContentLength ?? stream.Length);
-						}
-						catch
-						{
+							return requestResponse;
 						}
+					});
+
+					await EnsureSuccessAsync(response, cancellationToken);
 
-						return stream;
+					Stream stream = await response.Content.ReadAsStreamAsync(cancellationToken);
+
+					try
+					{
+						Interlocked.Add(ref _numBytes, response.Content.Headers.ContentLength ?? stream.Length);
 					}
+					catch
+					{
+					}
+
+					return stream;
 				}
 			}
 			finally
