Fixed a bug where we re-use the same HttpRequestMessage when retrying a WriteRequest.

diff --git a/Engine/Source/Programs/Shared/EpicGames.Horde/Storage/Backends/HttpStorageBackend.cs b/Engine/Source/Programs/Shared/EpicGames.Horde/Storage/Backends/HttpStorageBackend.cs
--- a/Engine/Source/Programs/Shared/EpicGames.Horde/Storage/Backends/HttpStorageBackend.cs
+++ b/Engine/Source/Programs/Shared/EpicGames.Horde/Storage/Backends/HttpStorageBackend.cs
@@ -93,7 +93,7 @@ namespace EpicGames.Horde.Storage.Backends
 				{
 					IAsyncPolicy<HttpResponseMessage> retryPolicy = HordeHttpMessageHandler.CreateDefaultTransientErrorPolicy();
 
-					HttpResponseMessage response = await retryPolicy.ExecuteAsync(async () =>
+					using HttpResponseMessage response = await retryPolicy.ExecuteAsync(async () =>
 					{
 						using (HttpRequestMessage request = new HttpRequestMessage(HttpMethod.Get, $"{_basePath}/blobs/{locator}"))
 						{
@@ -193,44 +193,50 @@ namespace EpicGames.Horde.Storage.Backends
 		{
 			using (HttpClient httpClient = _createClient())
 			{
-				using (HttpRequestMessage request = (locator != null)
+				IAsyncPolicy<HttpResponseMessage> retryPolicy = HordeHttpMessageHandler.CreateDefaultTransientErrorPolicy();
+				using HttpResponseMessage response = await retryPolicy.ExecuteAsync(async () =>
+				{
+					using (HttpRequestMessage request = (locator != null)
 					? new HttpRequestMessage(HttpMethod.Put, $"{_basePath}/blobs/{locator.Value}")
 					: new HttpRequestMessage(HttpMethod.Post, $"{_basePath}/blobs"))
-				{
-					using MultipartFormDataContent form = new MultipartFormDataContent();
-					if (streamContent != null)
-					{
-						form.Add(streamContent, "file", "filename");
-					}
-#pragma warning disable CA2000 // Disposed by form
-					if (imports != null)
 					{
-						foreach (BlobLocator import in imports)
+						using MultipartFormDataContent form = new MultipartFormDataContent();
+						if (streamContent != null)
 						{
-							form.Add(new StringContent(import.ToString()), "import");
+							form.Add(streamContent, "file", "filename");
 						}
-						if (imports.Count == 0)
+#pragma warning disable CA2000 // Disposed by form
+						if (imports != null)
 						{
-							form.Add(new StringContent("true"), "leaf");
+							foreach (BlobLocator import in imports)
+							{
+								form.Add(new StringContent(import.ToString()), "import");
+							}
+							if (imports.Count == 0)
+							{
+								form.Add(new StringContent("true"), "leaf");
+							}
 						}
-					}
-					form.Add(new StringContent(prefix ?? String.Empty), "prefix");
+						form.Add(new StringContent(prefix ?? String.Empty), "prefix");
 #pragma warning restore CA2000
 
-					request.Content = form;
-					IAsyncPolicy<HttpResponseMessage> retryPolicy = HordeHttpMessageHandler.CreateDefaultTransientErrorPolicy();
-					using (HttpResponseMessage response = await retryPolicy.ExecuteAsync(() => httpClient.SendAsync(request, cancellationToken)))
-					{
-						if (!response.IsSuccessStatusCode)
-						{
-							string responseText = await response.Content.ReadAsStringAsync(cancellationToken);
-							throw new StorageException($"{request.Method} to {request.RequestUri} failed ({response.StatusCode}). Response: {responseText}");
-						}
+						request.Content = form;
+
+						HttpResponseMessage requestResponse = await httpClient.SendAsync(request, cancellationToken);
 
-						WriteBlobResponse? data = await response.Content.ReadFromJsonAsync<WriteBlobResponse>(cancellationToken: cancellationToken);
-						return data!;
+						return requestResponse;
 					}
+				} );
+				
+				if (!response.IsSuccessStatusCode)
+				{
+					string responseText = await response.Content.ReadAsStringAsync(cancellationToken);
+
+					throw new StorageException($"{response.RequestMessage?.Method} to {response.RequestMessage?.RequestUri} failed ({response.StatusCode}). Response: {responseText}");
 				}
+
+				WriteBlobResponse? data = await response.Content.ReadFromJsonAsync<WriteBlobResponse>(cancellationToken: cancellationToken);
+				return data!;
 			}
 		}
 
