Fixing an issue where we were disposing of the HttpResponse content before reading it.

diff --git a/Engine/Source/Programs/Shared/EpicGames.Horde/Storage/Backends/HttpStorageBackend.cs b/Engine/Source/Programs/Shared/EpicGames.Horde/Storage/Backends/HttpStorageBackend.cs
--- a/Engine/Source/Programs/Shared/EpicGames.Horde/Storage/Backends/HttpStorageBackend.cs
+++ b/Engine/Source/Programs/Shared/EpicGames.Horde/Storage/Backends/HttpStorageBackend.cs
@@ -93,7 +93,7 @@ namespace EpicGames.Horde.Storage.Backends
 				{
 					IAsyncPolicy<HttpResponseMessage> retryPolicy = HordeHttpMessageHandler.CreateDefaultTransientErrorPolicy();
 
-					using HttpResponseMessage response = await retryPolicy.ExecuteAsync(async () =>
+					HttpResponseMessage response = await retryPolicy.ExecuteAsync(async () =>
 					{
 						using (HttpRequestMessage request = new HttpRequestMessage(HttpMethod.Get, $"{_basePath}/blobs/{locator}"))
 						{
