Mark package calsses like _Root as Hidden in the editor to exclude them from class pickers.
Create reference domains for all plugins with verse types, this allows creating plugins that do not support content without ensuring in the class picker.

diff --git a/Engine/Plugins/Editor/AssetReferenceRestrictions/Source/AssetReferenceRestrictions/Private/AssetReferencingDomains.cpp b/Engine/Plugins/Editor/AssetReferenceRestrictions/Source/AssetReferenceRestrictions/Private/AssetReferencingDomains.cpp
--- a/Engine/Plugins/Editor/AssetReferenceRestrictions/Source/AssetReferenceRestrictions/Private/AssetReferencingDomains.cpp
+++ b/Engine/Plugins/Editor/AssetReferenceRestrictions/Source/AssetReferenceRestrictions/Private/AssetReferencingDomains.cpp
@@ -256,7 +256,7 @@ void FDomainDatabase::RebuildFromScratch()
 		{
 			BuildUnrestrictedDomainFromPlugin(Plugin);
 		}
-		else if (Plugin->CanContainContent())
+		else if (Plugin->CanContainContent() || Plugin->CanContainVerse())
 		{
 			BuildDomainFromPlugin(Plugin);
 		}
@@ -452,7 +452,7 @@ void FDomainDatabase::BuildDomainFromPlugin(TSharedRef<IPlugin> Plugin)
 		if (Dependencies.bEnabled)
 		{
 			TSharedPtr<IPlugin> PluginWeDependOn = IPluginManager::Get().FindPlugin(Dependencies.Name);
-			if (PluginWeDependOn.IsValid() && PluginWeDependOn->CanContainContent())
+			if (PluginWeDependOn.IsValid() && (PluginWeDependOn->CanContainContent() || PluginWeDependOn->CanContainVerse()))
 			{
 				Domain->DomainsVisibleFromHere.Add(FindOrAddDomainByName(Dependencies.Name));
 			}
diff --git a/Engine/Source/Programs/Shared/EpicGames.UHT/Parsers/UhtNamespaceParser.cs b/Engine/Source/Programs/Shared/EpicGames.UHT/Parsers/UhtNamespaceParser.cs
--- a/Engine/Source/Programs/Shared/EpicGames.UHT/Parsers/UhtNamespaceParser.cs
+++ b/Engine/Source/Programs/Shared/EpicGames.UHT/Parsers/UhtNamespaceParser.cs
@@ -90,6 +90,8 @@ namespace EpicGames.UHT.Parsers
 
 		private static UhtParseResult ParseVModule(UhtParsingScope parentScope, ref UhtToken token)
 		{
+			// This function is equivalent to FSolClassGenerator::Prepare and they should be kept in sync
+
 			UhtClass classObj = new(parentScope.HeaderFile, parentScope.HeaderParser.GetNamespace(), parentScope.ScopeType, token.InputLine);
 			classObj.ClassType = UhtClassType.VModule;
 
@@ -104,7 +106,9 @@ namespace EpicGames.UHT.Parsers
 			classObj.GeneratedBodyLineNumber = classObj.LineNumber;
 			classObj.GeneratedBodyAccessSpecifier = UhtAccessSpecifier.Public;
 			classObj.PrologLineNumber = topScope.TokenReader.InputLine;
-			classObj.ClassFlags |= EClassFlags.Native;
+
+			// Hide module classes from the editor pickers
+			classObj.ClassFlags |= EClassFlags.Native | EClassFlags.Hidden;
 			classObj.SuperIdentifier = [new UhtToken(UhtTokenType.Identifier, token.UngetPos, token.UngetLine, token.InputStartPos, token.InputLine, "UObject", false)];
 
 			topScope.AddFormattedCommentsAsTooltipMetaData();
