Refactor the UHT counters.  Add addition counters to help track down why we are getting cache file level misses.

#rnx

diff --git a/Engine/Source/Programs/Shared/EpicGames.UHT/Types/UhtHeaderFile.cs b/Engine/Source/Programs/Shared/EpicGames.UHT/Types/UhtHeaderFile.cs
--- a/Engine/Source/Programs/Shared/EpicGames.UHT/Types/UhtHeaderFile.cs
+++ b/Engine/Source/Programs/Shared/EpicGames.UHT/Types/UhtHeaderFile.cs
@@ -282,7 +282,7 @@ namespace EpicGames.UHT.Types
 				AddReferencedHeader(Session.NoExportTypesHeader, UhtHeaderReferenceType.Passive);
 			}
 			_sourceFile.Read();
-			Session.IncrementTotalFilesRead();
+			Session.FilesRead.Increment();
 		}
 
 		/// <summary>
@@ -311,11 +311,11 @@ namespace EpicGames.UHT.Types
 				{
 					using UhtInputCacheReader reader = new(this, entry);
 					reader.ReadHeader();
-					Session.IncrementTotalFilesCacheRead();
+					Session.FilesCacheRead.Increment();
 				}
 				catch (Exception ex)
 				{
-					Session.IncrementTotalFileCacheReadErrors();
+					Session.FilesCacheReadErrors.Increment();
 					Session.Logger.LogInformation("File {FileName} failed to deserialize from input cache: {Message}", _sourceFile.FilePath, ex.Message);
 				}
 			}
diff --git a/Engine/Source/Programs/Shared/EpicGames.UHT/Utils/UhtInputCache.cs b/Engine/Source/Programs/Shared/EpicGames.UHT/Utils/UhtInputCache.cs
--- a/Engine/Source/Programs/Shared/EpicGames.UHT/Utils/UhtInputCache.cs
+++ b/Engine/Source/Programs/Shared/EpicGames.UHT/Utils/UhtInputCache.cs
@@ -1042,7 +1042,7 @@ namespace EpicGames.UHT.Utils
 	/// </summary>
 	public class UhtInputCache
 	{
-		private static int CacheVersion = 2;
+		private static int s_cacheVersion = 2;
 		private readonly Dictionary<string, UhtInputFile> _files = [];
 		private readonly ReadOnlyMemory<byte> _data;
 
@@ -1097,6 +1097,7 @@ namespace EpicGames.UHT.Utils
 			FileReference cacheFile = GetCacheFileReference(module);
 			if (!File.Exists(cacheFile.FullName))
 			{
+				module.Session.CacheFilesSkippedFNF.Increment();
 				return new();
 			}
 
@@ -1107,20 +1108,23 @@ namespace EpicGames.UHT.Utils
 				DateTime tocDateTime = File.GetLastAccessTimeUtc(cacheFile.FullName);
 				if (tocDateTime < module.Session.AssemblyDateTime)
 				{
+					module.Session.CacheFilesSkippedAssemblyDateTime.Increment();
 					return new();
 				}
 
 				byte[] headerData = File.ReadAllBytes(cacheFile.FullName);
 				MemoryReader reader = new(headerData);
 				int version = reader.ReadInt32();
-				if (version != CacheVersion)
+				if (version != s_cacheVersion)
 				{
+					module.Session.CacheFilesSkippedVersion.Increment();
 					return new();
 				}
 
 				ReadOnlyMemory<byte> cachedConfigData = reader.ReadVariableLengthBytes();
 				if (!cachedConfigData.Span.SequenceEqual(configData))
 				{
+					module.Session.CacheFilesSkippedConfigChanged.Increment();
 					return new();
 				}
 
@@ -1134,17 +1138,19 @@ namespace EpicGames.UHT.Utils
 						cachedFile.CacheEnd < 0 || cachedFile.CacheEnd > data.Length ||
 						cachedFile.CacheStart >= cachedFile.CacheEnd)
 					{
+						module.Session.CacheFilesSkippedCorrupted.Increment();
 						return new();
 					}
 				}
 
-				module.Session.IncrementTotalCacheFilesRead();
+				module.Session.CacheFilesRead.Increment();
 				return new(cachedFiles, reader.RemainingMemory);
 			}
 			catch (Exception ex)
 			{
 				module.Session.Logger.LogTrace("Error reading input cache {Message}", ex.Message);
 				DeleteCache(module);
+				module.Session.CacheFilesSkippedCorrupted.Increment();
 				return new();
 			}
 		}
@@ -1187,7 +1193,7 @@ namespace EpicGames.UHT.Utils
 					using Stream cacheStream = File.Open(cacheFile.FullName, FileMode.Create, FileAccess.Write, FileShare.Read);
 
 					using ChunkedMemoryWriter headerWriter = new();
-					headerWriter.WriteInt32(CacheVersion);
+					headerWriter.WriteInt32(s_cacheVersion);
 					headerWriter.WriteVariableLengthBytes(configData);
 					headerWriter.WriteList(cachedFiles, (writer, cachedFile) => cachedFile.Write(headerWriter));
 					cacheStream.Write(headerWriter.ToByteArray());
@@ -1198,7 +1204,7 @@ namespace EpicGames.UHT.Utils
 					}
 				}
 
-				module.Session.IncrementTotalCacheFilesWritten();
+				module.Session.CacheFilesWritten.Increment();
 			}
 			catch(Exception ex) 
 			{
diff --git a/Engine/Source/Programs/Shared/EpicGames.UHT/Utils/UhtSession.cs b/Engine/Source/Programs/Shared/EpicGames.UHT/Utils/UhtSession.cs
--- a/Engine/Source/Programs/Shared/EpicGames.UHT/Utils/UhtSession.cs
+++ b/Engine/Source/Programs/Shared/EpicGames.UHT/Utils/UhtSession.cs
@@ -633,6 +633,27 @@ namespace EpicGames.UHT.Utils
 		Verify,
 	};
 
+	/// <summary>
+	/// Helper class that implements a simple counter
+	/// </summary>
+	public class UhtCounter
+	{
+		private int _counter = 0;
+
+		/// <summary>
+		/// Return the current count
+		/// </summary>
+		public int Count => Interlocked.CompareExchange(ref _counter, 0, 0);
+
+		/// <summary>
+		/// Increment the current count
+		/// </summary>
+		public void Increment()
+		{
+			Interlocked.Increment(ref _counter);
+		}
+	}
+
 	/// <summary>
 	/// Session object that represents a UHT run
 	/// </summary>
@@ -1031,115 +1052,77 @@ namespace EpicGames.UHT.Utils
 		/// <summary>
 		/// The number of errors
 		/// </summary>
-		public long ErrorCount => Interlocked.Read(ref _errorCount);
-		private long _errorCount = 0;
+		public UhtCounter ErrorCount { get; } = new();
 
 		/// <summary>
 		/// The number of warnings
 		/// </summary>
-		public long WarningCount => Interlocked.Read(ref _warningCount);
-		private long _warningCount = 0;
+		public UhtCounter WarningCount { get; } = new();
 
 		/// <summary>
 		/// True if any errors have occurred or warnings if warnings are to be treated as errors 
 		/// </summary>
-		public bool HasErrors => ErrorCount > 0 || (WarningsAsErrors && WarningCount > 0);
+		public bool HasErrors => ErrorCount.Count > 0 || (WarningsAsErrors && WarningCount.Count > 0);
 
 		/// <summary>
 		/// Number of new/modified files written to disk
 		/// </summary>
-		public int TotalFilesWritten => Interlocked.Add(ref _totalFilesWritten, 0);
-		private int _totalFilesWritten = 0;
+		public UhtCounter FilesWritten { get; } = new();
 
 		/// <summary>
 		/// Number of source files read directly from disk
 		/// </summary>
-		public int TotalFilesRead => Interlocked.Add(ref _totalFilesRead, 0);
-		private int _totalFilesRead = 0;
+		public UhtCounter FilesRead { get; } = new();
 
 		/// <summary>
 		/// Number of source files read from cache
 		/// </summary>
-		public int TotalFilesCacheRead => Interlocked.Add(ref _totalFilesCacheRead, 0);
-		private int _totalFilesCacheRead = 0;
-
-		/// <summary>
-		/// Number of cache files read
-		/// </summary>
-		public int TotalCacheFilesRead => Interlocked.Add(ref _totalCacheFilesRead, 0);
-		private int _totalCacheFilesRead = 0;
+		public UhtCounter FilesCacheRead { get; } = new();
 
 		/// <summary>
-		/// Number of cache files written
+		/// Number of source files failed to deserialize
 		/// </summary>
-		public int TotalCacheFilesWritten => Interlocked.Add(ref _totalCacheFilesWritten, 0);
-		private int _totalCacheFilesWritten = 0;
+		public UhtCounter FilesCacheReadErrors { get; } = new();
 
 		/// <summary>
-		/// Number of source files failed to deserialize
+		/// Number of source files failed to serialize
 		/// </summary>
-		public int TotalFilesCacheReadErrors => Interlocked.Add(ref _totalFilesCacheReadErrors, 0);
-		private int _totalFilesCacheReadErrors = 0;
+		public UhtCounter FilesCacheWriteErrors { get; } = new();
 
 		/// <summary>
-		/// Number of source files failed to serialize
+		/// Number of cache files read
 		/// </summary>
-		public int TotalFilesCacheWriteErrors => Interlocked.Add(ref _totalFilesCacheWriteErrors, 0);
-		private int _totalFilesCacheWriteErrors = 0;
+		public UhtCounter CacheFilesRead { get; } = new();
 
 		/// <summary>
-		/// Increment the total number of files read from disk
+		/// Number of cache files written
 		/// </summary>
-		/// <returns></returns>
-		public int IncrementTotalFilesRead()
-		{
-			return Interlocked.Increment(ref _totalFilesRead);
-		}
+		public UhtCounter CacheFilesWritten { get; } = new();
 
 		/// <summary>
-		/// Increment the total number of files read from cache
+		/// Number of skipped caches files because they were not found
 		/// </summary>
-		/// <returns></returns>
-		public int IncrementTotalFilesCacheRead()
-		{
-			return Interlocked.Increment(ref _totalFilesCacheRead);
-		}
+		public UhtCounter CacheFilesSkippedFNF { get; } = new();
 
 		/// <summary>
-		/// Increment the total number of cache files read
+		/// Number of skipped caches files because the assembly time did not match
 		/// </summary>
-		/// <returns></returns>
-		public int IncrementTotalCacheFilesRead()
-		{
-			return Interlocked.Increment(ref _totalCacheFilesRead);
-		}
+		public UhtCounter CacheFilesSkippedAssemblyDateTime { get; } = new();
 
 		/// <summary>
-		/// Increment the total number of cache files written
+		/// Number of skipped caches files because version number mismatch
 		/// </summary>
-		/// <returns></returns>
-		public int IncrementTotalCacheFilesWritten()
-		{
-			return Interlocked.Increment(ref _totalCacheFilesWritten);
-		}
+		public UhtCounter CacheFilesSkippedVersion { get; } = new();
 
 		/// <summary>
-		/// Increment the total number of errors deserializing cached files
+		/// Number of skipped caches files because configuration changed
 		/// </summary>
-		/// <returns></returns>
-		public int IncrementTotalFileCacheReadErrors()
-		{
-			return Interlocked.Increment(ref _totalFilesCacheReadErrors);
-		}
+		public UhtCounter CacheFilesSkippedConfigChanged { get; } = new();
 
 		/// <summary>
-		/// Increment the total number of errors serializing cached files
+		/// Number of skipped caches files because file is corrupted
 		/// </summary>
-		/// <returns></returns>
-		public int IncrementTotalFileCacheWriteErrors()
-		{
-			return Interlocked.Increment(ref _totalFilesCacheWriteErrors);
-		}
+		public UhtCounter CacheFilesSkippedCorrupted { get; } = new();
 		#endregion
 
 		private Task? _referenceDeleteTask = null;
@@ -1761,7 +1744,7 @@ namespace EpicGames.UHT.Utils
 					break;
 
 				case UhtMessageType.Warning:
-					Interlocked.Increment(ref _warningCount);
+					WarningCount.Increment();
 					break;
 
 				case UhtMessageType.Info:
@@ -1778,7 +1761,7 @@ namespace EpicGames.UHT.Utils
 		/// </summary>
 		public void SignalError()
 		{
-			Interlocked.Increment(ref _errorCount);
+			ErrorCount.Increment();
 		}
 
 		/// <summary>
@@ -2171,7 +2154,7 @@ namespace EpicGames.UHT.Utils
 						catch (Exception ex)
 						{
 							Logger.LogInformation("File {Filename} failed to serialize to input cache: {Message}", headerFile.FilePath, ex.Message);
-							IncrementTotalFileCacheWriteErrors();
+							FilesCacheWriteErrors.Increment();
 						}
 					}
 
@@ -3086,7 +3069,7 @@ namespace EpicGames.UHT.Utils
 						{
 							if (output.Saved)
 							{
-								Interlocked.Increment(ref _totalFilesWritten);
+								FilesWritten.Increment();
 							}
 						}
 						foreach (string dep in factory.ExternalDependencies)
diff --git a/Engine/Source/Programs/UnrealBuildTool/Modes/UnrealHeaderToolMode.cs b/Engine/Source/Programs/UnrealBuildTool/Modes/UnrealHeaderToolMode.cs
--- a/Engine/Source/Programs/UnrealBuildTool/Modes/UnrealHeaderToolMode.cs
+++ b/Engine/Source/Programs/UnrealBuildTool/Modes/UnrealHeaderToolMode.cs
@@ -861,24 +861,27 @@ namespace UnrealBuildTool.Modes
 				stopwatch.Stop();
 				double elapsed = stopwatch.Elapsed.TotalSeconds;
 				LogLevel level = elapsed > 5 || Options.Mode != null || Unreal.IsBuildMachine() ? LogLevel.Information : LogLevel.Debug;
-				if (!Session.HasErrors || Session.TotalFilesWritten != 0)
+				if (!Session.HasErrors || Session.FilesWritten.Count != 0)
 				{
 					string ActualTargetName = Session.Manifest != null ? Session.Manifest.TargetName : "UnknownTarget";
 					if (Session.EnableInputCacheRead || Session.EnableInputCacheWrite)
 					{
-						Logger.Log(level, "UHT processed {TargetName} in {Time} seconds (Generated Files: {TotalFilesWritten}, Source Files Read (disk/cache): {TotalFilesRead}/{TotalFileCacheRead}, Cache Files (read/written): {TotalCacheFilesRead}/{TotalCacheFilesWritten})",
-							ActualTargetName, stopwatch.Elapsed.TotalSeconds, Session.TotalFilesWritten, Session.TotalFilesRead, Session.TotalFilesCacheRead, Session.TotalCacheFilesRead, Session.TotalCacheFilesWritten);
+						Logger.Log(level,
+							"UHT processed {TargetName} in {Time} seconds, Generated Files: {FilesWritten}, Source Files Read (disk/cache): {FilesRead}/{FileCacheRead}, Cache Files (written/read): {CacheFilesWritten}/{CacheFilesRead}, Cache Skipped (fnf/version/assembly/config/corrupted) {CacheFilesSkippedFNF}/{CacheFilesSkippedVersion}/{CacheFilesSkippedAssemblyDateTime}/{CacheFilesSkippedConfigChanged}/{CacheFilesSkippedCorrupted}",
+							ActualTargetName, stopwatch.Elapsed.TotalSeconds, Session.FilesWritten.Count, Session.FilesRead.Count, Session.FilesCacheRead.Count,
+							Session.CacheFilesWritten.Count, Session.CacheFilesRead.Count, Session.CacheFilesSkippedFNF.Count, Session.CacheFilesSkippedVersion.Count,
+							Session.CacheFilesSkippedAssemblyDateTime.Count, Session.CacheFilesSkippedConfigChanged.Count, Session.CacheFilesSkippedCorrupted.Count);
 					}
 					else
 					{
 						Logger.Log(level, "UHT processed {TargetName} in {Time} seconds ({TotalFilesWritten} generated files written)",
-							ActualTargetName, stopwatch.Elapsed.TotalSeconds, Session.TotalFilesWritten);
+							ActualTargetName, stopwatch.Elapsed.TotalSeconds, Session.FilesWritten.Count);
 					}
 				}
-				if (Session.TotalFilesCacheReadErrors != 0 || Session.TotalFilesCacheWriteErrors != 0)
+				if (Session.FilesCacheReadErrors.Count != 0 || Session.FilesCacheWriteErrors.Count != 0)
 				{
 					Logger.Log(level, "UHT input cache had errors ({TotalFilesCacheReadErrors} file failed deserialization, {TotalFilesCacheReadErrors} files failed serialization)",
-						Session.TotalFilesCacheReadErrors, Session.TotalFilesCacheWriteErrors);
+						Session.FilesCacheReadErrors.Count, Session.FilesCacheWriteErrors.Count);
 				}
 
 				return (int)(Session.HasErrors ? CompilationResult.OtherCompilationError : CompilationResult.Succeeded);
