Handle nulls

diff --git a/Engine/Source/Programs/Shared/EpicGames.Core/LogEventParser.cs b/Engine/Source/Programs/Shared/EpicGames.Core/LogEventParser.cs
--- a/Engine/Source/Programs/Shared/EpicGames.Core/LogEventParser.cs
+++ b/Engine/Source/Programs/Shared/EpicGames.Core/LogEventParser.cs
@@ -439,11 +439,14 @@ namespace EpicGames.Core
 			{
 				// Try to match an event
 				List<LogEvent>? events = null;
-				if (Regex.IsMatch(_buffer[0]!, "<-- Suspend Log Parsing -->", RegexOptions.IgnoreCase))
+				string line = _buffer[0]!;
+				if (line == null)
+					line = "";
+				if (Regex.IsMatch(line, "<-- Suspend Log Parsing -->", RegexOptions.IgnoreCase))
 				{
 					_matchingEnabled--;
 				}
-				else if (Regex.IsMatch(_buffer[0]!, "<-- Resume Log Parsing -->", RegexOptions.IgnoreCase))
+				else if (Regex.IsMatch(line, "<-- Resume Log Parsing -->", RegexOptions.IgnoreCase))
 				{
 					_matchingEnabled++;
 				}
@@ -464,7 +467,7 @@ namespace EpicGames.Core
 				{
 					foreach (Regex ignorePattern in IgnorePatterns)
 					{
-						if (ignorePattern.IsMatch(_buffer[0]!))
+						if (ignorePattern.IsMatch(line))
 						{
 							events = null;
 							break;
@@ -480,7 +483,7 @@ namespace EpicGames.Core
 				}
 				else
 				{
-					_logger.Log(LogLevel.Information, KnownLogEvents.None, _buffer[0]!, null, (state, exception) => state);
+					_logger.Log(LogLevel.Information, KnownLogEvents.None, line, null, (state, exception) => state);
 					_buffer.MoveNext();
 				}
 			}
