Fix compile warnings in Horde.

diff --git a/Engine/Source/Programs/Horde/Drivers/JobDriver/Parser/LogParser.cs b/Engine/Source/Programs/Horde/Drivers/JobDriver/Parser/LogParser.cs
--- a/Engine/Source/Programs/Horde/Drivers/JobDriver/Parser/LogParser.cs
+++ b/Engine/Source/Programs/Horde/Drivers/JobDriver/Parser/LogParser.cs
@@ -1,6 +1,5 @@
 // Copyright Epic Games, Inc. All Rights Reserved.
 
-using System.Reflection;
 using System.Text.RegularExpressions;
 using EpicGames.Core;
 using EpicGames.LogMatchers;
diff --git a/Engine/Source/Programs/Shared/EpicGames.Core/LogEventParser.cs b/Engine/Source/Programs/Shared/EpicGames.Core/LogEventParser.cs
--- a/Engine/Source/Programs/Shared/EpicGames.Core/LogEventParser.cs
+++ b/Engine/Source/Programs/Shared/EpicGames.Core/LogEventParser.cs
@@ -450,7 +450,9 @@ namespace EpicGames.Core
 				List<LogEvent>? events = null;
 				string line = _buffer[0]!;
 				if (line == null)
+				{
 					line = "";
+				}
 				if (Regex.IsMatch(line, "<-- Suspend Log Parsing -->", RegexOptions.IgnoreCase))
 				{
 					_matchingEnabled--;
