Cloud Storage - Added non-async methods to make it easier to interact with cloud storage from non-async methods (primarily UAT, its best to use async if possible).
Also worked around issue in zen with wildcard expansions when attempting to download parts of a build and targeting parts in the root (current directory specification was not supported)

diff --git a/Engine/Source/Programs/Shared/EpicGames.ProjectStore/CloudStorage/CloudStorage.cs b/Engine/Source/Programs/Shared/EpicGames.ProjectStore/CloudStorage/CloudStorage.cs
--- a/Engine/Source/Programs/Shared/EpicGames.ProjectStore/CloudStorage/CloudStorage.cs
+++ b/Engine/Source/Programs/Shared/EpicGames.ProjectStore/CloudStorage/CloudStorage.cs
@@ -7,7 +7,10 @@ using System.ComponentModel.DataAnnotations;
 using System.Diagnostics;
 using System.IO;
 using System.Linq;
+using System.Net;
 using System.Runtime.CompilerServices;
+using System.Text.Json;
+using System.Text.Json.Serialization;
 using System.Text.RegularExpressions;
 using System.Threading;
 using System.Threading.Tasks;
@@ -151,6 +154,11 @@ namespace EpicGames.ProjectStore
 			return DownloadBuildAsync(ns, bucketId, buildId, outputDir, includeFiles, excludeFiles, downloadOptions, progress, optionsOverride, cancellationToken);
 		}
 
+		public void DownloadBuild(StringId bucketId, CbObjectId buildId, DirectoryReference outputDir, List<FilePattern>? includeFiles = null, List<FilePattern>? excludeFiles = null, CloudBuildDownloadOptions? downloadOptions = null, IProgress<(float, string)>? progress = null, CloudBuildOptionsOverride? optionsOverride = null, CancellationToken cancellationToken = default(CancellationToken))
+		{
+			DownloadBuildAsync(bucketId, buildId, outputDir, includeFiles, excludeFiles, downloadOptions, progress, optionsOverride, cancellationToken).Wait(cancellationToken);
+		}
+
 		/// <inheritdoc />
 		public Task DownloadBuildFromUriAsync(Uri uri, DirectoryReference outputDir, List<FilePattern>? includeFiles = null, List<FilePattern>? excludeFiles = null, CloudBuildDownloadOptions? downloadOptions = null, IProgress<(float, string)>? progress = null, CloudBuildOptionsOverride? optionsOverride = null, CancellationToken cancellationToken = default(CancellationToken))
 		{
@@ -207,11 +215,17 @@ namespace EpicGames.ProjectStore
 				}
 			}
 
-			return Path.Combine(filePattern.BaseDirectory.MakeRelativeTo(outputDirectory), String.Join(String.Empty, GetTokens(filePattern)));
+			string s = Path.Combine(filePattern.BaseDirectory.MakeRelativeTo(outputDirectory), String.Join(String.Empty, GetTokens(filePattern)));
+			// zen assumes all wildcards are rooted in the local folder and does not support specifying it
+			if (s.StartsWith(".\\", StringComparison.OrdinalIgnoreCase))
+			{
+				s = s.Replace(".\\", "", StringComparison.OrdinalIgnoreCase);
+			}
+			return s;
 		}
 
 		/// <inheritdoc />
-		public async IAsyncEnumerable<FindBuildResponse> FindBuildsAsync(Regex bucketRegex, CbObject query, CloudBuildSearchOptions? searchOptions = null, CloudBuildOptionsOverride? optionsOverride = null, [EnumeratorCancellation] CancellationToken cancellationToken = default(CancellationToken))
+		public IEnumerable<FindBuildResponse> FindBuilds(Regex bucketRegex, CbObject query, CloudBuildSearchOptions? searchOptions = null, CloudBuildOptionsOverride? optionsOverride = null, CancellationToken cancellationToken = default(CancellationToken))
 		{
 			string? ns = GetCurrentNamespace(optionsOverride);
 			if (ns == null)
@@ -254,14 +268,14 @@ namespace EpicGames.ProjectStore
 				}
 				writer.EndObject();
 
-				await File.WriteAllBytesAsync(tempInputFile, writer.ToByteArray(), cancellationToken);
+				File.WriteAllBytes(tempInputFile, writer.ToByteArray());
 
 				string cmdline = $"builds list {_options.Value.GetHostCliArgs(optionsOverride)} --namespace {ns} --query-path \"{tempInputFile}\" --result-path \"{tempOutputFile}\"";
 
-				int exitCode = await RunZenCliAsync(cmdline, optionsOverride: optionsOverride, cancellationToken: cancellationToken);
+				int exitCode = RunZenCliAsync(cmdline, optionsOverride: optionsOverride, cancellationToken: cancellationToken).Result;
 
 				EnsureZenExitCodeSuccess(exitCode);
-				byte[] b = await File.ReadAllBytesAsync(tempOutputFile, cancellationToken);
+				byte[] b = File.ReadAllBytes(tempOutputFile);
 				CbObject o = new CbObject(b);
 
 				bool partialResult = o["partialResult"].AsBool();
@@ -325,9 +339,17 @@ namespace EpicGames.ProjectStore
 			}
 		}
 
-		/// <inheritdoc />
-		public IAsyncEnumerable<FindBuildResponse> FindBuildsAsync(string project, string? type, string? branch, CbObject query, CloudBuildSearchOptions? searchOptions = null,
-			CloudBuildOptionsOverride? optionsOverride = null, CancellationToken cancellationToken = default(CancellationToken))
+		public async IAsyncEnumerable<FindBuildResponse> FindBuildsAsync(Regex bucketRegex, CbObject query, CloudBuildSearchOptions? searchOptions = null, CloudBuildOptionsOverride? optionsOverride = null, [EnumeratorCancellation] CancellationToken cancellationToken = default(CancellationToken))
+		{
+			await Task.CompletedTask;
+
+			foreach (FindBuildResponse build in FindBuilds(bucketRegex, query, searchOptions, optionsOverride, cancellationToken))
+			{
+				yield return build;
+			}
+		}
+
+		private static Regex BuildRegex(string project, string? type, string? branch)
 		{
 			string regex = ICloudStorage.SanitizeBucketValue(project);
 			if (!String.IsNullOrEmpty(type))
@@ -340,8 +362,22 @@ namespace EpicGames.ProjectStore
 			}
 			// wildcard after this prefix
 			regex += "\\.*";
+			return new Regex(regex);
+		}
+
+		/// <inheritdoc />
+		public IAsyncEnumerable<FindBuildResponse> FindBuildsAsync(string project, string? type, string? branch, CbObject query, CloudBuildSearchOptions? searchOptions = null, CloudBuildOptionsOverride? optionsOverride = null, CancellationToken cancellationToken = default(CancellationToken))
+		{
+			Regex regex = BuildRegex(project, type, branch);
+
+			return FindBuildsAsync(regex, query, searchOptions, optionsOverride, cancellationToken);
+		}
+		
+		public IEnumerable<FindBuildResponse> FindBuilds(string project, string? type, string? branch, CbObject query, CloudBuildSearchOptions? searchOptions = null, CloudBuildOptionsOverride? optionsOverride = null, CancellationToken cancellationToken = default(CancellationToken))
+		{
+			Regex regex = BuildRegex(project, type, branch);
 
-			return FindBuildsAsync(new Regex(regex), query, searchOptions, optionsOverride, cancellationToken);
+			return FindBuilds(regex, query, searchOptions, optionsOverride, cancellationToken);
 		}
 
 		private static ZenExitCodes AsZenExitCodeValue(int exitCode)
diff --git a/Engine/Source/Programs/Shared/EpicGames.ProjectStore/CloudStorage/ICloudStorage.cs b/Engine/Source/Programs/Shared/EpicGames.ProjectStore/CloudStorage/ICloudStorage.cs
--- a/Engine/Source/Programs/Shared/EpicGames.ProjectStore/CloudStorage/ICloudStorage.cs
+++ b/Engine/Source/Programs/Shared/EpicGames.ProjectStore/CloudStorage/ICloudStorage.cs
@@ -25,7 +25,7 @@ namespace EpicGames.ProjectStore
 	/// <param name="Commit">The commit of the build, optional and can be a changelist, hash or other global unique value</param>
 	/// <param name="CreatedAt">The time this build was created</param>
 	/// <param name="Metadata">Metadata associated with the build</param>
-	public readonly record struct FindBuildResponse(
+	public record FindBuildResponse(
 		string Name,
 		CbObjectId BuildId,
 		StringId BucketId,
@@ -264,6 +264,21 @@ namespace EpicGames.ProjectStore
 		/// <returns></returns>
 		Task DownloadBuildAsync(StringId bucketId, CbObjectId buildId, DirectoryReference outputDir, List<FilePattern>? includeFiles = null, List<FilePattern>? excludeFiles = null, CloudBuildDownloadOptions? downloadOptions = null, IProgress<(float, string)>? progress = null, CloudBuildOptionsOverride? optionsOverride = null, CancellationToken cancellationToken = default(CancellationToken));
 
+		/// <summary>
+		/// Download a build using all the required fields
+		/// </summary>
+		/// <param name="bucketId">The bucket</param>
+		/// <param name="buildId">The id of the build</param>
+		/// <param name="outputDir">The directory to write the build to</param>
+		/// <param name="includeFiles">List of file patterns of files to download, if specified only files matching these patterns are downloaded</param>
+		/// <param name="excludeFiles">List of file patterns to exclude. If combined with include files these apply after the include files</param>
+		/// <param name="downloadOptions">Advanced options for reconfiguring the download</param>
+		/// <param name="progress">Progress reporting</param>
+		/// <param name="optionsOverride">Optional per invocation overrides of options</param>
+		/// <param name="cancellationToken">A cancellation token</param>
+		/// <returns></returns>
+		void DownloadBuild(StringId bucketId, CbObjectId buildId, DirectoryReference outputDir, List<FilePattern>? includeFiles = null, List<FilePattern>? excludeFiles = null, CloudBuildDownloadOptions? downloadOptions = null, IProgress<(float, string)>? progress = null, CloudBuildOptionsOverride? optionsOverride = null, CancellationToken cancellationToken = default(CancellationToken));
+
 		/// <summary>
 		/// Download a build using a single uri value that encodes everything needed to locate the build
 		/// </summary>
@@ -289,6 +304,17 @@ namespace EpicGames.ProjectStore
 		/// <returns></returns>
 		IAsyncEnumerable<FindBuildResponse> FindBuildsAsync(Regex bucketRegex, CbObject query, CloudBuildSearchOptions? searchOptions = null, CloudBuildOptionsOverride? optionsOverride = null, CancellationToken cancellationToken = default(CancellationToken));
 
+		/// <summary>
+		/// Search for builds in UE Cloud Storage using a regular expression to describe which buckets the build might be in
+		/// </summary>
+		/// <param name="bucketRegex">A regex that describe which buckets the build might be in</param>
+		/// <param name="query">CbObject that describes the query, see UE Cloud Storage documentation</param>
+		/// <param name="searchOptions">Optional set of options for the search</param>
+		/// <param name="optionsOverride">Optional per invocation overrides of options</param>		
+		/// <param name="cancellationToken">A cancellation token</param>
+		/// <returns></returns>
+		IEnumerable<FindBuildResponse> FindBuilds(Regex bucketRegex, CbObject query, CloudBuildSearchOptions? searchOptions = null, CloudBuildOptionsOverride? optionsOverride = null, CancellationToken cancellationToken = default(CancellationToken));
+
 		/// <summary>
 		/// Search for builds in UE Cloud Storage using common components
 		/// </summary>
@@ -302,6 +328,19 @@ namespace EpicGames.ProjectStore
 		/// <returns></returns>
 		IAsyncEnumerable<FindBuildResponse> FindBuildsAsync(string project, string? type, string? branch, CbObject query, CloudBuildSearchOptions? searchOptions = null, CloudBuildOptionsOverride? optionsOverride = null, CancellationToken cancellationToken = default(CancellationToken));
 
+		/// <summary>
+		/// Search for builds in UE Cloud Storage using common components
+		/// </summary>
+		/// <param name="project">The project you are searching builds for</param>
+		/// <param name="type">The type of build you are searching for, if omitted all types are considered</param>
+		/// <param name="branch">The branch you which to search, if omitted all branches are searched</param>
+		/// <param name="query">CbObject that describes the query, see UE Cloud Storage documentation</param>
+		/// <param name="searchOptions">Optional set of options for the search</param>
+		/// <param name="optionsOverride">Optional per invocation overrides of options</param>		
+		/// <param name="cancellationToken">A cancellation token</param>
+		/// <returns></returns>
+		IEnumerable<FindBuildResponse> FindBuilds(string project, string? type, string? branch, CbObject query, CloudBuildSearchOptions? searchOptions = null, CloudBuildOptionsOverride? optionsOverride = null, CancellationToken cancellationToken = default(CancellationToken));
+
 		/// <summary>
 		/// List the contents of a build
 		/// </summary>
