Cloud Storage - Added missing xml doc tags on implementation for non-async overloads

diff --git a/Engine/Source/Programs/Shared/EpicGames.ProjectStore/CloudStorage/CloudStorage.cs b/Engine/Source/Programs/Shared/EpicGames.ProjectStore/CloudStorage/CloudStorage.cs
--- a/Engine/Source/Programs/Shared/EpicGames.ProjectStore/CloudStorage/CloudStorage.cs
+++ b/Engine/Source/Programs/Shared/EpicGames.ProjectStore/CloudStorage/CloudStorage.cs
@@ -154,6 +154,7 @@ namespace EpicGames.ProjectStore
 			return DownloadBuildAsync(ns, bucketId, buildId, outputDir, includeFiles, excludeFiles, downloadOptions, progress, optionsOverride, cancellationToken);
 		}
 
+		/// <inheritdoc />
 		public void DownloadBuild(StringId bucketId, CbObjectId buildId, DirectoryReference outputDir, List<FilePattern>? includeFiles = null, List<FilePattern>? excludeFiles = null, CloudBuildDownloadOptions? downloadOptions = null, IProgress<(float, string)>? progress = null, CloudBuildOptionsOverride? optionsOverride = null, CancellationToken cancellationToken = default(CancellationToken))
 		{
 			DownloadBuildAsync(bucketId, buildId, outputDir, includeFiles, excludeFiles, downloadOptions, progress, optionsOverride, cancellationToken).Wait(cancellationToken);
@@ -339,6 +340,7 @@ namespace EpicGames.ProjectStore
 			}
 		}
 
+		/// <inheritdoc />
 		public async IAsyncEnumerable<FindBuildResponse> FindBuildsAsync(Regex bucketRegex, CbObject query, CloudBuildSearchOptions? searchOptions = null, CloudBuildOptionsOverride? optionsOverride = null, [EnumeratorCancellation] CancellationToken cancellationToken = default(CancellationToken))
 		{
 			await Task.CompletedTask;
@@ -373,6 +375,7 @@ namespace EpicGames.ProjectStore
 			return FindBuildsAsync(regex, query, searchOptions, optionsOverride, cancellationToken);
 		}
 		
+		/// <inheritdoc />
 		public IEnumerable<FindBuildResponse> FindBuilds(string project, string? type, string? branch, CbObject query, CloudBuildSearchOptions? searchOptions = null, CloudBuildOptionsOverride? optionsOverride = null, CancellationToken cancellationToken = default(CancellationToken))
 		{
 			Regex regex = BuildRegex(project, type, branch);
