[UBT]
* Fixed a bug where cache path hashes were registered too early (while they were being created in parallel threads). Changed so path hashes are now registered just before start processing actions

diff --git a/Engine/Source/Programs/Shared/EpicGames.UBA/ICacheClient.cs b/Engine/Source/Programs/Shared/EpicGames.UBA/ICacheClient.cs
--- a/Engine/Source/Programs/Shared/EpicGames.UBA/ICacheClient.cs
+++ b/Engine/Source/Programs/Shared/EpicGames.UBA/ICacheClient.cs
@@ -32,11 +32,9 @@ namespace EpicGames.UBA
 		public abstract void Disconnect();
 
 		/// <summary>
-		/// Register path with string that will be hashed. All files under this path will be ignored and instead refer to this path hash
+		/// Traverse all path hashes that have been registered and add them to the cache server. Note, this should be done just before processing actions
 		/// </summary>
-		/// <param name="path">Path that will represent all files under that path</param>
-		/// <param name="hash">String that is unique to the content of data under path</param>
-		public abstract bool RegisterPathHash(string path, string hash);
+		public abstract void PopulatePathHashes();
 
 		/// <summary>
 		/// Write to cache
diff --git a/Engine/Source/Programs/Shared/EpicGames.UBA/Impl/CacheClientImpl.cs b/Engine/Source/Programs/Shared/EpicGames.UBA/Impl/CacheClientImpl.cs
--- a/Engine/Source/Programs/Shared/EpicGames.UBA/Impl/CacheClientImpl.cs
+++ b/Engine/Source/Programs/Shared/EpicGames.UBA/Impl/CacheClientImpl.cs
@@ -96,11 +96,6 @@ namespace EpicGames.UBA.Impl
 				return false;
 			}
 
-			foreach (KeyValuePair<string, string> item in Utils.PathHashes)
-			{
-				CacheClient_RegisterPathHash(_handle, item.Key, item.Value);
-			}
-
 			if (desiredConnectionCount > 1)
 			{
 				_connectingTasks = new();
@@ -126,9 +121,12 @@ namespace EpicGames.UBA.Impl
 			CacheClient_Disconnect(_handle);
 		}
 
-		public bool RegisterPathHash(string path, string hash)
+		public void PopulatePathHashes()
 		{
-			return CacheClient_RegisterPathHash(_handle, path, hash) != 0;
+			foreach (KeyValuePair<string, string> item in Utils.PathHashes)
+			{
+				CacheClient_RegisterPathHash(_handle, item.Key, item.Value);
+			}
 		}
 
 		public bool WriteToCache(uint bucket, nint processHandle, byte[] inputs, uint inputsSize, byte[] outputs, uint outputsSize)
diff --git a/Engine/Source/Programs/UnrealBuildTool/Executors/UnrealBuildAccelerator/UBAExecutor.cs b/Engine/Source/Programs/UnrealBuildTool/Executors/UnrealBuildAccelerator/UBAExecutor.cs
--- a/Engine/Source/Programs/UnrealBuildTool/Executors/UnrealBuildAccelerator/UBAExecutor.cs
+++ b/Engine/Source/Programs/UnrealBuildTool/Executors/UnrealBuildAccelerator/UBAExecutor.cs
@@ -927,6 +927,12 @@ namespace UnrealBuildTool
 				return ProcessFinishedResponse.None;
 			});
 
+
+			foreach (ICacheClient client in cache)
+			{
+				client.PopulatePathHashes();
+			}
+
 			scheduler.Start();
 
 			ConsoleCancelEventHandler? cancelHandler = null;
