Making constinit uobjects compatible with modular builds.

Pointers between modules are initialized with a new codegen type UE::CodeGen::ConstInit::TObjectRefParam which can store an encoded reference or a pointer into unions such as in FObjectHandle.
These pointers are fixed up when boostrapping UObjects with a method LinkCompiledInPointerFields on relevant UObject/FField-derived types.

UObject registration is changed to have a single registration struct per module which also includes an array of linkable objects, and an array of pointers to those arrays in other modules for runtime linking.
Pointers are encoded as the index of the dependency module within the current module's array of dependencies, and the index of the target object in its owning module's array of linkable objects.

Misc changes:
FObjectHandle is a struct when UE_WITH_CONSTINIT_UOBJECT is set, and has a new union member which is an encoded compiled-in pointer, which is not distinguishable from an encoded pointer for UE_WITH_OBJECT_HANDLE_LATE_RESOLVE.
ObjectPtr_Private::Friend::GetHandleRef accessor for FObjectHandle& to mutate pointers during runtime linking.
FField::GetFieldClassPrivate implementation moved outline because the field's ConstInitClass can't be linked across dlls.
New tagged constructors for UClass for constructing UClass::StaticClass and UObject::StaticClass because other modules cannot see Z_ConstInit_UClass_UClass and Z_ConstInit_UClass_UObject in a ternary expression.
Delegate calling of UObject static methods to templated class TUObjectCppClassStaticThunks so that constint UClasses outside of CoreUObject can take pointers to the methods of TUObjectCppClassStaticThunks rather than the methods of UObject, which can't be used in constinit expressions.

#Synced-CL Unknown

diff --git a/Engine/Source/Programs/Shared/EpicGames.Core/UHTTypes.cs b/Engine/Source/Programs/Shared/EpicGames.Core/UHTTypes.cs
--- a/Engine/Source/Programs/Shared/EpicGames.Core/UHTTypes.cs
+++ b/Engine/Source/Programs/Shared/EpicGames.Core/UHTTypes.cs
@@ -148,18 +148,21 @@ namespace EpicGames.Core
 		{
 			CompiledInObjectFormat = Enum.Parse<UhtCompiledInObjectFormat>(reader.ReadString()!);
 			GenerateEnumMaxValues = reader.ReadBool();
+			ValidateModuleDependencies = reader.ReadBool();
 		}
 
 		public readonly void Write(BinaryArchiveWriter writer)
 		{
 			writer.WriteString(CompiledInObjectFormat.ToString());
 			writer.WriteBool(GenerateEnumMaxValues);
+			writer.WriteBool(ValidateModuleDependencies);
 		}
 
 		public readonly void Write(IMemoryWriter writer)
 		{
 			writer.WriteString(CompiledInObjectFormat.ToString());
 			writer.WriteBoolean(GenerateEnumMaxValues);
+			writer.WriteBoolean(ValidateModuleDependencies);
 		}
 
 		/// <summary>
@@ -174,6 +177,12 @@ namespace EpicGames.Core
 		/// </summary>
 		public bool GenerateEnumMaxValues { get; set; } = false;
 
+		/// <summary>
+		/// Placeholder value to control whether or not UHT validates module dependencies from UBT to find missing 
+		/// dependencies which are needed because of generated code (e.g. FStructBaseChain references)
+		/// </summary>
+		public bool ValidateModuleDependencies { get; set; } = false;
+		
 		/// <summary>
 		/// If true, enable the input cache read
 		/// </summary>
@@ -202,6 +211,9 @@ namespace EpicGames.Core
 			public List<string> InternalHeaders { get; set; } = [];
 			public List<string> PrivateHeaders { get; set; } = [];
 			public List<string> PublicDefines { get; set; } = [];
+			// Other modules in the manifest which this module has a link dependency on in UBT
+			// Can be used to validate dependencies which are required because of generated code
+			public List<string> ModuleDependencies { get; set; } = [];
 			public string? GeneratedCPPFilenameBase { get; set; } = null;
 			public bool SaveExportedHeaders { get; set; } = false;
 			[JsonConverter(typeof(JsonStringEnumConverter))]
diff --git a/Engine/Source/Programs/Shared/EpicGames.UHT/Exporters/CodeGen/UhtCodeGenExpressions.cs b/Engine/Source/Programs/Shared/EpicGames.UHT/Exporters/CodeGen/UhtCodeGenExpressions.cs
new file mode 100644
--- /dev/null
+++ b/Engine/Source/Programs/Shared/EpicGames.UHT/Exporters/CodeGen/UhtCodeGenExpressions.cs
@@ -0,0 +1,79 @@
+// Copyright Epic Games, Inc. All Rights Reserved.
+
+using System;
+using EpicGames.UHT.Types;
+using EpicGames.UHT.Utils;
+
+namespace EpicGames.UHT.Exporters.CodeGen
+{
+	/// <summary>
+	/// Utility functions to construct structs which can be used in interpolated strings.
+	/// It is recommended to import this type into a namespace with 'using static' so its
+	/// methods can be used without the type name prefix for brevity.
+	/// </summary>
+	internal static class UhtExpressionFactory
+	{
+		static readonly ReadOnlyMemory<char> s_tabs = string.Empty.PadLeft(100, '\t').AsMemory();
+	
+		internal static TabsExpression Tabs(int count) => new TabsExpression(count, s_tabs);
+		internal static SingletonRefExpression SingletonRef(IUhtObjectLinker provider, UhtObject? obj, UhtSingletonType type) 
+			=> new SingletonRefExpression(provider, obj, type, false);
+		internal static SingletonRefExpression ConstInitSingletonRef(IUhtObjectLinker provider, UhtObject? obj, bool explicitType=false) 
+			=> new SingletonRefExpression(provider, obj, UhtSingletonType.ConstInit, explicitType);
+	}
+
+	#region Expression structs
+	internal readonly record struct TabsExpression(int Count, ReadOnlyMemory<char> Tabs) : ISpanFormattable
+	{
+		public string ToString(string? format, IFormatProvider? formatProvider) => $"{this}";
+
+		public bool TryFormat(Span<char> destination, out int charsWritten, ReadOnlySpan<char> format, IFormatProvider? provider)
+		{
+			if (Count > Tabs.Length)
+			{
+				throw new UhtIceException($"Trying to indent something by more than {Tabs.Length}");
+			}
+			if (Count <=0)
+			{
+				charsWritten = 0;
+				return true;
+			}
+			if (Tabs.Span[0..Count].TryCopyTo(destination))
+			{
+				charsWritten = Count;
+				return true;
+			}
+			charsWritten = 0;
+			return false;
+		}
+	}	
+
+	/// <summary>
+	/// Utility class for simple expressions such as variable names derived from a type name.
+	/// Intended to be used as a return type from factory methods with a more explicit name, e.g. ModuleLinkableObjects
+	/// </summary>
+	/// <param name="Prefix">Optional prefix e.g. Z_SomeCategory_</param>
+	/// <param name="Unique">The part of the string derived from the type/object</param>
+	/// <param name="Suffix">Optional suffix if any, e.g. ()</param>
+	internal readonly record struct SimpleExpression(string? Prefix, string Unique, string? Suffix = null) : ISpanFormattable
+	{
+		public string ToString(string? format, IFormatProvider? formatProvider) => $"{this}";
+
+		public bool TryFormat(Span<char> destination, out int charsWritten, ReadOnlySpan<char> format, IFormatProvider? provider)
+			=> destination.TryWrite($"{Prefix}{Unique}{Suffix}", out charsWritten);
+	}
+
+	internal readonly record struct SingletonRefExpression(IUhtObjectLinker Provider, UhtObject? Obj, UhtSingletonType Type, bool Explicit) : ISpanFormattable
+	{
+		public string ToString(string? format, IFormatProvider? formatProvider) => $"{this}";
+		public bool TryFormat(Span<char> destination, out int charsWritten, ReadOnlySpan<char> format, IFormatProvider? provider) => (Obj, Type, Obj?.IsCrossModuleRef(Provider)) switch 
+		{
+			(null, _, _) => destination.TryWrite($"nullptr", out charsWritten), 
+			(_, UhtSingletonType.ConstInit, false) => destination.TryWrite($"&{Provider.GetSingletonName(Obj, UhtSingletonType.ConstInit)}", out charsWritten),
+			(_, UhtSingletonType.ConstInit, true) => 
+				destination.TryWrite($"UHT_LINK_{Obj.Module.ShortName}({Provider.GetModuleIndex(Obj.Module)}, {Provider.GetObjectIndex(Obj)}, &{Provider.GetSingletonName(Obj, UhtSingletonType.ConstInit)})", out charsWritten),
+			(_, _, _) => destination.TryWrite($"{Provider.GetSingletonName(Obj, Type)}", out charsWritten),
+		};
+	}
+	#endregion
+}
diff --git a/Engine/Source/Programs/Shared/EpicGames.UHT/Exporters/CodeGen/UhtCodeGenerator.cs b/Engine/Source/Programs/Shared/EpicGames.UHT/Exporters/CodeGen/UhtCodeGenerator.cs
--- a/Engine/Source/Programs/Shared/EpicGames.UHT/Exporters/CodeGen/UhtCodeGenerator.cs
+++ b/Engine/Source/Programs/Shared/EpicGames.UHT/Exporters/CodeGen/UhtCodeGenerator.cs
@@ -37,7 +37,35 @@ namespace EpicGames.UHT.Exporters.CodeGen
 		/// </summary>
 		Statics = 3,
 	}
-	
+
+	/// <summary>
+	/// Mirror of UhtSingletonType as flags for storing sets of singleton types
+	/// </summary>
+	[Flags]
+	public enum UhtSingletonTypeFlag
+	{
+		/// <summary>
+		/// Flag matching nothing
+		/// </summary>
+		None = 0,
+		/// <summary>
+		/// Flag for UhtSingletonType.Registered
+		/// </summary>
+		Registered = 1 << UhtSingletonType.Registered,
+		/// <summary>
+		/// Flag for UhtSingletonType.Unregistered
+		/// </summary>
+		Unregistered = 1 << UhtSingletonType.Unregistered,
+		/// <summary>
+		/// Flag for UhtSingletonType.ConstInit
+		/// </summary>
+		ConstInit = 1 << UhtSingletonType.ConstInit,
+		/// <summary>
+		/// Flag for UhtSingletonType.Statics
+		/// </summary>
+		Statics = 1 << UhtSingletonType.Statics,
+	}
+
 	/// <summary>
 	/// Interface to provide singleton names during code generation. 
 	/// </summary>
@@ -52,6 +80,34 @@ namespace EpicGames.UHT.Exporters.CodeGen
 		public string GetSingletonName(UhtObject? obj, UhtSingletonType type);
 	}
 
+	/// <summary>
+	/// Interface to link objects together in generated code by initialization expressions such as
+	/// address-of or initializing encoded pointers for runtime linking
+	/// </summary>
+	public interface IUhtObjectLinker : IUhtSingletonNameProvider
+	{
+		/// <summary>
+		/// Used by code generation to determine how a reference should be linked, i.e. by direct pointer or by encoded pointe requiring runtime fixup.
+		/// </summary>
+		/// <param name="obj"></param>
+		/// <returns>True if the object should be linked as a cross-module reference</returns>
+		public bool IsCrossModuleRef(UhtObject obj);
+
+		/// <summary>
+		/// 
+		/// </summary>
+		/// <param name="obj"></param>
+		/// <returns></returns>
+		public int GetModuleIndex(UhtModule obj);
+
+		/// <summary>
+		/// 
+		/// </summary>
+		/// <param name="obj"></param>
+		/// <returns></returns>
+		public int GetObjectIndex(UhtObject obj);
+	}
+
 	/// <summary>
 	/// Interface to provide global settings during code generation. 
 	/// </summary>
@@ -121,7 +177,8 @@ namespace EpicGames.UHT.Exporters.CodeGen
 			// e.g. Z_ConstInit_UClass_UObject
 			public string ConstInitSingletonName { get; set; }
 			// Full declaration to import variable containing const-initialized object
-			// e.g. COREUOBJECT_API extern UE_CONSTINIT_UOBJECT_DECL UClass Z_ConstInit_UClass_UObject
+			// e.g. extern UE_CONSTINIT_UOBJECT_DECL UClass Z_ConstInit_UClass_UObject
+			// Note lack of API macro 
 			public string ConstInitExternalDecl { get; set; }
 			// Name of statics struct in generated cpp containing construction parameters
 			// e.g. Z_Construct_UClass_UObject_Statics
@@ -502,7 +559,7 @@ namespace EpicGames.UHT.Exporters.CodeGen
 			objectInfo.UnregisteredExternalDecl = $"\t{api}U{obj.EngineLinkClassName}* {objectInfo.UnregisteredSingletonName}();\r\n";
 			objectInfo.RegisteredExternalDecl = $"\t{api}U{obj.EngineLinkClassName}* {objectInfo.RegisteredSingletonName}();\r\n";
 			objectInfo.ConstInitSingletonName = objectInfo.RegisteredSingletonName.Replace("Z_Construct_", "Z_ConstInit_", StringComparison.InvariantCultureIgnoreCase);
-			objectInfo.ConstInitExternalDecl = $"\t{api}extern UE_CONSTINIT_UOBJECT_DECL TNoDestroy<U{obj.EngineClassName}> {objectInfo.ConstInitSingletonName};\r\n";
+			objectInfo.ConstInitExternalDecl = $"\textern UE_CONSTINIT_UOBJECT_DECL TNoDestroy<U{obj.EngineClassName}> {objectInfo.ConstInitSingletonName};\r\n";
 
 			// Init the children
 			foreach (UhtType child in obj.Children)
@@ -579,4 +636,15 @@ namespace EpicGames.UHT.Exporters.CodeGen
 		}
 		#endregion
 	}
+
+	static class UhtCodeGeneratorExtensions
+	{
+		/// <summary>
+		/// Convenience function so allow use of ?. operator on objects to determine if they are cross-module references according to the given linker.
+		/// </summary>
+		/// <param name="obj"></param>
+		/// <param name="linker"></param>
+		/// <returns></returns>
+		public static bool IsCrossModuleRef(this UhtObject obj, IUhtObjectLinker linker) => linker.IsCrossModuleRef(obj);
+	}
 }
diff --git a/Engine/Source/Programs/Shared/EpicGames.UHT/Exporters/CodeGen/UhtHeaderCodeGeneratorCppFile.cs b/Engine/Source/Programs/Shared/EpicGames.UHT/Exporters/CodeGen/UhtHeaderCodeGeneratorCppFile.cs
--- a/Engine/Source/Programs/Shared/EpicGames.UHT/Exporters/CodeGen/UhtHeaderCodeGeneratorCppFile.cs
+++ b/Engine/Source/Programs/Shared/EpicGames.UHT/Exporters/CodeGen/UhtHeaderCodeGeneratorCppFile.cs
@@ -2,7 +2,6 @@
 
 using System;
 using System.Collections.Generic;
-using System.Diagnostics.Eventing.Reader;
 using System.Linq;
 using System.Text;
 using EpicGames.Core;
@@ -11,6 +10,8 @@ using EpicGames.UHT.Utils;
 
 namespace EpicGames.UHT.Exporters.CodeGen
 {
+	using static UhtExpressionFactory;
+
 	/// <summary>
 	/// Collection of all registrations for a specific package
 	/// </summary>
@@ -196,11 +197,11 @@ namespace EpicGames.UHT.Exporters.CodeGen
 				{
 					if (Session.IsUsingCompiledInObjectFormat(UhtCompiledInObjectFormat.ConstInit))
 					{
-						builder.Append("static_assert(UE_WITH_CONSTINIT_UOBJECT, \"This generated code can only be compiled with UE_WITH_CONSTINIT_OBJECT\");");
+						builder.Append("static_assert(UE_WITH_CONSTINIT_UOBJECT, \"This generated code can only be compiled with UE_WITH_CONSTINIT_UOBJECT\");");
 					}
 					if (Session.IsUsingCompiledInObjectFormat(UhtCompiledInObjectFormat.Params))
 					{
-						builder.Append("static_assert(!UE_WITH_CONSTINIT_UOBJECT, \"This generated code can only be compiled with !UE_WITH_CONSTINIT_OBJECT\");");
+						builder.Append("static_assert(!UE_WITH_CONSTINIT_UOBJECT, \"This generated code can only be compiled with !UE_WITH_CONSTINIT_UOBJECT\");");
 					}
 				}
 
@@ -208,43 +209,92 @@ namespace EpicGames.UHT.Exporters.CodeGen
 				string cleanFileName = HeaderFile.FileNameWithoutExtension.Replace('.', '_');
 				builder.Append("void EmptyLinkFunctionForGeneratedCode").Append(cleanFileName).Append("() {}\r\n");
 
-				if (!HeaderFile.References.CrossModule.IsEmpty)
+				if (!HeaderFile.References.Declaration.IsEmpty)
+				{
+					if (Session.IsUsingCompiledInObjectFormat(UhtCompiledInObjectFormat.ConstInit))
+					{
+						using UhtCodeBlockComment blockComment = new(builder, "Forward Declarations");
+						using UhtConditionalMacroBlock block = new(builder, "UE_WITH_CONSTINIT_UOBJECT", Session.IsUsingMultipleCompiledInObjectFormats);
+
+						foreach (UhtPackage package in HeaderFile.Module.Packages.OrderBy(x => x.SourceName))
+						{
+							builder.Append(GetExternalDecl(package, UhtSingletonType.ConstInit).AsSpan().TrimStart());
+						}
+						IEnumerable<UhtObject> decls = HeaderFile.References.Declaration.SortReferences(UhtSingletonTypeFlag.Registered | UhtSingletonTypeFlag.Unregistered | UhtSingletonTypeFlag.ConstInit);
+						foreach (UhtObject obj in decls)
+						{
+							builder.Append(GetExternalDecl(obj, UhtSingletonType.ConstInit).AsSpan().TrimStart());
+						}
+					}
+				}
+
+				if (!HeaderFile.References.ExternalObjects.IsEmpty)
 				{
 					using UhtCodeBlockComment blockComment = new(builder, "Cross Module References");
 					if (Session.IsUsingCompiledInObjectFormat(UhtCompiledInObjectFormat.ConstInit))
 					{
 						using UhtConditionalMacroBlock block = new(builder, "UE_WITH_CONSTINIT_UOBJECT", Session.IsUsingMultipleCompiledInObjectFormats);
-						// Cross-references based on constinit variables
-						IReadOnlyList<string> sorted = HeaderFile.References.CrossModule.GetSortedReferences(
-							(objectIndex, type) => type switch
-							{
-								UhtSingletonType.Statics => null,
-								_ => GetExternalDecl(objectIndex, UhtSingletonType.ConstInit)
-							});
-						foreach (string crossReference in sorted)
+						builder.AppendLine("// Cross-module objects can only be linked directly in monolithic builds");
+						builder.AppendLine("#if IS_MONOLITHIC");
+						UhtSingletonTypeFlag filter = UhtSingletonTypeFlag.Registered | UhtSingletonTypeFlag.Unregistered | UhtSingletonTypeFlag.ConstInit;
+						foreach (UhtObject obj in HeaderFile.References.ExternalObjects.SortReferences(filter))
 						{
-							builder.Append(crossReference.AsSpan().TrimStart());
+							builder.Append(GetExternalDecl(obj, UhtSingletonType.ConstInit).AsSpan().TrimStart());
 						}
+						builder.AppendLine("#endif // IS_MONOLITHIC");
 					}
+
 					if (Session.IsUsingCompiledInObjectFormat(UhtCompiledInObjectFormat.Params))
 					{
 						using UhtConditionalMacroBlock block = new(builder, "!UE_WITH_CONSTINIT_UOBJECT", Session.IsUsingMultipleCompiledInObjectFormats);
-						// Cross-references based on factory functions
-						IReadOnlyList<string> sorted = HeaderFile.References.CrossModule.GetSortedReferences(
-							(objectIndex, type) => type switch
-							{
-								UhtSingletonType.ConstInit or UhtSingletonType.Statics => null,
-								_ => GetExternalDecl(objectIndex, type)
-							});
-						foreach (string crossReference in sorted)
+
+						foreach (UhtObject obj in HeaderFile.References.ExternalObjects.SortReferences(UhtSingletonTypeFlag.Registered))
+						{
+							builder.Append(GetExternalDecl(obj, UhtSingletonType.Registered).AsSpan().TrimStart());
+						}
+
+						foreach (UhtObject obj in HeaderFile.References.ExternalObjects.SortReferences(UhtSingletonTypeFlag.Unregistered))
+						{
+							builder.Append(GetExternalDecl(obj, UhtSingletonType.Unregistered).AsSpan().TrimStart());
+						}
+					}
+				}
+
+				if (!HeaderFile.References.InternalObjects.IsEmpty)
+				{
+					using UhtCodeBlockComment blockComment = new(builder, "Same Module References");
+					if (Session.IsUsingCompiledInObjectFormat(UhtCompiledInObjectFormat.ConstInit))
+					{
+						using UhtConditionalMacroBlock block = new(builder, "UE_WITH_CONSTINIT_UOBJECT", Session.IsUsingMultipleCompiledInObjectFormats);
+						UhtSingletonTypeFlag filter = UhtSingletonTypeFlag.Registered | UhtSingletonTypeFlag.Unregistered | UhtSingletonTypeFlag.ConstInit;
+						foreach (UhtObject obj in HeaderFile.References.InternalObjects.SortReferences(filter))
+						{
+							builder.Append(GetExternalDecl(obj, UhtSingletonType.ConstInit).AsSpan().TrimStart());
+						}
+					}
+					if (Session.IsUsingCompiledInObjectFormat(UhtCompiledInObjectFormat.Params))
+					{
+						using UhtConditionalMacroBlock block = new(builder, "!UE_WITH_CONSTINIT_UOBJECT", Session.IsUsingMultipleCompiledInObjectFormats);
+
+						foreach (UhtObject obj in HeaderFile.References.InternalObjects.SortReferences(UhtSingletonTypeFlag.Registered))
+						{
+							builder.Append(GetExternalDecl(obj, UhtSingletonType.Registered).AsSpan().TrimStart());
+						}
+
+						foreach (UhtObject obj in HeaderFile.References.InternalObjects.SortReferences(UhtSingletonTypeFlag.Unregistered))
 						{
-							builder.Append(crossReference.AsSpan().TrimStart());
+							builder.Append(GetExternalDecl(obj, UhtSingletonType.Unregistered).AsSpan().TrimStart());
 						}
 					}
 				}
 
 				int generatedBodyStart = builder.Length;
 
+				builder.AppendLinkingMacros(Session, HeaderFile.References.Modules);
+				builder.AppendLine($$"""
+					#define UHT_STRUCT_BASE(INIT) UE::CodeGen::ConstInit::TCompiledInObjectPtr<const FStructBaseChain>(UE::Private::AsStructBaseChain(INIT))
+					""");
+
 				// We need to emit the internal UClass needed for VMODULEs
 				foreach (UhtField field in HeaderFile.References.ExportTypes)
 				{
@@ -336,27 +386,11 @@ namespace EpicGames.UHT.Exporters.CodeGen
 					if (Session.IsUsingCompiledInObjectFormat(UhtCompiledInObjectFormat.ConstInit))
 					{
 						using UhtConditionalMacroBlock block = new(builder, "UE_WITH_CONSTINIT_UOBJECT", Session.IsUsingMultipleCompiledInObjectFormats);
-						// Constinit versions 
-						// Array of pointers to constinit objects
-						builder.AppendInstances(registrations.Enumerations, UhtDefineScopeNames.Standard,
-							builder => builder.Append("\tstatic inline constinit UEnum* EnumInfo[] = {\r\n"),
-							(builder, enumObj) => builder.Append("\t\t").AppendConstInitSingletonRef(this, enumObj).Append(",\r\n"),
-							builder => builder.Append("\t};\r\n"));
-
-						builder.AppendInstances(registrations.ScriptStructs, UhtDefineScopeNames.Standard,
-							builder => builder.Append("\tstatic inline constinit UScriptStruct* ScriptStructInfo[] = {\r\n"),
-							(builder, scriptStruct) => builder.Append("\t\t").AppendConstInitSingletonRef(this, scriptStruct).Append(",\r\n"),
-							builder => builder.Append("\t};\r\n"));
-
-						builder.AppendInstances(registrations.Classes, UhtDefineScopeNames.Standard,
-							builder => builder.Append("\tstatic inline constinit UClass* ClassInfo[] = {\r\n"),
-							(builder, classObj) => builder.Append("\t\t").AppendConstInitSingletonRef(this, classObj).Append(",\r\n"),
-							builder => builder.Append("\t};\r\n"));
 
 						// Registration for RigVM functions
 						builder.AppendInstances(registrations.RigVMStructs, UhtDefineScopeNames.Standard,
 							builder => builder.Append("\tstatic inline constinit FRigVMCompiledInStruct RigVMStructs[] = {\r\n"),
-							(builder, scriptStruct) => builder.Append("\t\t{ .Struct = ").AppendConstInitSingletonRef(this, scriptStruct).Append(", ")
+							(builder, scriptStruct) => builder.Append($"\t\t{{ .Struct = {ConstInitSingletonRef(this, scriptStruct)}, ")
 								 .Append(".Functions = MakeArrayView(").AppendSingletonName(this, scriptStruct, UhtSingletonType.Statics).Append("::RigVMFunctions), ")
 								.Append(" },\r\n"),
 							builder => builder.Append("\t};\r\n"));
@@ -416,11 +450,6 @@ namespace EpicGames.UHT.Exporters.CodeGen
 					if (Session.IsUsingCompiledInObjectFormat(UhtCompiledInObjectFormat.ConstInit))
 					{ 
 						using UhtConditionalMacroBlock block = new(builder, "UE_WITH_CONSTINIT_UOBJECT", Session.IsUsingMultipleCompiledInObjectFormats);
-						builder.Append($"static FRegisterCompiledInObjects {name}_{combinedHash}{{\r\n")
-							.AppendArrayView(registrations.Classes, UhtDefineScopeNames.Standard, staticsName, "ClassInfo", 1, ",\r\n")
-							.AppendArrayView(registrations.ScriptStructs, UhtDefineScopeNames.Standard, staticsName, "ScriptStructInfo", 1, ",\r\n")
-							.AppendArrayView(registrations.Enumerations, UhtDefineScopeNames.Standard, staticsName, "EnumInfo", 1, ",\r\n")
-							.Append("};\r\n");
 						if (!registrations.RigVMStructs.IsEmpty)
 						{
 							builder.Append($"static FRegisterRigVMStructs RigVM_{name}{{\r\n")
@@ -444,6 +473,11 @@ namespace EpicGames.UHT.Exporters.CodeGen
 					}
 				}
 
+				builder.AppendLinkingMacrosUndef(Session, HeaderFile.References.Modules);
+				builder.AppendLine($$"""
+					#undef UHT_STRUCT_BASE
+					""");
+
 				if (Session.IncludeDebugOutput)
 				{
 					builder.Append("#if 0\r\n");
@@ -780,7 +814,7 @@ namespace EpicGames.UHT.Exporters.CodeGen
 			int hashCodeBlockStart = builder.Length;
 
 			// Collect the properties
-			PropertyMemberContextImpl context = new(CodeGenerator, scriptStruct, staticsName);
+			PropertyMemberContextImpl context = new(this, scriptStruct, staticsName);
 			UhtUsedDefineScopes<UhtProperty> properties = new(scriptStruct.Properties);
 
 			// Declare the statics structure
@@ -808,7 +842,7 @@ namespace EpicGames.UHT.Exporters.CodeGen
 				if (Session.IsUsingCompiledInObjectFormat(UhtCompiledInObjectFormat.ConstInit))
 				{ 
 					using UhtConditionalMacroBlock block = new(builder, "UE_WITH_CONSTINIT_UOBJECT", Session.IsUsingMultipleCompiledInObjectFormats);
-					builder.Append("\tstatic UE_CONSTINIT_UOBJECT_DECL inline const FStructBaseChain* StructBases[] = {\r\n");
+					builder.Append("\tstatic UE_CONSTINIT_UOBJECT_DECL inline UE::CodeGen::ConstInit::TCompiledInObjectPtr<const FStructBaseChain> StructBases[] = {\r\n");
 					if (scriptStruct.SuperScriptStruct is not null) 
 					{
 						List<UhtScriptStruct> supers = new();
@@ -818,10 +852,10 @@ namespace EpicGames.UHT.Exporters.CodeGen
 						}
 						foreach (UhtScriptStruct super in supers.AsEnumerable().Reverse())
 						{
-							builder.Append("\t\tUE::Private::AsStructBaseChain(").AppendConstInitSingletonRef(this, super).Append("),\r\n");
+							builder.AppendLine($"\t\tUHT_STRUCT_BASE({SingletonRef(this, super, UhtSingletonType.ConstInit)}), ");
 						}
 					}
-					builder.Append("\t\tUE::Private::AsStructBaseChain(").AppendConstInitSingletonRef(this, scriptStruct).Append("),\r\n");
+					builder.AppendLine($"\t\tUHT_STRUCT_BASE({SingletonRef(this, scriptStruct, UhtSingletonType.ConstInit)}), ");
 					builder.Append("\t};\r\n");
 				}
 
@@ -1000,7 +1034,7 @@ namespace EpicGames.UHT.Exporters.CodeGen
 					builder.AppendConstInitObjectParams(scriptStruct, this, 1);
 					builder.Append("\tUE::CodeGen::ConstInit::FUFieldParams{},\r\n");
 					builder.Append("\tUE::CodeGen::ConstInit::FStructParams{\r\n")
-						.Append($"\t\t.Super = ").AppendConstInitSingletonRef(this, scriptStruct.SuperScriptStruct).Append(",\r\n")
+						.AppendLine($"\t\t.Super = {ConstInitSingletonRef(this, scriptStruct.SuperScriptStruct)}, ")
 						.Append($"\t\t.BaseChain = MakeArrayView({staticsName}::StructBases),\r\n")
 						.Append($"\t\t.ChildProperties =  {staticsName}::GetFirstProperty(),\r\n")
 						.Append($"\t\t.PropertiesSize = {staticsName}::GetStructSize(),\r\n")
@@ -1254,7 +1288,7 @@ namespace EpicGames.UHT.Exporters.CodeGen
 			int hashCodeBlockStart = builder.Length;
 
 			// Collect the properties
-			PropertyMemberContextImpl context = new(CodeGenerator, function, "", eventParameters, staticsName);
+			PropertyMemberContextImpl context = new(this, function, "", eventParameters, staticsName);
 			UhtUsedDefineScopes<UhtProperty> properties = new(function.Properties);
 
 			// Statics declaration
@@ -1329,7 +1363,7 @@ namespace EpicGames.UHT.Exporters.CodeGen
 						.AppendScopeLink(function.NextFunction, UhtDefineScopeNames.WithEditor, context, 2, ".Next = ", ",\r\n")
 						.Append("\t},\r\n");
 					builder.Append("\tUE::CodeGen::ConstInit::FStructParams{\r\n")
-						.Append($"\t\t.Super = ").AppendConstInitSingletonRef(this, function.SuperFunction).Append(",\r\n")
+						.AppendLine($"\t\t.Super = {ConstInitSingletonRef(this, function.SuperFunction)}, ")
 						.Append($"\t\t.ChildProperties = {staticsName}::GetFirstProperty(),\r\n")
 						.Append(sizeOfStatic is null ? null : $"\t\t.PropertiesSize = {sizeOfStatic},\r\n")
 						.Append(alignOfStatic is null ? null : $"\t\t.MinAlignment = {alignOfStatic},\r\n")
@@ -1345,7 +1379,8 @@ namespace EpicGames.UHT.Exporters.CodeGen
 							(EFunctionFlags.Native, EFunctionFlags.None, _)
 								=> $"\t\t.NativeFunction = &{outerClass!.Namespace.FullSourceName}{outerClass!.NativeFunctionCallName}::exec{function.SourceName},\r\n",
 							// Non-native function 
-							(EFunctionFlags.None, _, _) => $"\t\t.NativeFunction = &UObject::ProcessInternal,\r\n",
+							// Can't use cross-dll assignment in modular builds here 
+							(EFunctionFlags.None, _, _) => ".NativeFunction = UE_IF(IS_MONOLITHIC, &UObject::ProcessInternal, nullptr),\r\n",
 							// No thunk for whatever reason - e.g. net service function
 							_ => $"",
 						})
@@ -2544,7 +2579,7 @@ namespace EpicGames.UHT.Exporters.CodeGen
 			string registrationName = $"Z_Registration_Info_UClass_{classObj.SourceName}";
 			string[]? sparseDataTypes = classObj.MetaData.GetStringArray(UhtNames.SparseClassDataTypes);
 
-			PropertyMemberContextImpl context = new(CodeGenerator, classObj, staticsName);
+			PropertyMemberContextImpl context = new(this, classObj, staticsName);
 
 			bool hasInterfaces = classObj.Bases.Any(x => x is UhtClass baseClass && baseClass.ClassFlags.HasAnyFlags(EClassFlags.Interface)) || classObj.FlattenedVerseInterfaces.Count > 0;
 
@@ -2734,7 +2769,7 @@ namespace EpicGames.UHT.Exporters.CodeGen
 					if (Session.IsUsingCompiledInObjectFormat(UhtCompiledInObjectFormat.ConstInit))
 					{
 						using UhtConditionalMacroBlock block = new(builder, "UE_WITH_CONSTINIT_UOBJECT", Session.IsUsingMultipleCompiledInObjectFormats);
-						builder.Append("\tstatic const UE::CodeGen::ConstInit::FClassImplementedInterface Interfaces[];\r\n");
+						builder.Append("\tstatic UE::CodeGen::ConstInit::FClassImplementedInterface Interfaces[];\r\n");
 					}
 					if (Session.IsUsingCompiledInObjectFormat(UhtCompiledInObjectFormat.Params))
 					{
@@ -2750,7 +2785,7 @@ namespace EpicGames.UHT.Exporters.CodeGen
 				if (Session.IsUsingCompiledInObjectFormat(UhtCompiledInObjectFormat.ConstInit))
 				{
 					using UhtConditionalMacroBlock block = new(builder, "UE_WITH_CONSTINIT_UOBJECT", Session.IsUsingMultipleCompiledInObjectFormats);
-					builder.Append("\tstatic UE_CONSTINIT_UOBJECT_DECL inline const FStructBaseChain* StructBases[] = {\r\n");
+					builder.Append("\tstatic UE_CONSTINIT_UOBJECT_DECL inline UE::CodeGen::ConstInit::TCompiledInObjectPtr<const FStructBaseChain> StructBases[] = {\r\n");
 					if (classObj.SuperClass is not null)
 					{
 						List<UhtClass> supers = new();
@@ -2760,10 +2795,10 @@ namespace EpicGames.UHT.Exporters.CodeGen
 						}
 						foreach (UhtClass super in supers.AsEnumerable().Reverse())
 						{
-							builder.Append("\t\tUE::Private::AsStructBaseChain(").AppendConstInitSingletonRef(this, super).Append("),\r\n");
+							builder.AppendLine($"\t\tUHT_STRUCT_BASE({ConstInitSingletonRef(this, super)}), ");
 						}
 					}
-					builder.Append("\t\tUE::Private::AsStructBaseChain(").AppendConstInitSingletonRef(this, classObj).Append("),\r\n");
+					builder.AppendLine($"\t\tUHT_STRUCT_BASE({ConstInitSingletonRef(this, classObj)}),");
 					builder.Append("\t};\r\n");
 				}
 
@@ -2800,15 +2835,14 @@ namespace EpicGames.UHT.Exporters.CodeGen
 					if (Session.IsUsingCompiledInObjectFormat(UhtCompiledInObjectFormat.ConstInit))
 					{
 						using UhtConditionalMacroBlock block = new(builder, "UE_WITH_CONSTINIT_UOBJECT", Session.IsUsingMultipleCompiledInObjectFormats);
-						builder.Append($"const UE::CodeGen::ConstInit::FClassImplementedInterface {staticsName}::Interfaces[] = {{\r\n");
+						builder.Append($"UE::CodeGen::ConstInit::FClassImplementedInterface {staticsName}::Interfaces[] = {{\r\n");
 						Action<UhtClass> appendConstInitInterfaceParam = (interfaceObj) =>
 						{
 							if (interfaceObj.ClassFlags.HasAnyFlags(EClassFlags.Interface))
 							{
 								string isVerseDirectInterface = interfaceObj.IsVerseField && classObj.VerseInterfaces.Contains(interfaceObj) ? "true" : "false";
 								builder.Append("\t{ ")
-									.Append(".Class = ")
-									.AppendSingletonRef(context, interfaceObj.AlternateObject, UhtSingletonType.ConstInit).Append(", ")
+									.Append($".Class = {ConstInitSingletonRef(this, interfaceObj.AlternateObject)}, ")
 									.Append(interfaceObj.IsVerseField
 										? ".PointerOffset= 0, .bImplementedByK2 = true, "
 										: $".PointerOffset = (int32)VTABLE_OFFSET({classObj.FullyQualifiedSourceName}, {interfaceObj.FullyQualifiedSourceName}), .bImplementedByK2 = false, ")
@@ -2864,7 +2898,7 @@ namespace EpicGames.UHT.Exporters.CodeGen
 						.AppendConstInitObjectParams(classObj, this, 1, "|RF_Standalone")
 						.Append("\tUE::CodeGen::ConstInit::FUFieldParams{},\r\n")
 						.Append("\tUE::CodeGen::ConstInit::FStructParams{\r\n")
-						.Append("\t\t.Super = ").AppendConstInitSingletonRef(this, classObj.Super).Append(",\r\n")
+						.AppendLine($"\t\t.Super = {ConstInitSingletonRef(this, classObj.Super)}, ")
 						.Append($"\t\t.FirstChild = {staticsName}::GetFirstFunction(),\r\n")
 						.Append($"\t\t.BaseChain = MakeArrayView({staticsName}::StructBases),\r\n")
 						.Append($"\t\t.ChildProperties = {staticsName}::GetFirstProperty(),\r\n")
@@ -2876,7 +2910,7 @@ namespace EpicGames.UHT.Exporters.CodeGen
 						.Append($"\t\t.ConfigName = {classObj.Namespace.FullSourceName}{classObj.SourceName}::StaticConfigName(),\r\n")
 						.Append($"\t\t.ClassConstructor = InternalConstructor<{classObj.Namespace.FullSourceName}{classObj.SourceName}>,\r\n")
 						.Append($"\t\t.ClassVTableHelperCtorCaller = InternalVTableHelperCtorCaller<{classObj.Namespace.FullSourceName}{classObj.SourceName}>,\r\n")
-						.Append($"\t\t.CppClassStaticFunctions = FUObjectCppClassStaticFunctions(({classObj.Namespace.FullSourceName}{classObj.SourceName}*)nullptr),\r\n")
+						.Append($"\t\t.CppClassStaticFunctions = FUObjectCppClassStaticFunctions::ForClass<{classObj.Namespace.FullSourceName}{classObj.SourceName}>(),\r\n")
 						.Append($"\t\t.ClassFlags = (EClassFlags)0x{(uint)classFlags:X8}u,\r\n")
 						.Append($"\t\t.ClassCastFlags = {classObj.Namespace.FullSourceName}{classObj.SourceName}::StaticAllClassCastFlags(),\r\n") // Propagate cast flags at compile time
 						.Append($"\t\t.CppClassTypeInfo = &{staticsName}::StaticCppClassTypeInfo,\r\n");
@@ -2884,17 +2918,17 @@ namespace EpicGames.UHT.Exporters.CodeGen
 					{
 						builder.Append("\t\t.NativeFunctions = ").AppendArrayView(nativeFunctions, UhtDefineScopeNames.WithEditor, staticsName, "Funcs", 2, ",\r\n");
 					}
-					if (classObj.ClassWithin is not null && classObj.ClassWithin != Session.UObject)
+					if (classObj.ClassWithin is not null && classObj != Session.UObject)
 					{
-						builder.Append($"\t\t.ClassWithin = &{GetSingletonName(classObj.ClassWithin, UhtSingletonType.ConstInit)},\r\n");
+						builder.AppendLine($"\t\t.ClassWithin = {ConstInitSingletonRef(this, classObj.ClassWithin)}, ");
 					}
 					if (hasInterfaces)
 					{
-						builder.Append($"\t\t.Interfaces = MakeConstArrayView({staticsName}::Interfaces),\r\n");
+						builder.Append($"\t\t.Interfaces = MakeArrayView({staticsName}::Interfaces),\r\n");
 					}
 					if (classObj.SparseClassDataStruct is not null)
 					{
-						builder.Append($"\t\t.SparseClassDataStruct = &{GetSingletonName(classObj.SparseClassDataStruct, UhtSingletonType.ConstInit)}\r\n");
+						builder.Append($"\t\t.SparseClassDataStruct = {ConstInitSingletonRef(this, classObj.SparseClassDataStruct)}\r\n");
 					}
 					builder.Append("\t},\r\n");
 					if (classObj.IsVerseField)
@@ -3254,13 +3288,13 @@ namespace EpicGames.UHT.Exporters.CodeGen
 
 		private class PropertyMemberContextImpl : IUhtPropertyMemberContext
 		{
-			private readonly UhtCodeGenerator _codeGenerator;
+			private readonly UhtHeaderCodeGenerator _codeGenerator;
 			private readonly UhtStruct _outerStruct;
 			private readonly string _outerStructNamespaceName;
 			private readonly string _outerStructSourceName;
 			private readonly string _staticsName;
 
-			public PropertyMemberContextImpl(UhtCodeGenerator codeGenerator, UhtStruct outerStruct, string outerStructNamespaceName, string outerStructSourceName, string staticsName)
+			public PropertyMemberContextImpl(UhtHeaderCodeGenerator codeGenerator, UhtStruct outerStruct, string outerStructNamespaceName, string outerStructSourceName, string staticsName)
 			{
 				_codeGenerator = codeGenerator;
 				_outerStruct = outerStruct;
@@ -3269,7 +3303,7 @@ namespace EpicGames.UHT.Exporters.CodeGen
 				_outerStructSourceName = outerStructSourceName;
 			}
 
-			public PropertyMemberContextImpl(UhtCodeGenerator codeGenerator, UhtStruct outerStruct, string staticsName)
+			public PropertyMemberContextImpl(UhtHeaderCodeGenerator codeGenerator, UhtStruct outerStruct, string staticsName)
 			{
 				_codeGenerator = codeGenerator;
 				_outerStruct = outerStruct;
@@ -3290,6 +3324,12 @@ namespace EpicGames.UHT.Exporters.CodeGen
 				return _codeGenerator.GetSingletonName(obj, type);
 			}
 
+			public bool IsCrossModuleRef(UhtObject obj) => _codeGenerator.IsCrossModuleRef(obj);
+
+			public int GetModuleIndex(UhtModule obj) => _codeGenerator.GetModuleIndex(obj);
+
+			public int GetObjectIndex(UhtObject obj) => _codeGenerator.GetObjectIndex(obj);
+
 			public uint GetTypeHash(UhtObject obj)
 			{
 				return _codeGenerator.ObjectInfos[obj.ObjectTypeIndex].Hash;
@@ -3312,22 +3352,30 @@ namespace EpicGames.UHT.Exporters.CodeGen
 		/// </summary>
 		/// <param name="builder"></param>
 		/// <param name="uhtObject"></param>
-		/// <param name="names"></param>
+		/// <param name="linker"></param>
 		/// <param name="tabs"></param>
 		/// <param name="additionalObjectFlags"></param>
 		/// <returns></returns>
-		public static StringBuilder AppendConstInitObjectParams(this StringBuilder builder, UhtObject uhtObject, IUhtSingletonNameProvider names, int tabs, string? additionalObjectFlags = null)
+		public static StringBuilder AppendConstInitObjectParams(this StringBuilder builder, UhtObject uhtObject, IUhtObjectLinker linker, int tabs, string? additionalObjectFlags = null)
 		{
 			const string ObjectFlags = "RF_Public|RF_Transient|RF_MarkAsNative";
+			if (uhtObject == uhtObject.Session.UClass)
+			{
+				builder.AppendTabs(tabs).AppendLine("UE::CodeGen::ConstInit::EClass::Class, ");
+			}
+			else if (uhtObject == uhtObject.Session.UObject)
+			{
+				builder.AppendTabs(tabs).AppendLine("UE::CodeGen::ConstInit::EObject::Object, ");
+			}
 			builder.AppendTabs(tabs).Append("UE::CodeGen::ConstInit::FObjectParams{\r\n");
 			builder.AppendTabs(tabs + 1).Append($".Flags = {ObjectFlags}{additionalObjectFlags},\r\n");
 			// Can't take a pointer to Class to initialize itself because of TNoDestroy, null class signals to take the address of Class in the constructor 
 			if (uhtObject != uhtObject.EngineClass)
 			{
-				builder.AppendTabs(tabs + 1).Append($".Class =").AppendConstInitSingletonRef(names, uhtObject.EngineClass).Append(",\r\n");
+				builder.AppendLine($"{Tabs(tabs+1)}.Class = {SingletonRef(linker, uhtObject.EngineClass, UhtSingletonType.ConstInit)}, ");
 			}
 			builder.AppendTabs(tabs + 1).Append($".NameUTF8 = UTF8TEXT(").AppendUTF8LiteralString(uhtObject.EngineName).Append("),\r\n");
-			builder.AppendTabs(tabs + 1).Append($".Outer = ").AppendConstInitSingletonRef(names, uhtObject.Outer as UhtObject).Append(",\r\n");
+			builder.AppendTabs(tabs + 1).AppendLine($".Outer = {ConstInitSingletonRef(linker, uhtObject.Outer as UhtObject)},");
 			builder.AppendTabs(tabs).Append("},\r\n");
 			return builder;
 		}
diff --git a/Engine/Source/Programs/Shared/EpicGames.UHT/Exporters/CodeGen/UhtMacroBlockEmitter.cs b/Engine/Source/Programs/Shared/EpicGames.UHT/Exporters/CodeGen/UhtMacroBlockEmitter.cs
--- a/Engine/Source/Programs/Shared/EpicGames.UHT/Exporters/CodeGen/UhtMacroBlockEmitter.cs
+++ b/Engine/Source/Programs/Shared/EpicGames.UHT/Exporters/CodeGen/UhtMacroBlockEmitter.cs
@@ -3,7 +3,6 @@
 using System;
 using System.Text;
 using EpicGames.UHT.Types;
-using EpicGames.UHT.Utils;
 
 namespace EpicGames.UHT.Exporters.CodeGen
 {
@@ -12,6 +11,7 @@ namespace EpicGames.UHT.Exporters.CodeGen
 		private readonly StringBuilder _builder;
 		private readonly UhtDefineScopeNames _defineScopeNames;
 		private UhtDefineScope _current = UhtDefineScope.None;
+		private bool _open = false;
 
 		public UhtMacroBlockEmitter(StringBuilder builder, UhtDefineScopeNames defineScopeNames, UhtDefineScope initialState)
 		{
@@ -20,6 +20,10 @@ namespace EpicGames.UHT.Exporters.CodeGen
 			Set(initialState);
 		}
 
+		/// <summary>
+		/// Change the define scope to the new provided one using #endif and #if 
+		/// </summary>
+		/// <param name="defineScope"></param>
 		public void Set(UhtDefineScope defineScope)
 		{
 			if (defineScope == UhtDefineScope.Invalid)
@@ -31,13 +35,59 @@ namespace EpicGames.UHT.Exporters.CodeGen
 				return;
 			}
 			_builder.AppendEndIfPreprocessor(_current, _defineScopeNames);
-			_builder.AppendIfPreprocessor(defineScope, _defineScopeNames);
+			if (defineScope != UhtDefineScope.None)
+			{
+				_builder.AppendIfPreprocessor(defineScope, _defineScopeNames);
+				_open = true;
+			}
+			else
+			{
+				_open = false;
+			}
 			_current = defineScope;
 		}
 
+		/// <summary>
+		/// Change the define scope to the new provided.
+		/// None -> Scope X uses #if
+		/// Scope X -> None uses #else
+		/// Scope X -> Scope Y uses #elif
+		/// </summary>
+		/// <param name="defineScope"></param>
+		public void SetElse(UhtDefineScope defineScope)
+		{
+			if (defineScope == UhtDefineScope.Invalid)
+			{
+				defineScope = UhtDefineScope.None;
+			}
+			if (_current == defineScope)
+			{
+				return;
+			}
+			if (_current == UhtDefineScope.None)
+			{
+				_builder.AppendIfPreprocessor(defineScope, _defineScopeNames);
+			}
+			else if (defineScope == UhtDefineScope.None)
+			{
+				_builder.AppendElsePreprocessor(_current, _defineScopeNames);
+			}
+			else
+			{
+				_builder.AppendElseIfPreprocessor(defineScope, _defineScopeNames);
+			}
+
+			_current = defineScope;
+			_open = true;
+		}
+
 		public void Dispose()
 		{
-			Set(UhtDefineScope.None);
+			if (_open)
+			{
+				_builder.AppendLine("#endif");
+				_open = false;
+			}
 		}
 	}
 
diff --git a/Engine/Source/Programs/Shared/EpicGames.UHT/Exporters/CodeGen/UhtPackageCodeGenerator.cs b/Engine/Source/Programs/Shared/EpicGames.UHT/Exporters/CodeGen/UhtPackageCodeGenerator.cs
--- a/Engine/Source/Programs/Shared/EpicGames.UHT/Exporters/CodeGen/UhtPackageCodeGenerator.cs
+++ b/Engine/Source/Programs/Shared/EpicGames.UHT/Exporters/CodeGen/UhtPackageCodeGenerator.cs
@@ -1,14 +1,17 @@
 // Copyright Epic Games, Inc. All Rights Reserved.
 
 using System;
+using System.Collections;
+using System.Collections.Generic;
 using System.Linq;
 using System.Text;
+using EpicGames.Core;
 using EpicGames.UHT.Types;
 using EpicGames.UHT.Utils;
 
 namespace EpicGames.UHT.Exporters.CodeGen
 {
-	class UhtPackageCodeGenerator : IUhtSingletonNameProvider
+	class UhtPackageCodeGenerator : IUhtObjectLinker
 	{
 		public static string HeaderCopyright =
 			"// Copyright Epic Games, Inc. All Rights Reserved.\r\n" +
@@ -51,6 +54,43 @@ namespace EpicGames.UHT.Exporters.CodeGen
 			return CodeGenerator.GetSingletonName(obj, type);
 		}
 
+		/// <inheritdoc/>
+		public bool IsCrossModuleRef(UhtObject obj) => obj.Module != Module;
+
+		/// <inheritdoc/>
+		public int GetModuleIndex(UhtModule otherModule)
+		{
+			if (Module.Linking == null)
+			{
+				throw new UhtIceException("Cannot link modules before link indices have been assigned by UhtSession");
+			}
+			if (Module.Linking.ReferencedModules.TryGetValue(otherModule, out int value))
+			{
+				return value;
+			}
+			else
+			{
+				throw new UhtIceException($"Module {otherModule.ShortName} was not discovered by reference collection from module {Module.ShortName}, so objects in it cannot be linked");
+			}
+		} 
+
+		/// <inheritdoc/>
+		public int GetObjectIndex(UhtObject obj) 
+		{
+			if (obj.Module.Linking == null)
+			{
+				throw new UhtIceException("Cannot link modules before link indices have been assigned by UhtSession");
+			}
+			if (obj.Module.Linking.LinkableObjects.TryGetValue(obj, out int value))
+			{
+				return value;
+			}
+			else
+			{
+				throw new UhtIceException($"Object {obj.FullName} is not linkable from module {obj.Module.ShortName}");
+			}
+		} 
+
 		/// <summary>
 		/// Return the external declaration for an object
 		/// </summary>
@@ -217,29 +257,68 @@ namespace EpicGames.UHT.Exporters.CodeGen
 		{
 			switch (type)
 			{
-			case UhtEngineType.Package:
-				return builder.AppendSingletonRef(provider, session.UPackage, UhtSingletonType.ConstInit);
-			case UhtEngineType.Class:
-				return builder.AppendSingletonRef(provider, session.UClass, UhtSingletonType.ConstInit);
-			case UhtEngineType.Delegate:
-				return builder.AppendSingletonRef(provider, session.UDelegateFunction, UhtSingletonType.ConstInit); 
-			case UhtEngineType.SparseDelegate:
-				return builder.AppendSingletonRef(provider, session.USparseDelegateFunction, UhtSingletonType.ConstInit); 
-			case UhtEngineType.Enum:
-				return builder.AppendSingletonRef(provider, session.UEnum, UhtSingletonType.ConstInit);
-			case UhtEngineType.Function:
-				return builder.AppendSingletonRef(provider, session.UFunction, UhtSingletonType.ConstInit);
-			case UhtEngineType.Interface:
-				return builder.AppendSingletonRef(provider, session.UInterface, UhtSingletonType.ConstInit);
-			case UhtEngineType.ScriptStruct:
-				return builder.AppendSingletonRef(provider, session.UScriptStruct, UhtSingletonType.ConstInit);
-			case UhtEngineType.NativeInterface:
-				throw new NotImplementedException("ConstInit singleton for NativeInterface not implemented");
-			case UhtEngineType.Property:
-				throw new ArgumentException("Cannot insert singleton ref for property type because properties are not UObjects");
-			default:
-				throw new ArgumentException($"Invalid value ${type} for UhtEngineType");
+				case UhtEngineType.Package:
+					return builder.AppendSingletonRef(provider, session.UPackage, UhtSingletonType.ConstInit);
+				case UhtEngineType.Class:
+					return builder.AppendSingletonRef(provider, session.UClass, UhtSingletonType.ConstInit);
+				case UhtEngineType.Delegate:
+					return builder.AppendSingletonRef(provider, session.UDelegateFunction, UhtSingletonType.ConstInit);
+				case UhtEngineType.SparseDelegate:
+					return builder.AppendSingletonRef(provider, session.USparseDelegateFunction, UhtSingletonType.ConstInit);
+				case UhtEngineType.Enum:
+					return builder.AppendSingletonRef(provider, session.UEnum, UhtSingletonType.ConstInit);
+				case UhtEngineType.Function:
+					return builder.AppendSingletonRef(provider, session.UFunction, UhtSingletonType.ConstInit);
+				case UhtEngineType.Interface:
+					return builder.AppendSingletonRef(provider, session.UInterface, UhtSingletonType.ConstInit);
+				case UhtEngineType.ScriptStruct:
+					return builder.AppendSingletonRef(provider, session.UScriptStruct, UhtSingletonType.ConstInit);
+				case UhtEngineType.NativeInterface:
+					throw new NotImplementedException("ConstInit singleton for NativeInterface not implemented");
+				case UhtEngineType.Property:
+					throw new ArgumentException("Cannot insert singleton ref for property type because properties are not UObjects");
+				default:
+					throw new ArgumentException($"Invalid value ${type} for UhtEngineType");
 			}
 		}
+
+		internal static StringBuilder AppendLinkingMacros(this StringBuilder builder, UhtSession session, IEnumerable<UhtModule> modules)
+		{
+			if (modules.Any() && session.IsUsingCompiledInObjectFormat(UhtCompiledInObjectFormat.ConstInit))
+			{
+				using UhtConditionalMacroBlock block = new(builder, "UE_WITH_CONSTINIT_UOBJECT", session.IsUsingMultipleCompiledInObjectFormats);
+				UhtModule[] modulesSorted = [.. modules.OrderBy(x => x.ShortName)];
+
+				builder.AppendLine("// Monolithic linking macros");
+				builder.AppendLine("#if IS_MONOLITHIC");
+				foreach (UhtModule otherModule in modulesSorted)
+				{
+					builder.AppendLine($"#define UHT_LINK_{otherModule.ShortName}(MODULE_INDEX, OBJECT_INDEX, POINTER) POINTER");
+				}
+
+				builder.AppendLine("// Modular linking macros");
+				builder.AppendLine("#else // IS_MONOLITHIC");
+				foreach (UhtModule otherModule in modulesSorted)
+				{
+					builder.AppendLine($"#define UHT_LINK_{otherModule.ShortName}(MODULE_INDEX, OBJECT_INDEX, POINTER) UE::CodeGen::ConstInit::FCompiledInObjectReference{{ MODULE_INDEX, OBJECT_INDEX }}");
+				}
+				builder.AppendLine("#endif // IS_MONOLITHIC");
+			}
+			return builder;
+		}
+
+		internal static StringBuilder AppendLinkingMacrosUndef(this StringBuilder builder, UhtSession session, IEnumerable<UhtModule> modules)
+		{
+			if (modules.Any() && session.IsUsingCompiledInObjectFormat(UhtCompiledInObjectFormat.ConstInit))
+			{
+				using UhtConditionalMacroBlock block = new(builder, "UE_WITH_CONSTINIT_UOBJECT", session.IsUsingMultipleCompiledInObjectFormats);
+				UhtModule[] modulesSorted = [.. modules.OrderBy(x => x.ShortName)];
+				foreach (UhtModule otherModule in modulesSorted)
+				{
+					builder.AppendLine($"#undef UHT_LINK_{otherModule.ShortName}");
+				}
+			}
+			return builder;
+		}
 	}
 }
diff --git a/Engine/Source/Programs/Shared/EpicGames.UHT/Exporters/CodeGen/UhtPackageCodeGeneratorCppFile.cs b/Engine/Source/Programs/Shared/EpicGames.UHT/Exporters/CodeGen/UhtPackageCodeGeneratorCppFile.cs
--- a/Engine/Source/Programs/Shared/EpicGames.UHT/Exporters/CodeGen/UhtPackageCodeGeneratorCppFile.cs
+++ b/Engine/Source/Programs/Shared/EpicGames.UHT/Exporters/CodeGen/UhtPackageCodeGeneratorCppFile.cs
@@ -10,6 +10,8 @@ using EpicGames.UHT.Utils;
 
 namespace EpicGames.UHT.Exporters.CodeGen
 {
+	using static UhtExpressionFactory;
+
 	internal class UhtPackageCodeGeneratorCppFile : UhtPackageCodeGenerator
 	{
 		/// <summary>
@@ -104,11 +106,11 @@ namespace EpicGames.UHT.Exporters.CodeGen
 				{
 					if (Session.IsUsingCompiledInObjectFormat(UhtCompiledInObjectFormat.ConstInit))
 					{
-						builder.Append("static_assert(UE_WITH_CONSTINIT_UOBJECT, \"This generated code can only be compiled with UE_WITH_CONSTINIT_OBJECT\");");
+						builder.AppendLine("static_assert(UE_WITH_CONSTINIT_UOBJECT, \"This generated code can only be compiled with UE_WITH_CONSTINIT_OBJECT\");");
 					}
 					if (Session.IsUsingCompiledInObjectFormat(UhtCompiledInObjectFormat.Params))
 					{
-						builder.Append("static_assert(!UE_WITH_CONSTINIT_UOBJECT, \"This generated code can only be compiled with !UE_WITH_CONSTINIT_OBJECT\");");
+						builder.AppendLine("static_assert(!UE_WITH_CONSTINIT_UOBJECT, \"This generated code can only be compiled with !UE_WITH_CONSTINIT_OBJECT\");");
 					}
 				}
 
@@ -129,23 +131,82 @@ namespace EpicGames.UHT.Exporters.CodeGen
 					builder.Append("#endif\r\n");
 				}
 
-				if (Session.IsUsingCompiledInObjectFormat(UhtCompiledInObjectFormat.ConstInit))
+				if (Session.IsUsingCompiledInObjectFormat(UhtCompiledInObjectFormat.Params))
 				{
-					using UhtConditionalMacroBlock block = new(builder, "UE_WITH_CONSTINIT_UOBJECT", Session.IsUsingMultipleCompiledInObjectFormats);
+					using UhtConditionalMacroBlock block = new(builder, "!UE_WITH_CONSTINIT_UOBJECT", Session.IsUsingMultipleCompiledInObjectFormats);
 					foreach (UhtObject obj in singletons)
 					{
-						builder.Append(GetExternalDecl(obj, UhtSingletonType.ConstInit));
+						builder.Append(GetExternalDecl(obj, UhtSingletonType.Registered));
 					}
 				}
-				if (Session.IsUsingCompiledInObjectFormat(UhtCompiledInObjectFormat.Params))
+
+				List<UhtObject> linkableObjectsOrdered = [.. Module.Linking!.LinkableObjects
+					.Select((int order, UhtObject obj) (p) => (p.Value, p.Key))
+					.OrderBy(p => p.order)
+					.Select(p => p.obj)];
+				if (Session.IsUsingCompiledInObjectFormat(UhtCompiledInObjectFormat.ConstInit))
 				{
-					using UhtConditionalMacroBlock block = new(builder, "!UE_WITH_CONSTINIT_UOBJECT", Session.IsUsingMultipleCompiledInObjectFormats);
+					using UhtConditionalMacroBlock block = new(builder, "UE_WITH_CONSTINIT_UOBJECT", Session.IsUsingMultipleCompiledInObjectFormats);
 					foreach (UhtObject obj in singletons)
 					{
-						builder.Append(GetExternalDecl(obj, UhtSingletonType.Registered));
+						builder.Append(GetExternalDecl(obj, UhtSingletonType.ConstInit));
+					}
+					foreach ((UhtObject obj, int index) in Module.Linking!.LinkableObjects)
+					{
+						builder.Append(GetExternalDecl(obj, UhtSingletonType.ConstInit));
+					}
+
+					builder.AppendLine($$"""
+						extern {{Module.Api}}UObject* const {{ModuleLinkableObjects(Module)}}[{{linkableObjectsOrdered.Count}}] = 
+						{
+						""");
+					foreach (UhtObject obj in linkableObjectsOrdered)
+					{
+						using UhtMacroBlockEmitter macroBlockEmitter = new(builder, UhtDefineScopeNames.Standard, obj.DefineScope);
+						builder.AppendLine($"\t{SingletonRef(this, obj, UhtSingletonType.ConstInit)}, ");
+						if (obj.DefineScope != UhtDefineScope.None)
+						{
+							macroBlockEmitter.SetElse(UhtDefineScope.None);
+							builder.AppendLine("\tnullptr, ");
+						}
+					}
+					builder.AppendLine($$"""
+						};	
+
+						""");
+
+					List<UhtModule> orderedLinkedModules = [.. Module.Linking!.ReferencedModules
+						.Select((int Index, UhtModule M) (p) => (p.Value, p.Key))
+						.OrderBy(p => p.Index)
+						.Select(p => p.M)
+						];
+					if (orderedLinkedModules.Count > 0)
+					{
+						builder.AppendLine("// Linkable object arrays from other modules");
+						builder.AppendLine("#if !IS_MONOLITHIC");
+						foreach (UhtModule otherModule in orderedLinkedModules)
+						{
+							builder.AppendLine($"extern {otherModule.Api}UObject* const {ModuleLinkableObjects(otherModule)}[{otherModule.Linking!.LinkableObjects.Count}];");
+						}
+
+						builder.AppendLine($$"""
+							const TConstArrayView<UObject* const> {{ModuleLinkedModules(Module)}}[{{orderedLinkedModules.Count}}] = 
+							{
+							""");
+						foreach (UhtModule module in orderedLinkedModules)
+						{
+							builder.AppendLine($"\tMakeArrayView({ModuleLinkableObjects(module)}, {module.Linking!.LinkableObjects.Count}), ");
+						}
+						builder.AppendLine($$"""
+							};	
+
+							""");
+						builder.AppendLine("#endif // !IS_MONOLITHIC");
 					}
 				}
 
+				builder.AppendLinkingMacros(Session, Module.Linking!.ReferencedModules.Keys);
+
 				foreach (UhtPackage package in Module.Packages)
 				{
 					string strippedName = PackageInfos[package.PackageTypeIndex].StrippedName;
@@ -158,35 +219,20 @@ namespace EpicGames.UHT.Exporters.CodeGen
 					if (Session.IsUsingCompiledInObjectFormat(UhtCompiledInObjectFormat.ConstInit))
 					{
 						using UhtConditionalMacroBlock block = new(builder, "UE_WITH_CONSTINIT_UOBJECT", Session.IsUsingMultipleCompiledInObjectFormats);
-						string? singletonsArray = null;
-						if (packageSingletons.Length != 0)
-						{
-							// List of objects of mixed types to register with the package
-							singletonsArray = $"Z_Singletons_UPackage_{strippedName}";
-							builder.AppendEach(packageSingletons,
-								(builder) => builder.Append($"\tconstinit UObject* {singletonsArray}[] = {{\r\n"),
-								(builder, field) => builder.Append($"\t\t&{GetSingletonName(field, UhtSingletonType.ConstInit)}, \r\n"),
-								(builder) => builder.Append("\t};\r\n")
-							);
-						}
-						builder.Append($"\tconstinit UPackage {GetSingletonName(package, UhtSingletonType.ConstInit)}( \r\n")
-							.Append("\t\tUE::CodeGen::ConstInit::FPackageParams{\r\n")
-							.Append("\t\t\t.Object = UE::CodeGen::ConstInit::FObjectParams{\r\n")
-							.Append("\t\t\t\t.Flags = RF_Public,\r\n")
-							.Append($"\t\t\t\t.Class = &{GetSingletonName(Session.UPackage, UhtSingletonType.ConstInit)},\r\n")
-							.Append("\t\t\t\t.NameUTF8 = UTF8TEXT(").AppendUTF8LiteralString(package.SourceName).Append("),\r\n")
-							.Append("\t\t\t\t.Outer = nullptr,\r\n")
-							.Append("\t\t\t},\r\n")
-							.Append($"\t\t\t.Flags = EPackageFlags(PKG_CompiledIn | 0x{(uint)flags:X8}),\r\n")
-							.Append($"\t\t\tIF_WITH_METADATA(.MetaData = ").AppendMetaDataArrayView(package.MetaData, null, MetaDataParamsName, 2, ",)\r\n")
-							.Append($"\t\t\t.BodiesHash = 0x{bodiesHash:X8},\r\n")
-							.Append($"\t\t\t.DeclarationsHash = 0x{declarationsHash:X8},\r\n")
-							.Append("\t});\r\n")
-							.Append("\r\n")
-							.Append($"\tstatic FRegisterCompiledInObjects Z_CompiledInDeferPackage_UPackage_{strippedName}{{\r\n")
-							.Append($"\t\t&{GetSingletonName(package, UhtSingletonType.ConstInit)}, \r\n")
-							.AppendArrayView(null, singletonsArray, 2, "\r\n")
-							.Append("};\r\n");
+						builder.Append($"constinit UPackage {GetSingletonName(package, UhtSingletonType.ConstInit)}( \r\n")
+							.Append("\tUE::CodeGen::ConstInit::FPackageParams{\r\n")
+							.Append("\t\t.Object = UE::CodeGen::ConstInit::FObjectParams{\r\n")
+							.Append("\t\t\t.Flags = RF_Public,\r\n")
+							.Append($"\t\t\t.Class = {ConstInitSingletonRef(this, Session.UPackage)},\r\n")
+							.Append("\t\t\t.NameUTF8 = UTF8TEXT(").AppendUTF8LiteralString(package.SourceName).Append("),\r\n")
+							.Append("\t\t\t.Outer = nullptr,\r\n")
+							.Append("\t\t},\r\n")
+							.Append($"\t\t.Flags = EPackageFlags(PKG_CompiledIn | 0x{(uint)flags:X8}),\r\n")
+							.Append($"\t\tIF_WITH_METADATA(.MetaData = ").AppendMetaDataArrayView(package.MetaData, null, MetaDataParamsName, 2, ",)\r\n")
+							.Append($"\t\t.BodiesHash = 0x{bodiesHash:X8},\r\n")
+							.Append($"\t\t.DeclarationsHash = 0x{declarationsHash:X8},\r\n")
+							.Append("});\r\n")
+							.Append("\r\n");
 					}
 
 					if (Session.IsUsingCompiledInObjectFormat(UhtCompiledInObjectFormat.Params))
@@ -229,6 +275,38 @@ namespace EpicGames.UHT.Exporters.CodeGen
 					}
 				}
 
+				// In constinit format, emit one registration struct per module, not per-package
+				if (Session.IsUsingCompiledInObjectFormat(UhtCompiledInObjectFormat.ConstInit))
+				{
+					using UhtConditionalMacroBlock block = new(builder, "UE_WITH_CONSTINIT_UOBJECT", Session.IsUsingMultipleCompiledInObjectFormats);
+
+					builder.AppendLine($"static UPackage* const Z_Packages_{Module.ShortName}[] = {{");
+					foreach (UhtPackage package in Module.Packages)
+					{
+						builder.AppendLine($"\t{SingletonRef(this, package, UhtSingletonType.ConstInit)}, ");
+					}
+					builder.AppendLine($"}};");
+
+					// UhtSession orders types according to UhtObject.LinkOrder to guarantee this ordering
+					(int classesStart, int classesEnd) = GetTypeRange<UhtClass>(linkableObjectsOrdered, 0);
+					(int structsStart, int structsEnd) = GetTypeRange<UhtScriptStruct>(linkableObjectsOrdered, classesEnd);
+					(int enumsStart, int enumsEnd) = GetTypeRange<UhtEnum>(linkableObjectsOrdered, structsEnd);
+					int othersStart = enumsEnd;
+					builder.AppendLine($$"""
+						static FRegisterCompiledInObjects Z_Register_Module{{Module.ShortName}}{ {
+						#if !IS_MONOLITHIC
+							.LinkedModules = {{ModuleLinkedModules(Module)}},
+							.LinkableObjects = MakeConstArrayView({{ModuleLinkableObjects(Module)}}),
+						#endif // !IS_MONOLITHIC
+							.Packages = MakeConstArrayView(Z_Packages_{{Module.ShortName}}),
+							.Classes = MakeConstArrayView({{ModuleLinkableObjects(Module)}}).Mid({{classesStart}}, {{classesEnd - classesStart}}),
+							.Structs = MakeConstArrayView({{ModuleLinkableObjects(Module)}}).Mid({{structsStart}}, {{structsEnd - structsStart}}),
+							.Enums = MakeConstArrayView({{ModuleLinkableObjects(Module)}}).Mid({{enumsStart}}, {{enumsEnd - enumsStart}}),
+							.Others = MakeConstArrayView({{ModuleLinkableObjects(Module)}}).RightChop({{othersStart}}),
+						} };
+						""");
+				}
+
 				// Verse registration
 				if (Module.Module.HasVerse)
 				{
@@ -258,6 +336,8 @@ namespace EpicGames.UHT.Exporters.CodeGen
 					builder.Append("}\r\n");
 				}
 
+				builder.AppendLinkingMacrosUndef(Session, Module.Linking!.ReferencedModules.Keys);
+
 				builder.Append(EnableDeprecationWarnings).Append("\r\n");
 
 				if (SaveExportedHeaders)
@@ -279,5 +359,28 @@ namespace EpicGames.UHT.Exporters.CodeGen
 			ReadOnlySpan<char> cppModuleName = packageName[(slashIndex + 1)..];
 			return $"{{ TEXT(\"{mountPointName}\"), TEXT(\"{cppModuleName}\") }}";
 		}
+
+		private static (int Start, int End) GetTypeRange<T>(IReadOnlyList<UhtObject> objects, int Start) 
+			where T: UhtObject
+		{
+			if (objects.Count > Start && objects[Start] is T)
+			{
+				for (int i = Start + 1; i < objects.Count; ++i)
+				{
+					if (objects[i] is not T)
+					{
+						return (Start, i);
+					}
+				}
+				return (Start, objects.Count);
+			}
+			return (Start, Start);
+		}
+
+		private static SimpleExpression ModuleLinkableObjects(UhtModule module) => new SimpleExpression("Z_LinkableObjects_", module.ShortName);
+		private static SimpleExpression ModuleLinkedModules(UhtModule module) => module.Linking!.ReferencedModules.Count == 0 
+			? new SimpleExpression("", "{}")
+			: new SimpleExpression("Z_LinkedModules_", module.ShortName);
 	}
+
 }
diff --git a/Engine/Source/Programs/Shared/EpicGames.UHT/Exporters/CodeGen/UhtUsedDefineScopes.cs b/Engine/Source/Programs/Shared/EpicGames.UHT/Exporters/CodeGen/UhtUsedDefineScopes.cs
--- a/Engine/Source/Programs/Shared/EpicGames.UHT/Exporters/CodeGen/UhtUsedDefineScopes.cs
+++ b/Engine/Source/Programs/Shared/EpicGames.UHT/Exporters/CodeGen/UhtUsedDefineScopes.cs
@@ -974,7 +974,7 @@ namespace EpicGames.UHT.Exporters.CodeGen
 			{
 				foreach (UhtFunction? next in link.Next)
 				{
-					blockEmitter.Set(next?.DefineScope ?? UhtDefineScope.None);
+					blockEmitter.SetElse(next?.DefineScope ?? UhtDefineScope.None);
 					builder.AppendTabs(tabs).Append(prefix);
 					builder.AppendConstInitSingletonRef(context, next).Append(endl);
 				}
diff --git a/Engine/Source/Programs/Shared/EpicGames.UHT/Types/Properties/UhtByteProperty.cs b/Engine/Source/Programs/Shared/EpicGames.UHT/Types/Properties/UhtByteProperty.cs
--- a/Engine/Source/Programs/Shared/EpicGames.UHT/Types/Properties/UhtByteProperty.cs
+++ b/Engine/Source/Programs/Shared/EpicGames.UHT/Types/Properties/UhtByteProperty.cs
@@ -13,6 +13,7 @@ using EpicGames.UHT.Utils;
 
 namespace EpicGames.UHT.Types
 {
+	using static UhtExpressionFactory;
 
 	/// <summary>
 	/// The form that the byte was defined
@@ -136,7 +137,7 @@ namespace EpicGames.UHT.Types
 		public override void CollectReferencesInternal(IUhtReferenceCollector collector, bool addForwardDeclarations, bool isTemplateProperty)
 		{
 			base.CollectReferencesInternal(collector, addForwardDeclarations, isTemplateProperty);
-			collector.AddCrossModuleReference(Enum, UhtSingletonType.Registered);
+			collector.AddObjectReference(Enum, UhtSingletonType.Registered);
 		}
 
 		/// <inheritdoc/>
@@ -185,7 +186,7 @@ namespace EpicGames.UHT.Types
 		public override StringBuilder AppendConstInitMemberDef(StringBuilder builder, IUhtPropertyMemberContext context, string name, string nameSuffix, Action<StringBuilder>? outerFunc, string? offset, int tabs)
 		{
 			AppendConstInitMemberDefStart(builder, context, name, nameSuffix, outerFunc, offset, tabs);
-			AppendMemberDefRef(builder, context, Enum, UhtSingletonType.ConstInit, true);
+			builder.Append($"{ConstInitSingletonRef(context, Enum, true)}, ");
 			AppendConstInitMemberDefEnd(builder, context);
 			return builder;
 		}
diff --git a/Engine/Source/Programs/Shared/EpicGames.UHT/Types/Properties/UhtClassProperty.cs b/Engine/Source/Programs/Shared/EpicGames.UHT/Types/Properties/UhtClassProperty.cs
--- a/Engine/Source/Programs/Shared/EpicGames.UHT/Types/Properties/UhtClassProperty.cs
+++ b/Engine/Source/Programs/Shared/EpicGames.UHT/Types/Properties/UhtClassProperty.cs
@@ -9,6 +9,8 @@ using EpicGames.UHT.Utils;
 
 namespace EpicGames.UHT.Types
 {
+	using static EpicGames.UHT.Exporters.CodeGen.UhtExpressionFactory;
+
 	/// <summary>
 	/// FClassProperty
 	/// </summary>
@@ -190,8 +192,7 @@ namespace EpicGames.UHT.Types
 		public override StringBuilder AppendConstInitMemberDef(StringBuilder builder, IUhtPropertyMemberContext context, string name, string nameSuffix, Action<StringBuilder>? outerFunc, string? offset, int tabs)
 		{
 			AppendConstInitMemberDefStart(builder, context, name, nameSuffix, outerFunc, offset, tabs);
-			AppendMemberDefRef(builder, context, Class, Exporters.CodeGen.UhtSingletonType.ConstInit);
-			AppendMemberDefRef(builder, context, MetaClass, Exporters.CodeGen.UhtSingletonType.ConstInit);
+			builder.Append($"{ConstInitSingletonRef(context, Class, true)}, {ConstInitSingletonRef(context, MetaClass, true)}, ");
 			AppendConstInitMemberDefEnd(builder, context);
 			return builder;
 		}
diff --git a/Engine/Source/Programs/Shared/EpicGames.UHT/Types/Properties/UhtDelegateProperty.cs b/Engine/Source/Programs/Shared/EpicGames.UHT/Types/Properties/UhtDelegateProperty.cs
--- a/Engine/Source/Programs/Shared/EpicGames.UHT/Types/Properties/UhtDelegateProperty.cs
+++ b/Engine/Source/Programs/Shared/EpicGames.UHT/Types/Properties/UhtDelegateProperty.cs
@@ -11,6 +11,8 @@ using EpicGames.UHT.Utils;
 
 namespace EpicGames.UHT.Types
 {
+	using static EpicGames.UHT.Exporters.CodeGen.UhtExpressionFactory;
+
 	/// <summary>
 	/// FDelegatePropertyDelegateProperty
 	/// </summary>
@@ -91,7 +93,7 @@ namespace EpicGames.UHT.Types
 		public override void CollectReferencesInternal(IUhtReferenceCollector collector, bool addForwardDeclarations, bool isTemplateProperty)
 		{
 			base.CollectReferencesInternal(collector, addForwardDeclarations, isTemplateProperty);
-			collector.AddCrossModuleReference(Function, UhtSingletonType.Registered);
+			collector.AddObjectReference(Function, UhtSingletonType.Registered);
 		}
 
 		/// <inheritdoc/>
@@ -131,7 +133,7 @@ namespace EpicGames.UHT.Types
 		public override StringBuilder AppendConstInitMemberDef(StringBuilder builder, IUhtPropertyMemberContext context, string name, string nameSuffix, Action<StringBuilder>? outerFunc, string? offset, int tabs)
 		{
 			AppendConstInitMemberDefStart(builder, context, name, nameSuffix, outerFunc, offset, tabs);
-			AppendMemberDefRef(builder, context, Function, Exporters.CodeGen.UhtSingletonType.ConstInit);
+			builder.Append($"{ConstInitSingletonRef(context, Function, true)}, ");
 			AppendConstInitMemberDefEnd(builder, context);
 			return builder;
 		}
diff --git a/Engine/Source/Programs/Shared/EpicGames.UHT/Types/Properties/UhtEnumProperty.cs b/Engine/Source/Programs/Shared/EpicGames.UHT/Types/Properties/UhtEnumProperty.cs
--- a/Engine/Source/Programs/Shared/EpicGames.UHT/Types/Properties/UhtEnumProperty.cs
+++ b/Engine/Source/Programs/Shared/EpicGames.UHT/Types/Properties/UhtEnumProperty.cs
@@ -13,6 +13,7 @@ using EpicGames.UHT.Utils;
 
 namespace EpicGames.UHT.Types
 {
+	using static EpicGames.UHT.Exporters.CodeGen.UhtExpressionFactory;
 
 	/// <summary>
 	/// FEnumProperty
@@ -98,7 +99,7 @@ namespace EpicGames.UHT.Types
 		public override void CollectReferencesInternal(IUhtReferenceCollector collector, bool addForwardDeclarations, bool isTemplateProperty)
 		{
 			base.CollectReferencesInternal(collector, addForwardDeclarations, isTemplateProperty);
-			collector.AddCrossModuleReference(Enum, UhtSingletonType.Registered);
+			collector.AddObjectReference(Enum, UhtSingletonType.Registered);
 			if (addForwardDeclarations && Enum.CppForm == UhtEnumCppForm.EnumClass)
 			{
 				collector.AddForwardDeclaration(Enum);
@@ -307,13 +308,13 @@ namespace EpicGames.UHT.Types
 					builder.Append("(int32)sizeof(");
 					UnderlyingProperty.AppendText(builder, UhtPropertyTextType.Generic).Append("), ");
 				}
-				AppendMemberDefRef(builder, context, Enum, Exporters.CodeGen.UhtSingletonType.ConstInit);
+				builder.Append($"{ConstInitSingletonRef(context, Enum, true)}, ");
 				AppendConstInitMemberDefEnd(builder, context);
 			}
 			else
 			{
 				AppendConstInitMemberDefStart(builder, context, name, nameSuffix, outerFunc, offset, tabs);
-				AppendMemberDefRef(builder, context, Enum, Exporters.CodeGen.UhtSingletonType.ConstInit);
+				builder.Append($"{ConstInitSingletonRef(context, Enum, true)}, ");
 				AppendConstInitMemberDefEnd(builder, context);
 			}
 			return builder;
diff --git a/Engine/Source/Programs/Shared/EpicGames.UHT/Types/Properties/UhtFieldPathProperty.cs b/Engine/Source/Programs/Shared/EpicGames.UHT/Types/Properties/UhtFieldPathProperty.cs
--- a/Engine/Source/Programs/Shared/EpicGames.UHT/Types/Properties/UhtFieldPathProperty.cs
+++ b/Engine/Source/Programs/Shared/EpicGames.UHT/Types/Properties/UhtFieldPathProperty.cs
@@ -110,7 +110,7 @@ namespace EpicGames.UHT.Types
 		public override StringBuilder AppendConstInitMemberDef(StringBuilder builder, IUhtPropertyMemberContext context, string name, string nameSuffix, Action<StringBuilder>? outerFunc, string? offset, int tabs)
 		{
 			AppendConstInitMemberDefStart(builder, context, name, nameSuffix, outerFunc, offset, tabs);
-			builder.Append("&F").Append(FieldClassName).Append("::ConstInitClass, ");
+			builder.Append($"UE_IF(IS_MONOLITHIC, &F{FieldClassName}::ConstInitClass, UTF8TEXT(").AppendUTF8LiteralString(FieldClassName).Append(")), ");
 			AppendConstInitMemberDefEnd(builder, context);
 			return builder;
 		}
diff --git a/Engine/Source/Programs/Shared/EpicGames.UHT/Types/Properties/UhtInterfaceProperty.cs b/Engine/Source/Programs/Shared/EpicGames.UHT/Types/Properties/UhtInterfaceProperty.cs
--- a/Engine/Source/Programs/Shared/EpicGames.UHT/Types/Properties/UhtInterfaceProperty.cs
+++ b/Engine/Source/Programs/Shared/EpicGames.UHT/Types/Properties/UhtInterfaceProperty.cs
@@ -13,6 +13,8 @@ using EpicGames.UHT.Utils;
 
 namespace EpicGames.UHT.Types
 {
+	using static EpicGames.UHT.Exporters.CodeGen.UhtExpressionFactory;
+
 	/// <summary>
 	/// FInterfaceProperty
 	/// </summary>
@@ -101,7 +103,7 @@ namespace EpicGames.UHT.Types
 		public override void CollectReferencesInternal(IUhtReferenceCollector collector, bool addForwardDeclarations, bool isTemplateProperty)
 		{
 			base.CollectReferencesInternal(collector, addForwardDeclarations, isTemplateProperty);
-			collector.AddCrossModuleReference(InterfaceClass, UhtSingletonType.Unregistered);
+			collector.AddObjectReference(InterfaceClass.AlternateObject ?? InterfaceClass, UhtSingletonType.Unregistered);
 			if (addForwardDeclarations)
 			{
 				UhtClass? exportClass = InterfaceClass;
@@ -162,7 +164,7 @@ namespace EpicGames.UHT.Types
 		public override StringBuilder AppendConstInitMemberDef(StringBuilder builder, IUhtPropertyMemberContext context, string name, string nameSuffix, Action<StringBuilder>? outerFunc, string? offset, int tabs)
 		{
 			AppendConstInitMemberDefStart(builder, context, name, nameSuffix, outerFunc, offset, tabs);
-			AppendMemberDefRef(builder, context, InterfaceClass.AlternateObject ?? InterfaceClass, Exporters.CodeGen.UhtSingletonType.ConstInit);
+			builder.Append($"{ConstInitSingletonRef(context, InterfaceClass.AlternateObject ?? InterfaceClass, true)}, ");
 			AppendConstInitMemberDefEnd(builder, context);
 			return builder;
 		}
diff --git a/Engine/Source/Programs/Shared/EpicGames.UHT/Types/Properties/UhtLazyObjectPtrProperty.cs b/Engine/Source/Programs/Shared/EpicGames.UHT/Types/Properties/UhtLazyObjectPtrProperty.cs
--- a/Engine/Source/Programs/Shared/EpicGames.UHT/Types/Properties/UhtLazyObjectPtrProperty.cs
+++ b/Engine/Source/Programs/Shared/EpicGames.UHT/Types/Properties/UhtLazyObjectPtrProperty.cs
@@ -9,6 +9,7 @@ using EpicGames.UHT.Utils;
 
 namespace EpicGames.UHT.Types
 {
+	using static EpicGames.UHT.Exporters.CodeGen.UhtExpressionFactory;
 
 	/// <summary>
 	/// FLazyObjectProperty
@@ -93,7 +94,7 @@ namespace EpicGames.UHT.Types
 		public override StringBuilder AppendConstInitMemberDef(StringBuilder builder, IUhtPropertyMemberContext context, string name, string nameSuffix, Action<StringBuilder>? outerFunc, string? offset, int tabs)
 		{
 			AppendConstInitMemberDefStart(builder, context, name, nameSuffix, outerFunc, offset, tabs);
-			AppendMemberDefRef(builder, context, Class, Exporters.CodeGen.UhtSingletonType.ConstInit);
+			builder.Append($"{ConstInitSingletonRef(context, Class, true)}, ");
 			AppendConstInitMemberDefEnd(builder, context);
 			return builder;
 		}
diff --git a/Engine/Source/Programs/Shared/EpicGames.UHT/Types/Properties/UhtMulticastDelegateProperty.cs b/Engine/Source/Programs/Shared/EpicGames.UHT/Types/Properties/UhtMulticastDelegateProperty.cs
--- a/Engine/Source/Programs/Shared/EpicGames.UHT/Types/Properties/UhtMulticastDelegateProperty.cs
+++ b/Engine/Source/Programs/Shared/EpicGames.UHT/Types/Properties/UhtMulticastDelegateProperty.cs
@@ -91,7 +91,7 @@ namespace EpicGames.UHT.Types
 		public override void CollectReferencesInternal(IUhtReferenceCollector collector, bool addForwardDeclarations, bool isTemplateProperty)
 		{
 			base.CollectReferencesInternal(collector, addForwardDeclarations, isTemplateProperty);
-			collector.AddCrossModuleReference(Function, UhtSingletonType.Registered);
+			collector.AddObjectReference(Function, UhtSingletonType.Registered);
 		}
 
 		/// <inheritdoc/>
diff --git a/Engine/Source/Programs/Shared/EpicGames.UHT/Types/Properties/UhtMulticastInlineDelegateProperty.cs b/Engine/Source/Programs/Shared/EpicGames.UHT/Types/Properties/UhtMulticastInlineDelegateProperty.cs
--- a/Engine/Source/Programs/Shared/EpicGames.UHT/Types/Properties/UhtMulticastInlineDelegateProperty.cs
+++ b/Engine/Source/Programs/Shared/EpicGames.UHT/Types/Properties/UhtMulticastInlineDelegateProperty.cs
@@ -7,6 +7,8 @@ using EpicGames.UHT.Utils;
 
 namespace EpicGames.UHT.Types
 {
+	using static EpicGames.UHT.Exporters.CodeGen.UhtExpressionFactory;
+
 	/// <summary>
 	/// FMulticastInlineDelegateProperty
 	/// </summary>
@@ -78,7 +80,7 @@ namespace EpicGames.UHT.Types
 		public override StringBuilder AppendConstInitMemberDef(StringBuilder builder, IUhtPropertyMemberContext context, string name, string nameSuffix, Action<StringBuilder>? outerFunc, string? offset, int tabs)
 		{
 			AppendConstInitMemberDefStart(builder, context, name, nameSuffix, outerFunc, offset, tabs);
-			AppendMemberDefRef(builder, context, Function, Exporters.CodeGen.UhtSingletonType.ConstInit);
+			builder.Append($"{ConstInitSingletonRef(context, Function, true)}, ");
 			AppendConstInitMemberDefEnd(builder, context);
 			return builder;
 		}
diff --git a/Engine/Source/Programs/Shared/EpicGames.UHT/Types/Properties/UhtMulticastSparseDelegateProperty.cs b/Engine/Source/Programs/Shared/EpicGames.UHT/Types/Properties/UhtMulticastSparseDelegateProperty.cs
--- a/Engine/Source/Programs/Shared/EpicGames.UHT/Types/Properties/UhtMulticastSparseDelegateProperty.cs
+++ b/Engine/Source/Programs/Shared/EpicGames.UHT/Types/Properties/UhtMulticastSparseDelegateProperty.cs
@@ -7,6 +7,7 @@ using EpicGames.UHT.Utils;
 
 namespace EpicGames.UHT.Types
 {
+	using static EpicGames.UHT.Exporters.CodeGen.UhtExpressionFactory;
 
 	/// <summary>
 	/// FMulticastSparseDelegateProperty
@@ -79,7 +80,7 @@ namespace EpicGames.UHT.Types
 		public override StringBuilder AppendConstInitMemberDef(StringBuilder builder, IUhtPropertyMemberContext context, string name, string nameSuffix, Action<StringBuilder>? outerFunc, string? offset, int tabs)
 		{
 			AppendConstInitMemberDefStart(builder, context, name, nameSuffix, outerFunc, offset, tabs);
-			AppendMemberDefRef(builder, context, Function, Exporters.CodeGen.UhtSingletonType.ConstInit);
+			builder.Append($"{ConstInitSingletonRef(context, Function, true)}, ");
 			AppendConstInitMemberDefEnd(builder, context);
 			return builder;
 		}
diff --git a/Engine/Source/Programs/Shared/EpicGames.UHT/Types/Properties/UhtObjectProperty.cs b/Engine/Source/Programs/Shared/EpicGames.UHT/Types/Properties/UhtObjectProperty.cs
--- a/Engine/Source/Programs/Shared/EpicGames.UHT/Types/Properties/UhtObjectProperty.cs
+++ b/Engine/Source/Programs/Shared/EpicGames.UHT/Types/Properties/UhtObjectProperty.cs
@@ -8,6 +8,8 @@ using EpicGames.UHT.Utils;
 
 namespace EpicGames.UHT.Types
 {
+	using static EpicGames.UHT.Exporters.CodeGen.UhtExpressionFactory;
+
 	/// <summary>
 	/// FObjectProperty
 	/// </summary>
@@ -173,7 +175,7 @@ namespace EpicGames.UHT.Types
 		public override StringBuilder AppendConstInitMemberDef(StringBuilder builder, IUhtPropertyMemberContext context, string name, string nameSuffix, Action<StringBuilder>? outerFunc, string? offset, int tabs)
 		{
 			AppendConstInitMemberDefStart(builder, context, name, nameSuffix, outerFunc, offset, tabs);
-			AppendMemberDefRef(builder, context, Class, Exporters.CodeGen.UhtSingletonType.ConstInit);
+			builder.Append($"{ConstInitSingletonRef(context, Class, true)}, ");
 			AppendConstInitMemberDefEnd(builder, context);
 			return builder;
 		}
diff --git a/Engine/Source/Programs/Shared/EpicGames.UHT/Types/Properties/UhtObjectPropertyBase.cs b/Engine/Source/Programs/Shared/EpicGames.UHT/Types/Properties/UhtObjectPropertyBase.cs
--- a/Engine/Source/Programs/Shared/EpicGames.UHT/Types/Properties/UhtObjectPropertyBase.cs
+++ b/Engine/Source/Programs/Shared/EpicGames.UHT/Types/Properties/UhtObjectPropertyBase.cs
@@ -547,8 +547,8 @@ namespace EpicGames.UHT.Types
 		public override void CollectReferencesInternal(IUhtReferenceCollector collector, bool addForwardDeclarations, bool isTemplateProperty)
 		{
 			base.CollectReferencesInternal(collector, addForwardDeclarations, isTemplateProperty);
-			collector.AddCrossModuleReference(Class, UhtSingletonType.Unregistered);
-			collector.AddCrossModuleReference(MetaClass, UhtSingletonType.Unregistered);
+			collector.AddObjectReference(Class, UhtSingletonType.Unregistered);
+			collector.AddObjectReference(MetaClass, UhtSingletonType.Unregistered);
 			if (addForwardDeclarations)
 			{
 				collector.AddForwardDeclaration(Class);
diff --git a/Engine/Source/Programs/Shared/EpicGames.UHT/Types/Properties/UhtSoftClassProperty.cs b/Engine/Source/Programs/Shared/EpicGames.UHT/Types/Properties/UhtSoftClassProperty.cs
--- a/Engine/Source/Programs/Shared/EpicGames.UHT/Types/Properties/UhtSoftClassProperty.cs
+++ b/Engine/Source/Programs/Shared/EpicGames.UHT/Types/Properties/UhtSoftClassProperty.cs
@@ -10,6 +10,7 @@ using EpicGames.UHT.Utils;
 
 namespace EpicGames.UHT.Types
 {
+	using static EpicGames.UHT.Exporters.CodeGen.UhtExpressionFactory;
 
 	/// <summary>
 	/// FSoftClassProperty
@@ -56,7 +57,7 @@ namespace EpicGames.UHT.Types
 		public override void CollectReferencesInternal(IUhtReferenceCollector collector, bool addForwardDeclarations, bool isTemplateProperty)
 		{
 			base.CollectReferencesInternal(collector, addForwardDeclarations, isTemplateProperty);
-			collector.AddCrossModuleReference(MetaClass, UhtSingletonType.Unregistered);
+			collector.AddObjectReference(MetaClass, UhtSingletonType.Unregistered);
 			if (addForwardDeclarations && MetaClass != null)
 			{
 				collector.AddForwardDeclaration(MetaClass);
@@ -94,8 +95,7 @@ namespace EpicGames.UHT.Types
 		public override StringBuilder AppendConstInitMemberDef(StringBuilder builder, IUhtPropertyMemberContext context, string name, string nameSuffix, Action<StringBuilder>? outerFunc, string? offset, int tabs)
 		{
 			AppendConstInitMemberDefStart(builder, context, name, nameSuffix, outerFunc, offset, tabs);
-			AppendMemberDefRef(builder, context, Session.UClass, Exporters.CodeGen.UhtSingletonType.ConstInit);
-			AppendMemberDefRef(builder, context, MetaClass, Exporters.CodeGen.UhtSingletonType.ConstInit);
+			builder.Append($"{ConstInitSingletonRef(context, Session.UClass, true)}, {ConstInitSingletonRef(context, MetaClass, true)}, ");
 			AppendConstInitMemberDefEnd(builder, context);
 			return builder;
 		}
diff --git a/Engine/Source/Programs/Shared/EpicGames.UHT/Types/Properties/UhtSoftObjectProperty.cs b/Engine/Source/Programs/Shared/EpicGames.UHT/Types/Properties/UhtSoftObjectProperty.cs
--- a/Engine/Source/Programs/Shared/EpicGames.UHT/Types/Properties/UhtSoftObjectProperty.cs
+++ b/Engine/Source/Programs/Shared/EpicGames.UHT/Types/Properties/UhtSoftObjectProperty.cs
@@ -8,6 +8,8 @@ using EpicGames.UHT.Utils;
 
 namespace EpicGames.UHT.Types
 {
+	using static EpicGames.UHT.Exporters.CodeGen.UhtExpressionFactory;
+
 	/// <summary>
 	/// FSoftObjectProperty
 	/// </summary>
@@ -100,7 +102,7 @@ namespace EpicGames.UHT.Types
 		public override StringBuilder AppendConstInitMemberDef(StringBuilder builder, IUhtPropertyMemberContext context, string name, string nameSuffix, Action<StringBuilder>? outerFunc, string? offset, int tabs)
 		{
 			AppendConstInitMemberDefStart(builder, context, name, nameSuffix, outerFunc, offset, tabs);
-			AppendMemberDefRef(builder, context, Class, Exporters.CodeGen.UhtSingletonType.ConstInit);
+			builder.Append($"{ConstInitSingletonRef(context, Class, true)}, ");
 			AppendConstInitMemberDefEnd(builder, context);
 			return builder;
 		}
diff --git a/Engine/Source/Programs/Shared/EpicGames.UHT/Types/Properties/UhtStructProperty.cs b/Engine/Source/Programs/Shared/EpicGames.UHT/Types/Properties/UhtStructProperty.cs
--- a/Engine/Source/Programs/Shared/EpicGames.UHT/Types/Properties/UhtStructProperty.cs
+++ b/Engine/Source/Programs/Shared/EpicGames.UHT/Types/Properties/UhtStructProperty.cs
@@ -14,6 +14,8 @@ using EpicGames.UHT.Utils;
 
 namespace EpicGames.UHT.Types
 {
+	using static EpicGames.UHT.Exporters.CodeGen.UhtExpressionFactory;
+
 	/// <summary>
 	/// FStructProperty
 	/// </summary>
@@ -124,7 +126,7 @@ namespace EpicGames.UHT.Types
 		public override void CollectReferencesInternal(IUhtReferenceCollector collector, bool addForwardDeclarations, bool isTemplateProperty)
 		{
 			base.CollectReferencesInternal(collector, addForwardDeclarations, isTemplateProperty);
-			collector.AddCrossModuleReference(ScriptStruct, UhtSingletonType.Registered);
+			collector.AddObjectReference(ScriptStruct, UhtSingletonType.Registered);
 			if (addForwardDeclarations && !ScriptStruct.IsCoreType)
 			{
 				collector.AddForwardDeclaration(ScriptStruct);
@@ -173,7 +175,7 @@ namespace EpicGames.UHT.Types
 		{
 			AppendConstInitMemberDefStart(builder, context, name, nameSuffix, outerFunc, offset, tabs);
 			builder.Append($"(int32)Align(sizeof({ScriptStruct.Namespace.FullSourceName}{ScriptStruct.SourceName}), alignof({ScriptStruct.Namespace.FullSourceName}{ScriptStruct.SourceName})), ");
-			AppendMemberDefRef(builder, context, ScriptStruct, UhtSingletonType.ConstInit);
+			builder.Append($"{ConstInitSingletonRef(context, ScriptStruct, true)}, ");
 			AppendConstInitMemberDefEnd(builder, context);
 			return builder;
 		}
diff --git a/Engine/Source/Programs/Shared/EpicGames.UHT/Types/Properties/UhtWeakObjectPtrProperty.cs b/Engine/Source/Programs/Shared/EpicGames.UHT/Types/Properties/UhtWeakObjectPtrProperty.cs
--- a/Engine/Source/Programs/Shared/EpicGames.UHT/Types/Properties/UhtWeakObjectPtrProperty.cs
+++ b/Engine/Source/Programs/Shared/EpicGames.UHT/Types/Properties/UhtWeakObjectPtrProperty.cs
@@ -10,6 +10,7 @@ using EpicGames.UHT.Utils;
 
 namespace EpicGames.UHT.Types
 {
+	using static EpicGames.UHT.Exporters.CodeGen.UhtExpressionFactory;
 
 	/// <summary>
 	/// Represents a FWeakObjectProperty
@@ -100,7 +101,7 @@ namespace EpicGames.UHT.Types
 		public override StringBuilder AppendConstInitMemberDef(StringBuilder builder, IUhtPropertyMemberContext context, string name, string nameSuffix, Action<StringBuilder>? outerFunc, string? offset, int tabs)
 		{
 			AppendConstInitMemberDefStart(builder, context, name, nameSuffix, outerFunc, offset, tabs);
-			AppendMemberDefRef(builder, context, Class, Exporters.CodeGen.UhtSingletonType.ConstInit);
+			builder.Append($"{ConstInitSingletonRef(context, Class, true)}, ");
 			AppendConstInitMemberDefEnd(builder, context);
 			return builder;
 		}
diff --git a/Engine/Source/Programs/Shared/EpicGames.UHT/Types/UhtClass.cs b/Engine/Source/Programs/Shared/EpicGames.UHT/Types/UhtClass.cs
--- a/Engine/Source/Programs/Shared/EpicGames.UHT/Types/UhtClass.cs
+++ b/Engine/Source/Programs/Shared/EpicGames.UHT/Types/UhtClass.cs
@@ -2292,14 +2292,14 @@ namespace EpicGames.UHT.Types
 			// Add myself as declarations and cross module references
 			collector.AddDeclaration(this, UhtSingletonType.Unregistered);
 			collector.AddDeclaration(this, UhtSingletonType.Registered);
-			collector.AddCrossModuleReference(this, UhtSingletonType.Unregistered);
-			collector.AddCrossModuleReference(this, UhtSingletonType.Registered);
+			collector.AddObjectReference(this, UhtSingletonType.Unregistered);
+			collector.AddObjectReference(this, UhtSingletonType.Registered);
 
 			// Add the super class
 			if (SuperClass != null)
 			{
 				collector.AddDeclaration(SuperClass, UhtSingletonType.Registered);
-				collector.AddCrossModuleReference(SuperClass, UhtSingletonType.Registered);
+				collector.AddObjectReference(SuperClass, UhtSingletonType.Registered);
 			}
 
 			// Add any other base classes
@@ -2307,27 +2307,27 @@ namespace EpicGames.UHT.Types
 			{
 				if (baseObj is UhtClass baseClass)
 				{
-					collector.AddCrossModuleReference(baseClass, UhtSingletonType.Unregistered);
+					collector.AddObjectReference(baseClass, UhtSingletonType.Unregistered);
 				}
 			}
 
 			// Add any verse interfaces
 			foreach (UhtStruct interfaceObj in FlattenedVerseInterfaces)
 			{
-				collector.AddCrossModuleReference(interfaceObj, UhtSingletonType.Unregistered);
+				collector.AddObjectReference(interfaceObj, UhtSingletonType.Unregistered);
 			}
 
 			// Add the package
 			collector.AddDeclaration(Package, UhtSingletonType.Registered);
-			collector.AddCrossModuleReference(Package, UhtSingletonType.Registered);
+			collector.AddObjectReference(Package, UhtSingletonType.Registered);
 
 			// Pointers that are needed for constinit only
-			collector.AddCrossModuleReference(Session.UClass, UhtSingletonType.ConstInit);
-			collector.AddCrossModuleReference(ClassWithin, UhtSingletonType.ConstInit);
+			collector.AddObjectReference(Session.UClass, UhtSingletonType.ConstInit);
+			collector.AddObjectReference(ClassWithin, UhtSingletonType.ConstInit);
 			// Add all parents for constinit initializion of bases array 
 			for (UhtClass? super = SuperClass; super is not null; super = super.SuperClass)
 			{
-				collector.AddCrossModuleReference(super, UhtSingletonType.ConstInit);
+				collector.AddObjectReference(super, UhtSingletonType.ConstInit);
 			}
 
 			// Collect children references
diff --git a/Engine/Source/Programs/Shared/EpicGames.UHT/Types/UhtEnum.cs b/Engine/Source/Programs/Shared/EpicGames.UHT/Types/UhtEnum.cs
--- a/Engine/Source/Programs/Shared/EpicGames.UHT/Types/UhtEnum.cs
+++ b/Engine/Source/Programs/Shared/EpicGames.UHT/Types/UhtEnum.cs
@@ -597,10 +597,10 @@ namespace EpicGames.UHT.Types
 		{
 			collector.AddExportType(this);
 			collector.AddDeclaration(this, UhtSingletonType.Registered);
-			collector.AddCrossModuleReference(this, UhtSingletonType.Registered);
-			collector.AddCrossModuleReference(Package, UhtSingletonType.Registered);
+			collector.AddObjectReference(this, UhtSingletonType.Registered);
+			collector.AddObjectReference(Package, UhtSingletonType.Registered);
 
-			collector.AddCrossModuleReference(Session.UEnum, UhtSingletonType.ConstInit);
+			collector.AddObjectReference(Session.UEnum, UhtSingletonType.ConstInit);
 		}
 	}
 }
diff --git a/Engine/Source/Programs/Shared/EpicGames.UHT/Types/UhtFunction.cs b/Engine/Source/Programs/Shared/EpicGames.UHT/Types/UhtFunction.cs
--- a/Engine/Source/Programs/Shared/EpicGames.UHT/Types/UhtFunction.cs
+++ b/Engine/Source/Programs/Shared/EpicGames.UHT/Types/UhtFunction.cs
@@ -1347,8 +1347,8 @@ namespace EpicGames.UHT.Types
 				collector.AddExportType(this);
 				collector.AddSingleton(this);
 				collector.AddDeclaration(this, UhtSingletonType.Registered);
-				collector.AddCrossModuleReference(this, UhtSingletonType.Registered);
-				collector.AddCrossModuleReference(Package, UhtSingletonType.Registered);
+				collector.AddObjectReference(this, UhtSingletonType.Registered);
+				collector.AddObjectReference(Package, UhtSingletonType.Registered);
 			}
 			foreach (UhtType child in Children)
 			{
@@ -1358,15 +1358,15 @@ namespace EpicGames.UHT.Types
 			switch (FunctionType)
 			{
 				case UhtFunctionType.Function:
-					collector.AddCrossModuleReference(Session.UFunction, UhtSingletonType.ConstInit);
-					collector.AddCrossModuleReference(this, UhtSingletonType.ConstInit);
-					collector.AddCrossModuleReference(Package, UhtSingletonType.ConstInit);
+					collector.AddObjectReference(Session.UFunction, UhtSingletonType.ConstInit);
+					collector.AddObjectReference(this, UhtSingletonType.ConstInit);
+					collector.AddObjectReference(Package, UhtSingletonType.ConstInit);
 					break;
 				case UhtFunctionType.Delegate:
-					collector.AddCrossModuleReference(Session.UDelegateFunction, UhtSingletonType.ConstInit);
+					collector.AddObjectReference(Session.UDelegateFunction, UhtSingletonType.ConstInit);
 					break;
 				case UhtFunctionType.SparseDelegate:
-					collector.AddCrossModuleReference(Session.USparseDelegateFunction, UhtSingletonType.ConstInit);
+					collector.AddObjectReference(Session.USparseDelegateFunction, UhtSingletonType.ConstInit);
 					break;
 				default:
 					break;
diff --git a/Engine/Source/Programs/Shared/EpicGames.UHT/Types/UhtHeaderFile.cs b/Engine/Source/Programs/Shared/EpicGames.UHT/Types/UhtHeaderFile.cs
--- a/Engine/Source/Programs/Shared/EpicGames.UHT/Types/UhtHeaderFile.cs
+++ b/Engine/Source/Programs/Shared/EpicGames.UHT/Types/UhtHeaderFile.cs
@@ -227,7 +227,7 @@ namespace EpicGames.UHT.Types
 		/// Resource collector for the header file
 		/// </summary>
 		[JsonIgnore]
-		public UhtReferenceCollector References { get; } = new UhtReferenceCollector();
+		public UhtReferenceCollector References { get; init; } 
 
 		/// <summary>
 		/// Collection of headers directly included by this header
@@ -262,6 +262,7 @@ namespace EpicGames.UHT.Types
 		public UhtHeaderFile(UhtModule module, string path)
 		{
 			Module = module;
+			References = new UhtReferenceCollector(module);
 			HeaderFileTypeIndex = Session.GetNextHeaderFileTypeIndex();
 			_messageSite = new UhtSimpleMessageSite(Session);
 			_sourceFile = new UhtSourceFile(Session, path);
diff --git a/Engine/Source/Programs/Shared/EpicGames.UHT/Types/UhtModule.cs b/Engine/Source/Programs/Shared/EpicGames.UHT/Types/UhtModule.cs
--- a/Engine/Source/Programs/Shared/EpicGames.UHT/Types/UhtModule.cs
+++ b/Engine/Source/Programs/Shared/EpicGames.UHT/Types/UhtModule.cs
@@ -3,6 +3,7 @@
 using System;
 using System.Collections.Generic;
 using System.IO;
+using System.Linq;
 using System.Text.Json.Serialization;
 using EpicGames.Core;
 using EpicGames.UHT.Utils;
@@ -10,6 +11,12 @@ using UnrealBuildBase;
 
 namespace EpicGames.UHT.Types
 {
+	/// <summary>
+	/// Index assignments for other modules which are a dependency of this module, and objects which are linkable in this module from others.
+	/// </summary>
+	/// <param name="ReferencedModules">Mapping of dependency module to the index to use for encoded pointers when linking to it from this module.</param>
+	/// <param name="LinkableObjects">Mapping of public linkable objects in this module to indexes that other modules can use when encoding pointers to them.</param>
+	public record UhtModuleLinking(IReadOnlyDictionary<UhtModule, int> ReferencedModules, IReadOnlyDictionary<UhtObject, int> LinkableObjects);
 
 	/// <summary>
 	/// Represents a UHT module
@@ -18,6 +25,7 @@ namespace EpicGames.UHT.Types
 	{
 		private readonly List<UhtHeaderFile> _headers = new();
 		private readonly SortedDictionary<string, UhtPackage> _packages = new();
+		private readonly HashSet<string> _ubtModuleDependencies = [];
 		private readonly Lazy<UhtInputCache> _inputCache;
 
 		/// <summary>
@@ -78,6 +86,11 @@ namespace EpicGames.UHT.Types
 		[JsonIgnore]
 		public IReadOnlyList<UhtHeaderFile> Headers => _headers;
 
+		/// <summary>
+		/// Collection of module dependencies from UBT (may also include the names of modules not present in UHT)
+		/// </summary>
+		public IReadOnlySet<string> UBTModuleDependencies => _ubtModuleDependencies;
+
 		/// <summary>
 		/// True if the package is part of the engine
 		/// </summary>
@@ -114,6 +127,11 @@ namespace EpicGames.UHT.Types
 		[JsonIgnore]
 		public bool IsPlugin => Module.BaseDirectory.Replace('\\', '/').Contains("/Plugins/", StringComparison.Ordinal);
 
+		/// <summary>
+		/// Information about linking objects between this module and others generated by UhtSession after reference collection. 
+		/// </summary>
+		public UhtModuleLinking? Linking { get; set; }
+
 		/// <summary>
 		/// Create a new module from the given manifest module
 		/// </summary>
@@ -123,6 +141,10 @@ namespace EpicGames.UHT.Types
 		{
 			Session = session;
 			Module = module;
+			foreach (string dep in module.ModuleDependencies)
+			{
+				_ubtModuleDependencies.Add(dep);
+			}
 
 			EPackageFlags packageFlags = EPackageFlags.ContainsScript | EPackageFlags.Compiling;
 
diff --git a/Engine/Source/Programs/Shared/EpicGames.UHT/Types/UhtObject.cs b/Engine/Source/Programs/Shared/EpicGames.UHT/Types/UhtObject.cs
--- a/Engine/Source/Programs/Shared/EpicGames.UHT/Types/UhtObject.cs
+++ b/Engine/Source/Programs/Shared/EpicGames.UHT/Types/UhtObject.cs
@@ -33,6 +33,17 @@ namespace EpicGames.UHT.Types
 		[JsonIgnore]
 		public UhtObject? AlternateObject { get; set; } = null;
 
+		/// <summary>
+		/// Linkable top-level objects per packages should be in a specific order to simplify codegen and runtime initialization
+		/// </summary>
+		public int LinkOrder => this switch
+		{
+			UhtClass => 0,
+			UhtScriptStruct => 1,
+			UhtEnum => 2,
+			_ => 3,
+		};
+
 		/// <inheritdoc/>
 		public override string EngineClassName => "Object";
 
diff --git a/Engine/Source/Programs/Shared/EpicGames.UHT/Types/UhtProperty.cs b/Engine/Source/Programs/Shared/EpicGames.UHT/Types/UhtProperty.cs
--- a/Engine/Source/Programs/Shared/EpicGames.UHT/Types/UhtProperty.cs
+++ b/Engine/Source/Programs/Shared/EpicGames.UHT/Types/UhtProperty.cs
@@ -647,7 +647,7 @@ namespace EpicGames.UHT.Types
 	/// <summary>
 	/// Property member context provides extra context when formatting the member declaration and definitions.
 	/// </summary>
-	public interface IUhtPropertyMemberContext : IUhtSingletonNameProvider, IUhtSessionSettingsProvider
+	public interface IUhtPropertyMemberContext : IUhtObjectLinker, IUhtSessionSettingsProvider
 	{
 		/// <summary>
 		/// The outer/owning structure
@@ -1595,7 +1595,7 @@ namespace EpicGames.UHT.Types
 				}
 			}
 
-			builder.Append($"UE::CodeGen::ConstInit::FPropertyParams{{ .FieldClass = &F{EngineClassName}::ConstInitClass, .Owner = ");
+			builder.Append($"UE::CodeGen::ConstInit::FPropertyParams{{ .Owner = ");
 			if (outerFunc is not null)
 			{
 				outerFunc(builder);
@@ -2442,7 +2442,7 @@ namespace EpicGames.UHT.Types
 			CollectReferencesInternal(collector, PropertyFlags.HasAnyFlags(EPropertyFlags.ParmFlags), false);
 			if (Outer is UhtObject outerObject)
 			{
-				collector.AddCrossModuleReference(outerObject, UhtSingletonType.ConstInit);
+				collector.AddObjectReference(outerObject, UhtSingletonType.ConstInit);
 			}
 		}
 
diff --git a/Engine/Source/Programs/Shared/EpicGames.UHT/Types/UhtScriptStruct.cs b/Engine/Source/Programs/Shared/EpicGames.UHT/Types/UhtScriptStruct.cs
--- a/Engine/Source/Programs/Shared/EpicGames.UHT/Types/UhtScriptStruct.cs
+++ b/Engine/Source/Programs/Shared/EpicGames.UHT/Types/UhtScriptStruct.cs
@@ -1131,21 +1131,21 @@ namespace EpicGames.UHT.Types
 			}
 			collector.AddExportType(this);
 			collector.AddDeclaration(this, UhtSingletonType.Registered);
-			collector.AddCrossModuleReference(this, UhtSingletonType.Registered);
+			collector.AddObjectReference(this, UhtSingletonType.Registered);
 			if (SuperScriptStruct != null)
 			{
-				collector.AddCrossModuleReference(SuperScriptStruct, UhtSingletonType.Registered);
+				collector.AddObjectReference(SuperScriptStruct, UhtSingletonType.Registered);
 			}
-			collector.AddCrossModuleReference(Package, UhtSingletonType.Registered);
+			collector.AddObjectReference(Package, UhtSingletonType.Registered);
 			foreach (UhtType child in Children)
 			{
 				child.CollectReferences(collector);
 			}
 
-			collector.AddCrossModuleReference(Session.UScriptStruct, UhtSingletonType.ConstInit);
+			collector.AddObjectReference(Session.UScriptStruct, UhtSingletonType.ConstInit);
 			for (UhtStruct? superStruct = SuperScriptStruct; superStruct is not null; superStruct = superStruct.SuperStruct)
 			{
-				collector.AddCrossModuleReference(superStruct, UhtSingletonType.ConstInit);
+				collector.AddObjectReference(superStruct, UhtSingletonType.ConstInit);
 			}
 		}
 
diff --git a/Engine/Source/Programs/Shared/EpicGames.UHT/Types/UhtStruct.cs b/Engine/Source/Programs/Shared/EpicGames.UHT/Types/UhtStruct.cs
--- a/Engine/Source/Programs/Shared/EpicGames.UHT/Types/UhtStruct.cs
+++ b/Engine/Source/Programs/Shared/EpicGames.UHT/Types/UhtStruct.cs
@@ -556,7 +556,7 @@ namespace EpicGames.UHT.Types
 		public override void CollectReferences(IUhtReferenceCollector collector)
 		{
 			base.CollectReferences(collector);
-			collector.AddCrossModuleReference(SparseClassDataStruct, UhtSingletonType.ConstInit);
+			collector.AddObjectReference(SparseClassDataStruct, UhtSingletonType.ConstInit);
 		}
 	}
 }
diff --git a/Engine/Source/Programs/Shared/EpicGames.UHT/Types/UhtType.cs b/Engine/Source/Programs/Shared/EpicGames.UHT/Types/UhtType.cs
--- a/Engine/Source/Programs/Shared/EpicGames.UHT/Types/UhtType.cs
+++ b/Engine/Source/Programs/Shared/EpicGames.UHT/Types/UhtType.cs
@@ -889,7 +889,7 @@ namespace EpicGames.UHT.Types
 		public int LineNumber { get; set; }
 
 		/// <summary>
-		/// #define scope where the type was defined
+		/// #define scope where the type was defined. None means defined regardless of preprocessor definitions.
 		/// </summary>
 		[JsonConverter(typeof(JsonStringEnumConverter))]
 		public UhtDefineScope DefineScope { get; set; } = UhtDefineScope.None;
diff --git a/Engine/Source/Programs/Shared/EpicGames.UHT/Utils/UhtReferenceCollector.cs b/Engine/Source/Programs/Shared/EpicGames.UHT/Utils/UhtReferenceCollector.cs
--- a/Engine/Source/Programs/Shared/EpicGames.UHT/Utils/UhtReferenceCollector.cs
+++ b/Engine/Source/Programs/Shared/EpicGames.UHT/Utils/UhtReferenceCollector.cs
@@ -2,6 +2,7 @@
 
 using System;
 using System.Collections.Generic;
+using System.Linq;
 using System.Text;
 using EpicGames.Core;
 using EpicGames.UHT.Exporters.CodeGen;
@@ -16,13 +17,14 @@ namespace EpicGames.UHT.Utils
 	/// </summary>
 	public interface IUhtReferenceCollector
 	{
-
 		/// <summary>
-		/// Add a cross module reference to a given object type.
+		/// Add a reference to an object which needs to be linked. Depending on the type parameter and the configuration, 
+		/// it may be linked as a direct pointer, encoded pointer, or a pointer to a factory method.
+		/// Also records the module the object is in as a module-module dependency
 		/// </summary>
 		/// <param name="obj">Object/type being referenced</param>
 		/// <param name="type">Type of reference required</param>
-		void AddCrossModuleReference(UhtObject? obj, UhtSingletonType type);
+		void AddObjectReference(UhtObject? obj, UhtSingletonType type);
 
 		/// <summary>
 		/// Add an object declaration
@@ -63,40 +65,26 @@ namespace EpicGames.UHT.Utils
 	/// </summary>
 	public class UhtUniqueReferenceCollection
 	{
-		private const int Shift = 2;
-		private const int Mask = 0x3;
-
 		/// <summary>
 		/// Returns whether there is anything in this collection
 		/// </summary>
 		public bool IsEmpty => Uniques.Count == 0;
 
 		/// <summary>
-		/// List of all unique reference keys.  Use UngetKey to get the object index and the flag.
+		/// List of all unique references and the types required.
 		/// </summary>
-		private HashSet<int> Uniques { get; } = new HashSet<int>();
+		private Dictionary<UhtObject, UhtSingletonTypeFlag> Uniques { get; } = [];
 
-		/// <summary>
-		/// Return an encoded key that represents the object and type of reference.
-		/// If the object has the alternate object set (i.e. native interfaces), then
-		/// that object's index is used to generate the key.
-		/// </summary>
-		/// <param name="obj">Object being referenced</param>
-		/// <param name="type">The type of reference required for linking.</param>
-		/// <returns>Integer key value.</returns>
-		public static int GetKey(UhtObject obj, UhtSingletonType type)
-		{
-			return obj.AlternateObject != null ? GetKey(obj.AlternateObject, type) : (obj.ObjectTypeIndex << Shift) | ((int)type);
-		}
+		private static readonly UhtSingletonTypeFlag[] s_typeFlags = Enum.GetValues<UhtSingletonTypeFlag>();
 
-		/// <summary>
-		/// Given a key, return the object index and registered flag
-		/// </summary>
-		/// <param name="key">The key in question</param>
-		public static (int objectIndex, UhtSingletonType type) UngetKey(int key)
+		private static UhtSingletonType FlagToValue(UhtSingletonTypeFlag flag) => flag switch
 		{
-			return (key >> Shift, (UhtSingletonType)(key & Mask));
-		}
+			UhtSingletonTypeFlag.Registered => UhtSingletonType.Registered,
+			UhtSingletonTypeFlag.Unregistered => UhtSingletonType.Unregistered,
+			UhtSingletonTypeFlag.ConstInit => UhtSingletonType.ConstInit,
+			UhtSingletonTypeFlag.Statics => UhtSingletonType.Statics,
+			_ => throw new UhtIceException("flag should be single valued")
+		};
 
 		/// <summary>
 		/// Add the given object to the references
@@ -107,10 +95,28 @@ namespace EpicGames.UHT.Utils
 		{
 			if (obj != null)
 			{
-				Uniques.Add(GetKey(obj, type));
+				UhtSingletonTypeFlag flag = (UhtSingletonTypeFlag)(1 << (int)type);
+				if (!Uniques.TryAdd(obj, flag))
+				{
+					Uniques[obj] |= flag;
+				}
 			}
 		}
 
+		/// <summary>
+		/// Get references of a given set of types as a sorted enumerable.
+		/// </summary>
+		/// <param name="typesFilter"></param>
+		/// <returns></returns>
+		public IEnumerable<UhtObject> SortReferences(UhtSingletonTypeFlag typesFilter) 
+		{
+			return Uniques.Select((UhtObject Obj, UhtSingletonTypeFlag Matching) (p) => (p.Key, p.Value & typesFilter))
+				.Where(p => p.Matching != UhtSingletonTypeFlag.None)
+				.Select(p => p.Obj)
+				.OrderBy(o => o.Module.ShortName)
+				.ThenBy(o => o.EngineName);
+		}
+
 		/// <summary>
 		/// Return the collection of references sorted by the API string returned by the delegate.
 		/// </summary>
@@ -118,17 +124,23 @@ namespace EpicGames.UHT.Utils
 		/// <returns>Read only memory region of all the string.</returns>
 		public IReadOnlyList<string> GetSortedReferences(Func<int, UhtSingletonType, string?> referenceToString)
 		{
-			// Collect the unsorted array
-			List<string> sorted = new(Uniques.Count);
-			foreach(int key in Uniques)
+			IEnumerable<string> GetStrings(int objIndex, UhtSingletonTypeFlag flags)
 			{
-				(int objectIndex, UhtSingletonType type) = UngetKey(key);
-				if (referenceToString(objectIndex, type) is string decl)
+				foreach (UhtSingletonTypeFlag test in s_typeFlags)
 				{
-					sorted.Add(decl);
+					if ((flags & test) != UhtSingletonTypeFlag.None)
+					{
+						if (referenceToString(objIndex, FlagToValue(test)) is string s)
+						{
+							yield return s;
+						}
+					}
 				}
 			}
 
+			// Collect the unsorted array
+			List<string> sorted = [.. Uniques.SelectMany(p => GetStrings(p.Key.ObjectTypeIndex, p.Value))];
+
 			// Sort the array
 			sorted.Sort(StringComparerUE.OrdinalIgnoreCase);
 
@@ -156,15 +168,25 @@ namespace EpicGames.UHT.Utils
 	}
 
 	/// <summary>
-	/// Standard implementation of the reference collector interface
+	/// Standard implementation of the reference collector interface. Used by a single header file in
+	/// a module to collect references which are needed for the generated code for that particular header.
 	/// </summary>
-	public class UhtReferenceCollector : IUhtReferenceCollector
+	public class UhtReferenceCollector(UhtModule thisModule) : IUhtReferenceCollector
 	{
+		/// <summary>
+		/// Other modules encountered during reference collection for building module-to-module links
+		/// </summary>
+		public HashSet<UhtModule> Modules { get; } = [];
+
+		/// <summary>
+		/// Collection of unique object references in this module
+		/// </summary>
+		public UhtUniqueReferenceCollection InternalObjects { get; set; } = new();
 
 		/// <summary>
-		/// Collection of unique cross module references
+		/// Collection of unique object references in other modules
 		/// </summary>
-		public UhtUniqueReferenceCollection CrossModule { get; set; } = new();
+		public UhtUniqueReferenceCollection ExternalObjects { get; set; } = new();
 
 		/// <summary>
 		/// Collection of unique declarations
@@ -197,9 +219,20 @@ namespace EpicGames.UHT.Utils
 		public HashSet<UhtHeaderFile> ReferencedHeaders { get; } = [];
 
 		/// <inheritdoc/>
-		public void AddCrossModuleReference(UhtObject? obj, UhtSingletonType type)
+		public void AddObjectReference(UhtObject? obj, UhtSingletonType type)
 		{
-			CrossModule.Add(obj, type);
+			if (obj is not null)
+			{
+				if (obj.Module == thisModule)
+				{
+					InternalObjects.Add(obj, type);
+				}
+				else
+				{
+					ExternalObjects.Add(obj, type);
+					Modules.Add(obj.Module);	
+				}
+			}
 			if (obj != null && obj is not UhtPackage && type == UhtSingletonType.Registered)
 			{
 				ReferencedHeaders.Add(obj.HeaderFile);
diff --git a/Engine/Source/Programs/Shared/EpicGames.UHT/Utils/UhtSession.cs b/Engine/Source/Programs/Shared/EpicGames.UHT/Utils/UhtSession.cs
--- a/Engine/Source/Programs/Shared/EpicGames.UHT/Utils/UhtSession.cs
+++ b/Engine/Source/Programs/Shared/EpicGames.UHT/Utils/UhtSession.cs
@@ -1453,6 +1453,7 @@ namespace EpicGames.UHT.Utils
 			StepResolveFinal();
 			StepResolveValidate();
 			StepCollectReferences();
+			StepLinkModules();
 			TopologicalSortHeaderFiles();
 
 			// If we are deleting the reference directory, then wait for that task to complete
@@ -2254,6 +2255,41 @@ namespace EpicGames.UHT.Utils
 			ForEachHeader(headerFile => headerFile.CollectReferences(), true);
 		}
 
+		private void StepLinkModules()
+		{
+			Log.Logger.LogTrace("Step - Link modules");
+		
+			foreach (UhtModule module in Modules)
+			{
+				Dictionary<UhtModule, int> referencedModules = module.Headers.SelectMany(h => h.References.Modules)
+					.Distinct()
+					.OrderBy(m => m.ShortName)
+					.Select((module, index) => (module, index))
+					.ToDictionary();
+
+				if (TargetSettings?.ValidateModuleDependencies is true && IsUsingCompiledInObjectFormat(UhtCompiledInObjectFormat.ConstInit))
+				{
+					foreach (UhtModule otherModule in referencedModules.Keys)
+					{
+						if (!module.UBTModuleDependencies.Contains(otherModule.ShortName))
+						{
+							this.LogError($"Module {module.ShortName} needs a link dependency on module {otherModule.ShortName} to link UObjects, but it does not have one.");
+						}
+					}
+				}
+
+				// Top-level objects, excluding those with alternate objects (native inteface objects)
+				Dictionary<UhtObject, int> linkableObjects = module.Packages.SelectMany(p => p.Children.OfType<UhtObject>().Where(o => o.AlternateObject is null))
+					// Objects whose outer is not a package but which need to be registered
+					.Concat(module.Headers.SelectMany(h => h.References.Singletons).Where(o => o.Package != o.Outer))
+					.OrderBy(o => o.LinkOrder)
+					.Select((obj, index) => (obj, index))
+					.ToDictionary();
+
+				module.Linking = new UhtModuleLinking(referencedModules, linkableObjects);
+			}
+		}
+
 		private void ResolveAllHeaders(UhtResolvePhase resolvePhase)
 		{
 			ForEachHeader(headerFile => headerFile.Resolve(resolvePhase), resolvePhase.IsMultiThreadedResolvePhase());
diff --git a/Engine/Source/Programs/Shared/EpicGames.UHT/Utils/UhtStringBuilderExtensions.cs b/Engine/Source/Programs/Shared/EpicGames.UHT/Utils/UhtStringBuilderExtensions.cs
--- a/Engine/Source/Programs/Shared/EpicGames.UHT/Utils/UhtStringBuilderExtensions.cs
+++ b/Engine/Source/Programs/Shared/EpicGames.UHT/Utils/UhtStringBuilderExtensions.cs
@@ -297,7 +297,8 @@ namespace EpicGames.UHT.Utils
 						{
 							builder.Append("#if !UE_BUILD_SHIPPING\r\n");
 						}
-						builder.AppendTabs(tabs + 1).Append("{ UTF8TEXT(").AppendUTF8LiteralString(kvp.Key).Append("), UTF8TEXT(").AppendUTF8LiteralString(kvp.Value).Append(") },\r\n");
+						// TODO: Make UTF8TEXT and change code generation to emit \u and \U instad of \x 
+						builder.AppendTabs(tabs + 1).Append("{ ").AppendUTF8LiteralString(kvp.Key).Append(", ").AppendUTF8LiteralString(kvp.Value).Append(" },\r\n");
 						if (restricted)
 						{
 							builder.Append("#endif\r\n");
diff --git a/Engine/Source/Programs/UnrealBuildTool/Configuration/UEBuildTarget.cs b/Engine/Source/Programs/UnrealBuildTool/Configuration/UEBuildTarget.cs
--- a/Engine/Source/Programs/UnrealBuildTool/Configuration/UEBuildTarget.cs
+++ b/Engine/Source/Programs/UnrealBuildTool/Configuration/UEBuildTarget.cs
@@ -1414,10 +1414,6 @@ namespace UnrealBuildTool
 				throw new BuildException("bEnableConstInitUObject and bWithLiveCoding cannot currently be set simultaneously");
 			}
 
-			if (RulesObject.bEnableConstInitUObject && RulesObject.LinkType != TargetLinkType.Monolithic)
-			{
-				throw new BuildException("bEnableConstInitUObject can currently only be set on monolithic targets");
-			}
 
 			// If we're using the shared build environment, make sure all the settings are valid
 			if (RulesObject.BuildEnvironment == TargetBuildEnvironment.Shared)
@@ -6319,6 +6315,7 @@ namespace UnrealBuildTool
 			if (Rules.bEnableConstInitUObject)
 			{
 				GlobalCompileEnvironment.Definitions.Add("UE_WITH_CONSTINIT_UOBJECT=1");
+				GlobalCompileEnvironment.Definitions.Add("WITH_HOT_RELOAD=0");
 			}
 			else
 			{
diff --git a/Engine/Source/Programs/UnrealBuildTool/System/UHTExecution.cs b/Engine/Source/Programs/UnrealBuildTool/System/UHTExecution.cs
--- a/Engine/Source/Programs/UnrealBuildTool/System/UHTExecution.cs
+++ b/Engine/Source/Programs/UnrealBuildTool/System/UHTExecution.cs
@@ -106,6 +106,11 @@ namespace UnrealBuildTool
 		/// </summary>
 		public required DirectoryReference[] ModuleIncludePaths;
 
+		/// <summary>
+		/// Other modules containing UObject declarations which this module has link dependencies on
+		/// </summary>
+		public HashSet<string> ModuleDependencies = [];
+
 		/// <summary>
 		/// Module type
 		/// </summary>
@@ -232,6 +237,7 @@ namespace UnrealBuildTool
 			ModuleRulesFile = Reader.ReadFileReference();
 			ModuleDirectories = Reader.ReadArray<DirectoryItem>(() => Reader.ReadDirectoryItem()!)!;
 			ModuleIncludePaths = Reader.ReadArray<DirectoryReference>(Reader.ReadDirectoryReferenceNotNull)!;
+			ModuleDependencies = Reader.ReadHashSet(Reader.ReadString)!;
 			ModuleType = Reader.ReadString()!;
 			OverrideModuleType = Reader.ReadString()!;
 			PublicUObjectClassesHeaders = Reader.ReadList(() => Reader.ReadFileItem())!;
@@ -263,6 +269,7 @@ namespace UnrealBuildTool
 			Writer.WriteFileReference(ModuleRulesFile);
 			Writer.WriteArray<DirectoryItem>(ModuleDirectories, Writer.WriteDirectoryItem);
 			Writer.WriteArray<DirectoryReference>(ModuleIncludePaths, Writer.WriteDirectoryReference);
+			Writer.WriteHashSet(ModuleDependencies, Writer.WriteString);
 			Writer.WriteString(ModuleType);
 			Writer.WriteString(OverrideModuleType);
 			Writer.WriteList(PublicUObjectClassesHeaders, Item => Writer.WriteFileItem(Item));
@@ -499,6 +506,14 @@ namespace UnrealBuildTool
 					DirectoryItem GeneratedCodeDirectory = DirectoryItem.GetItemByDirectoryReference(Module.GeneratedCodeDirectoryUHT!);
 
 					DirectoryReference[] ModuleIncludePaths = [.. Module.PublicIncludePaths.Union(Module.InternalIncludePaths).Union(Module.PrivateIncludePaths)];
+					HashSet<string> ModuleDependencies = [];
+					foreach (UEBuildModule OtherModule in Module.GetDependencies(false, false))
+					{
+						if (ModuleToIndex.ContainsKey(OtherModule))
+						{
+							ModuleDependencies.Add(OtherModule.Name);	
+						}
+					}
 
 					string VerseMountPoint = String.Empty;
 					string VersePackageName = String.Empty;
@@ -517,6 +532,7 @@ namespace UnrealBuildTool
 						ModuleRulesFile = Module.RulesFile,
 						ModuleDirectories = Module.ModuleDirectories,
 						ModuleIncludePaths = ModuleIncludePaths,
+						ModuleDependencies = ModuleDependencies,
 						ModuleType = ModuleToType[Module].ToString(),
 						GeneratedCodeDirectory = GeneratedCodeDirectory,
 						GeneratedCodeVersion = GeneratedCodeVersion,
@@ -1056,6 +1072,7 @@ namespace UnrealBuildTool
 						PublicHeaders = [.. UObjectModule.PublicUObjectHeaders.Select((Header) => Header.AbsolutePath)],
 						InternalHeaders = [.. UObjectModule.InternalUObjectHeaders.Select((Header) => Header.AbsolutePath)],
 						PrivateHeaders = [.. UObjectModule.PrivateUObjectHeaders.Select((Header) => Header.AbsolutePath)],
+						ModuleDependencies = [.. UObjectModule.ModuleDependencies ],
 						PublicDefines = UObjectModule.PublicDefines,
 						GeneratedCPPFilenameBase = UObjectModule.GeneratedCPPFilenameBase,
 						SaveExportedHeaders = !UObjectModule.bIsReadOnly,
diff --git a/Engine/Source/Runtime/CoreUObject/Private/UObject/Class.cpp b/Engine/Source/Runtime/CoreUObject/Private/UObject/Class.cpp
--- a/Engine/Source/Runtime/CoreUObject/Private/UObject/Class.cpp
+++ b/Engine/Source/Runtime/CoreUObject/Private/UObject/Class.cpp
@@ -7689,10 +7689,6 @@ void UFunction::Bind()
 {
 	UClass* OwnerClass = GetOwnerClass();
 
-#if UE_WITH_CONSTINIT_UOBJECT
-	decltype(Func) OldFunc = Func;
-#endif // UE_WITH_CONSTINIT_UOBJECT
-
 	// if this isn't a native function, or this function belongs to a native interface class (which has no C++ version),
 	// use ProcessInternal (call into script VM only) as the function pointer for this function
 	if (!HasAnyFunctionFlags(FUNC_Native))
@@ -7702,6 +7698,10 @@ void UFunction::Bind()
 	}
 	else
 	{
+#if UE_WITH_CONSTINIT_UOBJECT
+		decltype(Func) OldFunc = Func;
+#endif // UE_WITH_CONSTINIT_UOBJECT
+
 		// Find the function in the class's native function lookup table.
 		FName Name = GetFName();
 		FNativeFunctionLookup* Found = OwnerClass->NativeFunctionLookupTable.FindByPredicate([=](const FNativeFunctionLookup& NativeFunctionLookup){ return Name == NativeFunctionLookup.Name; });
@@ -7715,13 +7715,13 @@ void UFunction::Bind()
 			UE_LOG(LogClass, Warning,TEXT("Failed to bind native function %s.%s"),*OwnerClass->GetName(),*GetName());
 		}
 #endif
-	}
 #if UE_WITH_CONSTINIT_UOBJECT
-	if (HasAnyInternalFlags(EInternalObjectFlags::Native))
-	{
-		check(Func == OldFunc);
-	}
+		if (HasAnyInternalFlags(EInternalObjectFlags::Native))
+		{
+			check(Func == OldFunc);
+		}
 #endif // UE_WITH_CONSTINIT_UOBJECT
+	}
 }
 
 void UFunction::Link(FArchive& Ar, bool bRelinkExistingProperties)
diff --git a/Engine/Source/Runtime/CoreUObject/Private/UObject/CompiledInObjectRegistry.h b/Engine/Source/Runtime/CoreUObject/Private/UObject/CompiledInObjectRegistry.h
--- a/Engine/Source/Runtime/CoreUObject/Private/UObject/CompiledInObjectRegistry.h
+++ b/Engine/Source/Runtime/CoreUObject/Private/UObject/CompiledInObjectRegistry.h
@@ -6,9 +6,11 @@
 
 #if UE_WITH_CONSTINIT_UOBJECT 
 
+#include "Concepts/DerivedFrom.h"
 #include "Misc/TVariant.h"
-#include "UObject/UObjectGlobals.h"
+#include "UObject/CompiledInObjectPtr.h"
 #include "UObject/RegisterCompiledInObjects.h"
+#include "UObject/UObjectGlobals.h"
 
 class FCompiledInObjectRegistry
 {
@@ -16,15 +18,72 @@ class FCompiledInObjectRegistry
     FCompiledInObjectRegistry() = default;
     ~FCompiledInObjectRegistry() = default;
 
-    // Linked list of objects to construct. New items are added here, before items that may have
-    // already been partially processed which are indicated by the following pointers.
-    FRegisterCompiledInObjects* ListHead = nullptr;
-    // First block of objects that have already been processed by AddAndHashObjects
-    FRegisterCompiledInObjects* AlreadyAdded = nullptr;
-    // First block of objects that have already been processed by FinishConstructingObjects
-    FRegisterCompiledInObjects* AlreadyConstructed = nullptr;
-    // First block of objects that have already been processed by CreateClassDefaultObjects
-    FRegisterCompiledInObjects* AlreadyConstructedDefaultObjects = nullptr;
+    // Matches list markers in TList noting which objects have been constructed to which step 
+    enum class EConstructionStep
+    {
+        AddObjects,
+        ConstructObjects,
+        ConstructDefaultObjects,
+    };
+
+    template<typename T> 
+    struct TList
+    {
+        // Linked list of objects to construct. New items are added here, before items that may have
+        // already been partially processed which are indicated by the following pointers.
+        T* Head = nullptr;
+        // First block of objects that have already been processed by AddAndHashObjects
+        T* AlreadyAdded = nullptr;
+        // First block of objects that have already been processed by FinishConstructingObjects
+        T* AlreadyConstructed = nullptr;
+        // First block of objects that have already been processed by CreateClassDefaultObjects
+        T* AlreadyConstructedDefaultObjects = nullptr;
+
+        /** Remove all objects after all construction work is complete. */
+        void EmptyObjects();
+
+        [[nodiscard]] T* AdvanceConstruction(EConstructionStep Step)
+        {
+            T** Link = nullptr;
+            switch (Step)
+            {
+                case EConstructionStep::AddObjects: 
+                    Link = &AlreadyAdded;
+                    break;
+                case EConstructionStep::ConstructObjects: 
+                    Link = &AlreadyConstructed;
+                    break;
+                case EConstructionStep::ConstructDefaultObjects: 
+                    Link = &AlreadyConstructedDefaultObjects;
+                    break;
+                default: check(false); 
+            }
+
+            T* Ret = *Link;
+            *Link = Head;
+            return Ret;
+        }
+
+        /** Return whether there are any registrants that not been processed by AddAndHashObjects or FinishConstructingObjects. */
+        [[nodiscard]] bool HasObjectsPendingConstruction() const
+        {
+            return Head && (Head != AlreadyAdded || Head != AlreadyConstructed);
+        }
+
+        /** Return whether there are any registrants that have not done all phases of construction. */
+        [[nodiscard]] bool HasPendingObjects() const
+        {
+            return Head && (Head != AlreadyAdded || Head != AlreadyConstructed || Head != AlreadyConstructedDefaultObjects);
+        }
+    };
+
+    // Used to mark a point in the lists to iterate up to when advancing construction
+    using FIterationMarker = TTuple<const FRegisterCompiledInObjects*, const FRegisterIntrinsicClass*>;
+
+    // UHT-generated objects in various states of construction
+    TList<FRegisterCompiledInObjects> GeneratedObjects;
+    // Intrinsic classes in various states of construction
+    TList<FRegisterIntrinsicClass> IntrinsicClasses;
 
 public:
     static FCompiledInObjectRegistry& Get()
@@ -36,6 +95,9 @@ public:
     /** Add the given set of objects to the linked list of items to process. */
     void AddObjects(FRegisterCompiledInObjects* Info);
 
+    /** Add the given intrinsic class to the linked list of items to process. */
+    void AddIntrinsicClass(FRegisterIntrinsicClass* Info);
+
     /** Remove all objects after all construction work is complete. */
     void EmptyObjects();
 
@@ -45,7 +107,10 @@ public:
     /** Return whether there are any registrants that have done all phases of construction. */
     [[nodiscard]] bool HasPendingObjects() const;
 
-    /** Initialize the name and index of all pending objects and add them to the object hash. */
+    /** 
+     * Initialize the name and index of all pending objects and add them to the object hash.
+     * Also links cross-module pointers if necessary.
+     */
     void AddAndHashObjects();
 
     /** 
@@ -65,21 +130,28 @@ public:
     void AssembleReferenceTokenStream();
 
 private:
-    // Iterate all linked FRegisterCompiledInObjects until Stop (not including it) and call F with each as a concrete variant
+	// Get the links in the lists of generated and intrinsic objects which have already completed the given construction step.
+    [[nodiscard]] FIterationMarker AdvanceConstruction(EConstructionStep Step);
+    
+    // Iterate all objects before the given links and call the functors on the FRegisterCompiledInObjects and FRegisterIntrinsicClass structs
     template<typename... FUNCTORS>
-    void Iterate(FRegisterCompiledInObjects* Stop, FUNCTORS&&... FS);
+    void Iterate(FIterationMarker Stop, FUNCTORS&&... FS);
 
+    // Iterate all objects before the given links, calling the functor on all non-null entries
     template<typename FUNCTOR>
-    void IterateObjects(FRegisterCompiledInObjects* Stop, FUNCTOR&& F);
+    void IterateObjects(FIterationMarker Stop, FUNCTOR&& F);
 
+    // Iterate all objects of type UScriptStruct before the given links, calling the functor on all non-null entries
     template<typename FUNCTOR>
-    void IterateScriptStructs(FRegisterCompiledInObjects* Stop, FUNCTOR&& F);
+    void IterateScriptStructs(FIterationMarker Stop, FUNCTOR&& F);
 
+    // Iterate all objects of type UClass before the given links, calling the functor on all non-null entries
     template<typename FUNCTOR>
-    void IterateClasses(FRegisterCompiledInObjects* Stop, FUNCTOR&& F);
+    void IterateClasses(FIterationMarker Stop, FUNCTOR&& F);
 
+    // Iterate all intrinsic classes before the given links, calling the functor on the intrinsic class and its static constructor function
     template<typename FUNCTOR>
-    void IterateIntrinsicClasses(FRegisterCompiledInObjects* Stop, FUNCTOR&& F);
+    void IterateIntrinsicClasses(FIterationMarker Stop, FUNCTOR&& F);
 
 #if WITH_METADATA
 	static TConstArrayView<UE::CodeGen::ConstInit::FMetaData> GetCompiledInMetaData(UEnum* InEnum);
@@ -88,14 +160,7 @@ private:
 #endif
 	static void AddMetaData(UStruct* Object);
 	static void AddMetaData(UEnum* Object);
-};
 
-
-inline void FCompiledInObjectRegistry::AddObjects(FRegisterCompiledInObjects* Info)
-{
-    // Insert new objects at the head of the list - AlreadyAdded/AlreadyConstructed keeps track of what we've already processed
-    Info->ListNext = ListHead;
-    ListHead = Info;
-}
+};
 
 #endif // UE_WITH_CONSTINIT_UOBJECT 
\ No newline at end of file
diff --git a/Engine/Source/Runtime/CoreUObject/Private/UObject/CompiledInUObjectInit.cpp b/Engine/Source/Runtime/CoreUObject/Private/UObject/CompiledInUObjectInit.cpp
--- a/Engine/Source/Runtime/CoreUObject/Private/UObject/CompiledInUObjectInit.cpp
+++ b/Engine/Source/Runtime/CoreUObject/Private/UObject/CompiledInUObjectInit.cpp
@@ -1,6 +1,10 @@
 // Copyright Epic Games, Inc. All Rights Reserved.
 
+#include "UObject/CompiledInObjectPtr.h"
 #include "UObject/ObjectMacros.h"
+#include "UObject/PropertyOptional.h"
+#include "UObject/UObjectBase.h"
+#include "UObject/UObjectGlobals.h"
 
 #if UE_WITH_CONSTINIT_UOBJECT
 
@@ -29,7 +33,6 @@ void ProcessNewlyLoadedUObjects(FName InModuleName, bool bCanProcessNewlyLoadedO
 
 
 // Requirements
-static_assert(IS_MONOLITHIC, "UE_WITH_CONSTINIT_UOBJECT can currently only be set for monolithic builds");
 static_assert(!WITH_LIVE_CODING, "UE_WITH_CONSTINIT_UOBJECT and WITH_LIVE_CODING cannot currently be set simultaneously.");
 static_assert(!WITH_HOT_RELOAD, "UE_WITH_CONSTINIT_UOBJECT and WITH_HOT_RELOAD cannot currently be set simultaneously.");
 static_assert(!WITH_RELOAD, "UE_WITH_CONSTINIT_UOBJECT and WITH_RELOAD cannot currently be set simultaneously.");
@@ -81,6 +84,14 @@ void ProcessNewlyLoadedUObjects(FName InModuleName, bool bCanProcessNewlyLoadedO
 {
 	if (!bCanProcessNewlyLoadedObjects)
 	{
+#if !IS_MONOLITHIC
+		// Even if we're delayed, do minimal construction to fixup names and encoded pointers because some
+		// module loads like to call StaticClass before they really ought to
+		if (UObjectInitialized())
+		{
+			FCompiledInObjectRegistry::Get().AddAndHashObjects();
+		}
+#endif
 		FCoreUObjectDelegates::CompiledInUObjectsRegisteredDelegate.Broadcast(InModuleName, ECompiledInUObjectsRegisteredStatus::Delayed);
 		return;
 	}
@@ -110,125 +121,230 @@ void ProcessNewlyLoadedUObjects(FName InModuleName, bool bCanProcessNewlyLoadedO
 
 static void InitializeConstInitProperties(UStruct* Struct);
 
-void FCompiledInObjectRegistry::EmptyObjects()
+namespace UE::CoreUObject::Private
+{
+#if !IS_MONOLITHIC
+	class FCompiledInObjectLinker
+	{
+		TConstArrayView<TConstArrayView<UObject*>> LinkedModules;
+	public:
+		FCompiledInObjectLinker(TConstArrayView<TConstArrayView<UObject*>> InLinkedModules)
+			: LinkedModules(InLinkedModules)
+		{
+		}
+
+		// Handler for all TObjectPtr-like things
+		template<typename T>
+		void LinkCompiledInObjectPointer(T& InPtr)
+		{
+			LinkCompiledInObjectPointer(ObjectPtr_Private::Friend::GetHandleRef(InPtr));
+		}
+		
+		void LinkCompiledInObjectPointer(FObjectHandle& InHandle)
+		{
+			if (InHandle.CompiledInPtr.IsEncodedRef())
+			{
+				UObject* Object = FindCompiledInObject(InHandle.CompiledInPtr);
+				check(Object);
+				InHandle.Pointer = Object;
+			}
+		}
+
+		void LinkCompiledInObjectPointer(UE::CodeGen::ConstInit::TCompiledInObjectPtr<const FStructBaseChain>& InPtr)
+		{
+			if (InPtr.IsEncodedRef())
+			{
+				UStruct* Object = FindCompiledInObject<UStruct>(InPtr.CompiledInPtr);
+				check(Object);
+				InPtr.Pointer = UE::Private::AsStructBaseChain(Object);
+			}
+		}
+
+		template<UE::CDerivedFrom<UObject> T>
+		void LinkCompiledInObjectPointer(UE::CodeGen::ConstInit::TCompiledInObjectPtr<T>& InPtr)
+		{
+			if (InPtr.IsEncodedRef())
+			{
+				T* Object = FindCompiledInObject<T>(InPtr.CompiledInPtr);
+				check(Object);
+				InPtr.Pointer = Object;
+			}
+		}
+
+		/**
+		 * Look up a compiled in object by its encoded pointer and return it as the given type.
+		 * Type checking is NOT performed, a reference to an object of the correct type is expected to have been
+		 * encoded by UHT.
+		 */
+		template<UE::CDerivedFrom<UObject> T> 
+		T* FindCompiledInObject(UE::CodeGen::ConstInit::FCompiledInObjectPtr InPtr)
+		{
+			return reinterpret_cast<T*>(FindCompiledInObject(InPtr));
+		}
+
+		/** Look up a compiled in object by its encoded pointer and return it */
+		UObject* FindCompiledInObject(UE::CodeGen::ConstInit::FCompiledInObjectPtr InPtr)
+		{
+			const UE::CodeGen::ConstInit::FCompiledInObjectReference Decoded = InPtr.GetDecoded();
+			return LinkedModules[Decoded.ModuleId][Decoded.ObjectId];
+		}
+	};
+#endif // !IS_MONOLITHIC
+}
+
+template<typename T>
+void FCompiledInObjectRegistry::TList<T>::EmptyObjects()
 {
-	checkfSlow(ListHead == nullptr || (ListHead == AlreadyAdded && ListHead == AlreadyConstructed && ListHead == AlreadyConstructedDefaultObjects), 
+	checkfSlow(Head == nullptr || (Head == AlreadyAdded && Head == AlreadyConstructed && Head == AlreadyConstructedDefaultObjects), 
 		TEXT("Attempting to empty deferred registry when we don't appear to have constructed all objects"));
 	// Just drop the whole list - if the structs are recompiled and call Link again we will replace the Next ptr
-	ListHead = nullptr;
+	Head = nullptr;
 	AlreadyAdded = nullptr;
 	AlreadyConstructed = nullptr;
 	AlreadyConstructedDefaultObjects = nullptr;
 }
 
+void FCompiledInObjectRegistry::EmptyObjects()
+{
+	GeneratedObjects.EmptyObjects();
+	IntrinsicClasses.EmptyObjects();
+}
+
 bool FCompiledInObjectRegistry::HasObjectsPendingConstruction() const
 {
-	return ListHead && (ListHead != AlreadyAdded || ListHead != AlreadyConstructed);
+	return GeneratedObjects.HasObjectsPendingConstruction() || IntrinsicClasses.HasObjectsPendingConstruction();
 }
 
 bool FCompiledInObjectRegistry::HasPendingObjects() const
 {
-	return ListHead && (ListHead != AlreadyAdded || ListHead != AlreadyConstructed || ListHead != AlreadyConstructedDefaultObjects);
+	return GeneratedObjects.HasPendingObjects() || IntrinsicClasses.HasPendingObjects();
+}
+
+FCompiledInObjectRegistry::FIterationMarker FCompiledInObjectRegistry::AdvanceConstruction(EConstructionStep Step)
+{
+	return MakeTuple(GeneratedObjects.AdvanceConstruction(Step), IntrinsicClasses.AdvanceConstruction(Step));
 }
 
 template<typename... FUNCTORS>
-void FCompiledInObjectRegistry::Iterate(FRegisterCompiledInObjects* Stop, FUNCTORS&&... FS)
+void FCompiledInObjectRegistry::Iterate(FIterationMarker Stop, FUNCTORS&&... FS)
 {
-	for (FRegisterCompiledInObjects* It = ListHead; It != Stop; It = It->ListNext)
+	auto Dispatch = UE::Overload(Forward<FUNCTORS>(FS)...);
+
+	if constexpr (std::is_invocable_v<decltype(Dispatch), const FRegisterIntrinsicClass&>)
 	{
-		Visit([Dispatch=UE::Overload(Forward<FUNCTORS>(FS)...)](const auto& Registrant) 
+		for (const FRegisterIntrinsicClass* It = IntrinsicClasses.Head; It != Stop.Get<const FRegisterIntrinsicClass*>(); It = It->ListNext)
 		{
-			if constexpr (std::is_invocable_v<decltype(Dispatch), decltype(Registrant)>)
-			{
-				Dispatch(Registrant);
-			}
-		}, It->Registrants);
+			Dispatch(*It);
+		}
+	}
+
+	if constexpr (std::is_invocable_v<decltype(Dispatch), const FRegisterCompiledInObjects&>)
+	{
+		for (const FRegisterCompiledInObjects* It = GeneratedObjects.Head; It != Stop.Get<const FRegisterCompiledInObjects*>(); It = It->ListNext)
+		{
+			Dispatch(*It);	
+		}
 	}
 }
 
 template<typename FUNCTOR>
-void FCompiledInObjectRegistry::IterateObjects(FRegisterCompiledInObjects* Stop, FUNCTOR&& F)
+void FCompiledInObjectRegistry::IterateObjects(FIterationMarker Stop, FUNCTOR&& F)
 {
 	Iterate(Stop, 
-		[&F](const FRegisterCompiledInObjects::FPackageRegistrants& PReg) 
-		{ 
-			F(PReg.Package);
-			for (UObject* Obj : PReg.Objects)
+		[&F](const FRegisterCompiledInObjects& Generated)
+		{
+			for (UPackage* Package : Generated.Objects.Packages)
 			{
-				F(Obj);
+				F(Package);
 			}
-		},
-		[&F](const FRegisterCompiledInObjects::FClassRegistrants& CReg) 
-		{ 
-			for (UScriptStruct* Struct : CReg.Structs)
+			for (UClass* Class : Generated.Objects.GetClasses())
+			{
+				if(Class)
+				{
+					F(Class);
+				}
+			}
+			for (UScriptStruct* Struct : Generated.Objects.GetStructs())
 			{
-				F(Struct);
+				if (Struct)
+				{
+					F(Struct);
+				}
 			}
-			for (UEnum* Enum : CReg.Enums)
+			for (UEnum* Enum : Generated.Objects.GetEnums())
 			{
-				F(Enum);
+				if (Enum)
+				{
+					F(Enum);
+				}
 			}
-			for (UClass* Class : CReg.Classes)
+			for (UObject* Obj : Generated.Objects.Others)
 			{
-				F(Class);
+				if (Obj)
+				{
+					F(Obj);
+				}
 			}
 		},
-		[&F](const FRegisterCompiledInObjects::FIntrinsicClassRegistrant& IReg) 
-		{ 
-			F(IReg.IntrinsicClass);
+		[&F](const FRegisterIntrinsicClass& Intrinsic)
+		{
+			F(Intrinsic.Class);
 		}
 	);
 }
 
 template<typename FUNCTOR>
-void FCompiledInObjectRegistry::IterateScriptStructs(FRegisterCompiledInObjects* Stop, FUNCTOR&& F)
+void FCompiledInObjectRegistry::IterateScriptStructs(FIterationMarker Stop, FUNCTOR&& F)
 {
 	Iterate(Stop, 
-		[&F](const FRegisterCompiledInObjects::FPackageRegistrants& PReg) 
-		{ 
-			for (UObject* Obj : PReg.Objects)
+		[&F](const FRegisterCompiledInObjects& Generated)
+		{
+			for (UScriptStruct* Struct : Generated.Objects.GetStructs())
 			{
-				if (UScriptStruct* Struct = Cast<UScriptStruct>(Obj))
-				{
+				if (Struct)
+				{ 
 					F(Struct);
 				}
 			}
-		},
-		[&F](const FRegisterCompiledInObjects::FClassRegistrants& CReg) 
-		{ 
-			for (UScriptStruct* Struct : CReg.Structs)
+			for (UObject* Obj : Generated.Objects.Others)
 			{
-				F(Struct);
+				if(UScriptStruct* Struct = Cast<UScriptStruct>(Obj))
+				{
+					F(Struct);
+				}
 			}
 		}
 	);
 }
 
 template<typename FUNCTOR>
-void FCompiledInObjectRegistry::IterateClasses(FRegisterCompiledInObjects* Stop, FUNCTOR&& F)
+void FCompiledInObjectRegistry::IterateClasses(FIterationMarker Stop, FUNCTOR&& F)
 {
 	Iterate(Stop, 
-		[&F](const FRegisterCompiledInObjects::FClassRegistrants& CReg) 
+		[&F](const FRegisterCompiledInObjects& Generated)
 		{
-		   	for (UClass* Class : CReg.Classes)
+			for (UClass* Class : Generated.Objects.GetClasses())
 			{
-				F(Class);
+				if (Class)
+				{
+					F(Class);
+				}
 			}
 		},
-		[&F](const FRegisterCompiledInObjects::FIntrinsicClassRegistrant& IReg) 
+		[&F](const FRegisterIntrinsicClass& Intrinsic)
 		{
-			F(IReg.IntrinsicClass);
+			F(Intrinsic.Class);
 		}
 	);
 }
 
 template<typename FUNCTOR>
-void FCompiledInObjectRegistry::IterateIntrinsicClasses(FRegisterCompiledInObjects* Stop, FUNCTOR&& F)
+void FCompiledInObjectRegistry::IterateIntrinsicClasses(FIterationMarker Stop, FUNCTOR&& F)
 {
 	Iterate(Stop, 
-		[&F](const FRegisterCompiledInObjects::FIntrinsicClassRegistrant& IReg) 
+		[&F](const FRegisterIntrinsicClass& Intrinsic) 
 		{
-			F(IReg.IntrinsicClass, IReg.Constructor);
+			F(Intrinsic.Class, Intrinsic.Constructor);
 		});
 }
 
@@ -236,8 +352,59 @@ void FCompiledInObjectRegistry::IterateIntrinsicClasses(FRegisterCompiledInObjec
 void FCompiledInObjectRegistry::AddAndHashObjects()
 {
 	SCOPED_BOOT_TIMING("FDeferredRegistry::AddAndHashPendingRegistrants");
-	FRegisterCompiledInObjects* Stop = AlreadyAdded;
-	AlreadyAdded = ListHead;
+	FIterationMarker Stop = AdvanceConstruction(EConstructionStep::AddObjects);
+
+	// Fixup object pointers first - in order to hash a class we will need e.g. the super pointer
+	// This is only necessary for generated objects and not intrinsic ones
+#if !IS_MONOLITHIC
+	Iterate(Stop,
+		[this](const FRegisterCompiledInObjects& Reg)
+		{
+			UE::CoreUObject::Private::FCompiledInObjectLinker Linker(Reg.Objects.LinkedModules);
+			for (UPackage* Package : Reg.Objects.Packages)
+			{
+				Package->LinkCompiledInPointerFields(Linker);
+			}
+			for (UClass* Class : Reg.Objects.GetClasses())
+			{
+				if (Class)
+				{ 
+					Class->LinkCompiledInPointerFields(Linker);
+					for (UField* Field = Class->Children; Field; Field = Field->Next)
+					{
+						Field->LinkCompiledInPointerFields(Linker);
+					}
+				}
+			}
+			for (UScriptStruct* Struct: Reg.Objects.GetStructs())
+			{
+				if (Struct)
+				{ 
+					Struct->LinkCompiledInPointerFields(Linker);
+					for (UField* Field = Struct->Children; Field; Field = Field->Next)
+					{
+						Field->LinkCompiledInPointerFields(Linker);
+					}
+				}
+			}
+			for (UEnum* Enum : Reg.Objects.GetEnums())
+			{
+				if (Enum)
+				{
+					Enum->LinkCompiledInPointerFields(Linker);
+				}
+			}
+			for (UObject* Obj : Reg.Objects.Others)
+			{
+				if (Obj)
+				{
+					Obj->LinkCompiledInPointerFields(Linker);
+				}
+			}
+		}
+	);
+#endif // !IS_MONOLITHIC
+
 	IterateObjects(Stop, [](UObject* Obj)
 	{
 		Obj->AddConstInitObject();
@@ -286,7 +453,7 @@ void FCompiledInObjectRegistry::AddMetaData(UObject* Object, TConstArrayView<UE:
 		FMetaData& MetaData = Object->GetPackage()->GetMetaData();
 		for (const UE::CodeGen::ConstInit::FMetaData& MetaDataParam : InMetaData)
 		{
-			MetaData.SetValue(Object, UTF8_TO_TCHAR(MetaDataParam.NameUTF8), UTF8_TO_TCHAR(MetaDataParam.ValueUTF8));
+			MetaData.SetValue(Object, UTF8_TO_TCHAR(reinterpret_cast<const UTF8CHAR*>(MetaDataParam.NameUTF8)), UTF8_TO_TCHAR(reinterpret_cast<const UTF8CHAR*>(MetaDataParam.ValueUTF8)));
 		}
 	}
 }
@@ -324,8 +491,7 @@ void FCompiledInObjectRegistry::AddMetaData(UEnum* Object)
 void FCompiledInObjectRegistry::FinishConstructingObjects()
 {
 	SCOPED_BOOT_TIMING("FDeferredRegistry::FinishConstructingPendingRegistrants");
-	FRegisterCompiledInObjects* Stop = AlreadyConstructed;
-	AlreadyConstructed = ListHead;
+	FIterationMarker Stop = AdvanceConstruction(EConstructionStep::ConstructObjects);
 	// Prepare struct cpp ops before linking to try and remove dependency between FStructProperty and UScriptStruct
 	IterateScriptStructs(Stop, [](UScriptStruct* Struct){
 		Struct->PrepareCppStructOps();
@@ -469,14 +635,21 @@ void FCompiledInObjectRegistry::FinishConstructingObjects()
 				{
 					if (ImplementedInterface.bVerseDirectInterface)
 					{
-						if (UVerseClass* InterfaceClass = CastChecked<UVerseClass>(ImplementedInterface.Class, ECastCheckedType::NullAllowed))
+						if (UVerseClass* InterfaceClass = CastChecked<UVerseClass>((UClass*)ImplementedInterface.Class, ECastCheckedType::NullAllowed))
 						{
 							VerseClass->DirectInterfaces.Add(InterfaceClass);
 						}
 					}
 				}
 			}
-			new (&Class->Interfaces) TArray<FImplementedInterface>(Interfaces);
+			
+
+			new (&Class->Interfaces) TArray<FImplementedInterface>();
+			Class->Interfaces.Reserve(Interfaces.Num());
+			for (const UE::CodeGen::ConstInit::FClassImplementedInterface& Interface : Interfaces)
+			{
+				Class->Interfaces.Emplace(Interface.Class, Interface.PointerOffset, Interface.bImplementedByK2);
+			}
 		}
 
 		Class->StaticLink();
@@ -503,8 +676,7 @@ void FCompiledInObjectRegistry::CreateClassDefaultObjects(FName InModuleName)
 	TRACE_LOADTIME_REQUEST_GROUP_SCOPE(TEXT("FCompiledInObjectRegistry::CreateClassDefaultObjects"));
 	SCOPED_BOOT_TIMING("FCompiledInObjectRegistry::CreateClassDefaultObjects");
 
-	FRegisterCompiledInObjects* Stop = AlreadyConstructedDefaultObjects;
-	AlreadyConstructedDefaultObjects = ListHead;
+	FIterationMarker Stop = AdvanceConstruction(EConstructionStep::ConstructDefaultObjects);
 
 	FCoreUObjectDelegates::CompiledInUObjectsRegisteredDelegate.Broadcast(InModuleName, ECompiledInUObjectsRegisteredStatus::PreCDO);
 
@@ -610,7 +782,7 @@ void FCompiledInObjectRegistry::CreateClassDefaultObjects(FName InModuleName)
 
 void FCompiledInObjectRegistry::AssembleReferenceTokenStream()
 {
-	IterateClasses(nullptr, [](UClass* Class)
+	IterateClasses(FIterationMarker{nullptr, nullptr}, [](UClass* Class)
 	{
 		// Assemble reference token stream for garbage collection/ RTGC.
 		if (!Class->HasAnyFlags(RF_ClassDefaultObject) && !Class->HasAnyClassFlags(CLASS_TokenStreamAssembled))
@@ -629,6 +801,26 @@ void FRegisterCompiledInObjects::Register()
 	FCompiledInObjectRegistry::Get().AddObjects(this);
 }
 
+void FRegisterIntrinsicClass::Register()
+{
+	FCompiledInObjectRegistry::Get().AddIntrinsicClass(this);
+}
+
+void FCompiledInObjectRegistry::AddObjects(FRegisterCompiledInObjects* Info)
+{
+	// Insert new objects at the head of the list - AlreadyAdded/AlreadyConstructed keeps track of what we've already processed
+	Info->ListNext = GeneratedObjects.Head;
+	GeneratedObjects.Head = Info;
+}
+
+void FCompiledInObjectRegistry::AddIntrinsicClass(FRegisterIntrinsicClass* Info)
+{
+	// Insert new objects at the head of the list - AlreadyAdded/AlreadyConstructed keeps track of what we've already processed
+	Info->ListNext = IntrinsicClasses.Head;
+	IntrinsicClasses.Head = Info;
+}
+
+
 void FField::InitializeConstInitField(FField* InOwner)
 {
 #if DO_GUARD_SLOW
@@ -736,6 +928,139 @@ void UObjectBase::AddConstInitObject()
 	checkSlow(!(OuterPrivate) != !(ClassPrivate->IsChildOf<UPackage>())); // logical xor
 }
 
+#if !IS_MONOLITHIC
+void UObjectBase::LinkCompiledInPointerFields(UE::CoreUObject::Private::FCompiledInObjectLinker& Linker)
+{
+	Linker.LinkCompiledInObjectPointer(ClassPrivate);
+	Linker.LinkCompiledInObjectPointer(OuterPrivate);
+}
+
+void UField::LinkCompiledInPointerFields(UE::CoreUObject::Private::FCompiledInObjectLinker& Linker)
+{
+	Super::LinkCompiledInPointerFields(Linker);
+	Linker.LinkCompiledInObjectPointer(Next);
+}
+
+void UStruct::LinkCompiledInPointerFields(UE::CoreUObject::Private::FCompiledInObjectLinker& Linker)
+{
+	Super::LinkCompiledInPointerFields(Linker);
+	Linker.LinkCompiledInObjectPointer(SuperStruct);
+#if USTRUCT_FAST_ISCHILDOF_IMPL == USTRUCT_ISCHILDOF_STRUCTARRAY
+	for (int32 i=0; i <= NumStructBasesInChainMinusOne; ++i)
+	{
+		UE::CodeGen::ConstInit::TCompiledInObjectPtr<const FStructBaseChain>& CompiledInPtr = CompiledInBases[i];
+		Linker.LinkCompiledInObjectPointer(CompiledInPtr);
+	}
+#endif
+	for (FField* Prop = ChildProperties; Prop; Prop = Prop->Next)
+	{
+		Prop->LinkCompiledInPointerFields(Linker);
+	}
+}
+
+void UClass::LinkCompiledInPointerFields(UE::CoreUObject::Private::FCompiledInObjectLinker& Linker)
+{
+	Super::LinkCompiledInPointerFields(Linker);
+	Linker.LinkCompiledInObjectPointer(ClassWithin);
+	Linker.LinkCompiledInObjectPointer(SparseClassDataStruct);
+	for (UE::CodeGen::ConstInit::FClassImplementedInterface& Interface : CompiledInInterfaces)
+	{
+		Linker.LinkCompiledInObjectPointer(Interface.Class);
+	}
+}
+
+void FField::LinkCompiledInPointerFields(UE::CoreUObject::Private::FCompiledInObjectLinker& Linker)
+{
+}
+
+void FByteProperty::LinkCompiledInPointerFields(UE::CoreUObject::Private::FCompiledInObjectLinker& Linker)
+{
+	Super::LinkCompiledInPointerFields(Linker);	
+	Linker.LinkCompiledInObjectPointer(Enum);
+}
+
+void FObjectPropertyBase::LinkCompiledInPointerFields(UE::CoreUObject::Private::FCompiledInObjectLinker& Linker)
+{
+	Super::LinkCompiledInPointerFields(Linker);	
+	Linker.LinkCompiledInObjectPointer(PropertyClass);
+}
+
+void FClassProperty::LinkCompiledInPointerFields(UE::CoreUObject::Private::FCompiledInObjectLinker& Linker)
+{
+	Super::LinkCompiledInPointerFields(Linker);	
+	Linker.LinkCompiledInObjectPointer(MetaClass);
+}
+
+void FSoftClassProperty::LinkCompiledInPointerFields(UE::CoreUObject::Private::FCompiledInObjectLinker& Linker)
+{
+	Super::LinkCompiledInPointerFields(Linker);	
+	Linker.LinkCompiledInObjectPointer(MetaClass);
+}
+
+void FInterfaceProperty::LinkCompiledInPointerFields(UE::CoreUObject::Private::FCompiledInObjectLinker& Linker)
+{
+	Super::LinkCompiledInPointerFields(Linker);	
+	Linker.LinkCompiledInObjectPointer(InterfaceClass);
+}
+
+void FStructProperty::LinkCompiledInPointerFields(UE::CoreUObject::Private::FCompiledInObjectLinker& Linker)
+{
+	Super::LinkCompiledInPointerFields(Linker);	
+	Linker.LinkCompiledInObjectPointer(Struct);
+}
+
+void FDelegateProperty::LinkCompiledInPointerFields(UE::CoreUObject::Private::FCompiledInObjectLinker& Linker)
+{
+	Super::LinkCompiledInPointerFields(Linker);	
+	Linker.LinkCompiledInObjectPointer(SignatureFunction);
+}
+
+void FMulticastDelegateProperty::LinkCompiledInPointerFields(UE::CoreUObject::Private::FCompiledInObjectLinker& Linker)
+{
+	Super::LinkCompiledInPointerFields(Linker);	
+	Linker.LinkCompiledInObjectPointer(SignatureFunction);
+}
+
+void FEnumProperty::LinkCompiledInPointerFields(UE::CoreUObject::Private::FCompiledInObjectLinker& Linker)
+{
+	Super::LinkCompiledInPointerFields(Linker);	
+	Linker.LinkCompiledInObjectPointer(Enum);
+	UnderlyingProp->LinkCompiledInPointerFields(Linker);
+}
+
+void FMapProperty::LinkCompiledInPointerFields(UE::CoreUObject::Private::FCompiledInObjectLinker& Linker)
+{
+	Super::LinkCompiledInPointerFields(Linker);	
+	KeyProp->LinkCompiledInPointerFields(Linker);
+	ValueProp->LinkCompiledInPointerFields(Linker);
+}
+
+void FSetProperty::LinkCompiledInPointerFields(UE::CoreUObject::Private::FCompiledInObjectLinker& Linker)
+{
+	Super::LinkCompiledInPointerFields(Linker);	
+	ElementProp->LinkCompiledInPointerFields(Linker);
+}
+
+void FArrayProperty::LinkCompiledInPointerFields(UE::CoreUObject::Private::FCompiledInObjectLinker& Linker)
+{
+	Super::LinkCompiledInPointerFields(Linker);	
+	Inner->LinkCompiledInPointerFields(Linker);
+}
+
+void FFieldPathProperty::LinkCompiledInPointerFields(UE::CoreUObject::Private::FCompiledInObjectLinker& Linker)
+{
+	Super::LinkCompiledInPointerFields(Linker);	
+	FName FieldClassName(PropertyClassName);
+	PropertyClass = FFieldClass::GetNameToFieldClassMap().FindRef(FieldClassName);
+}
+
+void FOptionalProperty::LinkCompiledInPointerFields(UE::CoreUObject::Private::FCompiledInObjectLinker& Linker)
+{
+	Super::LinkCompiledInPointerFields(Linker);	
+	ValueProperty->LinkCompiledInPointerFields(Linker);
+}
+#endif // !IS_MONOLITHIC
+
 void UClass::InitFuncMap()
 {
 	// Called during UObject init, should be no lock contention
diff --git a/Engine/Source/Runtime/CoreUObject/Private/UObject/Field.cpp b/Engine/Source/Runtime/CoreUObject/Private/UObject/Field.cpp
--- a/Engine/Source/Runtime/CoreUObject/Private/UObject/Field.cpp
+++ b/Engine/Source/Runtime/CoreUObject/Private/UObject/Field.cpp
@@ -47,6 +47,9 @@ Field.cpp: Defines FField property system fundamentals
 #if UE_WITH_CONSTINIT_UOBJECT
 FRegisterFieldClass::FRegisterFieldClass(FFieldClass* InFieldClass, const TCHAR* InCppName)
 {
+#if !IS_MONOLITHIC
+	InFieldClass->SuperClass = InFieldClass->GetSuperStaticClass();
+#endif
 	InFieldClass->Name = &InCppName[1];
 	FFieldClass::GetAllFieldClasses().Add(InFieldClass);
 	FFieldClass::GetNameToFieldClassMap().Add(InFieldClass->Name, InFieldClass);
@@ -279,21 +282,19 @@ FField* FField::Construct(const FFieldVariant& InOwner, const FName& InName, EOb
 	// Can't construct an abstract type
 	return nullptr;
 }
+
 #if UE_WITH_CONSTINIT_UOBJECT
-constinit TNoDestroy<FFieldClass> FField::ConstInitClass(
-    NoDestroyConstEval,
-    ConstEval,
-    TEXT("FField"),
-    FField::StaticClassCastFlagsPrivate(),
-    FField::StaticClassCastFlags(),
-    nullptr,
-    (FFieldClass::FConstructFunction*)&FField::Construct // Explicit cast required to select overload for template argument
-);
+constinit TNoDestroy<TFieldClass<FField>> FField::ConstInitClass;
 #endif // UE_WITH_CONSTINIT_UOBJECT
+
 FFieldClass* FField::StaticClass()
 {
+#if UE_WITH_CONSTINIT_UOBJECT
+	return &ConstInitClass;
+#else
 	static FFieldClass StaticFieldClass(TEXT("FField"), FField::StaticClassCastFlagsPrivate(), FField::StaticClassCastFlags(), nullptr, &FField::Construct);
 	return &StaticFieldClass;
+#endif
 }
 
 FField::FField(EInternal InInernal, FFieldClass* InClass)
diff --git a/Engine/Source/Runtime/CoreUObject/Private/UObject/PropertyBool.cpp b/Engine/Source/Runtime/CoreUObject/Private/UObject/PropertyBool.cpp
--- a/Engine/Source/Runtime/CoreUObject/Private/UObject/PropertyBool.cpp
+++ b/Engine/Source/Runtime/CoreUObject/Private/UObject/PropertyBool.cpp
@@ -137,7 +137,8 @@ int32 FBoolProperty::GetMinAlignment() const
 void FBoolProperty::LinkInternal(FArchive& Ar)
 {
 #if UE_WITH_CONSTINIT_UOBJECT
-	if (IsNative())
+	// Can't call FField::IsNative here because sometimes bool props don't have an owner when linked in the python script plugin
+	if (UObject* Obj = GetOwnerUObject(); Obj && Obj->IsNative())
 	{
 		// Calculate bit mask settings since we couldn't do it at compile time
 		bool bIsNativeBool = IsNativeBool();
diff --git a/Engine/Source/Runtime/CoreUObject/Public/UObject/Class.h b/Engine/Source/Runtime/CoreUObject/Public/UObject/Class.h
--- a/Engine/Source/Runtime/CoreUObject/Public/UObject/Class.h
+++ b/Engine/Source/Runtime/CoreUObject/Public/UObject/Class.h
@@ -215,6 +215,9 @@ public:
 	COREUOBJECT_API virtual void PostLoad() override;
 	COREUOBJECT_API virtual bool NeedsLoadForClient() const override;
 	COREUOBJECT_API virtual bool NeedsLoadForServer() const override;
+#if UE_WITH_CONSTINIT_UOBJECT && !IS_MONOLITHIC
+	COREUOBJECT_API virtual void LinkCompiledInPointerFields(UE::CoreUObject::Private::FCompiledInObjectLinker& Linker) override;
+#endif // UE_WITH_CONSTINIT_UOBJECT && !IS_MONOLITHIC
 
 	// UField interface.
 	COREUOBJECT_API virtual void AddCppProperty(FProperty* Property);
@@ -418,6 +421,14 @@ public:
 	UE_NONCOPYABLE(FStructBaseChain);
 	
 #if UE_WITH_CONSTINIT_UOBJECT
+#if !IS_MONOLITHIC
+	consteval FStructBaseChain(TArrayView<UE::CodeGen::ConstInit::TCompiledInObjectPtr<const FStructBaseChain>> InBases)
+		: CompiledInBases(InBases.GetData())
+		, NumStructBasesInChainMinusOne(InBases.Num()-1)
+	{
+	}
+#endif // !IS_MONOLITHIC
+
 	consteval FStructBaseChain(TArrayView<const FStructBaseChain*> InBases)
 		: StructBaseChainArray(InBases.GetData())
 		, NumStructBasesInChainMinusOne(InBases.Num()-1)
@@ -434,7 +445,15 @@ public:
 	}
 
 private:
-	const FStructBaseChain** StructBaseChainArray = nullptr;
+	union
+	{
+#if UE_WITH_CONSTINIT_UOBJECT && !IS_MONOLITHIC
+		// Compiled in list of bases which may need runtime linking.
+		UE::CodeGen::ConstInit::TCompiledInObjectPtr<const FStructBaseChain>* CompiledInBases;
+#endif // UE_WITH_CONSTINIT_UOBJECT && !IS_MONOLITHIC
+		// Runtime allocated or post-converstion compiled in list of bases
+		const FStructBaseChain** StructBaseChainArray = nullptr;
+	};
 	int32 NumStructBasesInChainMinusOne = -1;
 
 	friend class UStruct;
@@ -443,7 +462,7 @@ private:
 
 namespace UE::Private
 {
-consteval FStructBaseChain* AsStructBaseChain(UStruct* Struct);
+constexpr FStructBaseChain* AsStructBaseChain(UStruct* Struct);
 
 template<typename... Classes>
 struct TIntrinsicClassBaseChainImpl;
@@ -491,7 +510,7 @@ protected:
 	// For construction of intrinsic class GC schema 
 	friend struct Z_Construct_UClass_UStruct_Statics;
 	// For access to struct base chain from internal code
-	friend consteval FStructBaseChain* UE::Private::AsStructBaseChain(UStruct*);
+	friend constexpr FStructBaseChain* UE::Private::AsStructBaseChain(UStruct*);
 	// For access to struct base chain from possible non-internal intrinsic classes
 	template<typename... Classes>
 	friend struct UE::Private::TIntrinsicClassBaseChainImpl;
@@ -625,6 +644,9 @@ public:
 #endif // UE_WITH_CONSTINIT_UOBJECT
 
 	// UObject interface.
+#if UE_WITH_CONSTINIT_UOBJECT && !IS_MONOLITHIC
+	COREUOBJECT_API virtual void LinkCompiledInPointerFields(UE::CoreUObject::Private::FCompiledInObjectLinker& Linker) override;
+#endif // UE_WITH_CONSTINIT_UOBJECT && !IS_MONOLITHIC
 	COREUOBJECT_API virtual void Serialize(FArchive& Ar) override;
 	COREUOBJECT_API virtual void Serialize(FStructuredArchive::FRecord Record) override;
 	COREUOBJECT_API virtual void PostLoad() override;
@@ -1016,7 +1038,7 @@ private:
 
 #if UE_WITH_CONSTINIT_UOBJECT
 // Private access to protected base class for UHT generated code
-UE_INTERNAL inline consteval FStructBaseChain* UE::Private::AsStructBaseChain(UStruct* Struct)
+inline constexpr FStructBaseChain* UE::Private::AsStructBaseChain(UStruct* Struct)
 {
 	return Struct;
 }
@@ -2819,6 +2841,9 @@ public:
 	)
 		: Super(ObjectParams, FieldParams)
 		, CppTypeStaticUTF8(EnumParams.CppTypeStaticUTF8)
+#if WITH_METADATA 
+		, CompiledInMetaData(EnumParams.MetaData)
+#endif
 		, Names(EnumParams.EnumValues.Num(), EnumParams.StaticNamesUTF8.GetData(), EnumParams.EnumValues.GetData())
 		, CppForm((ECppForm)EnumParams.CppForm)
 		, EnumFlags(EnumParams.EnumFlags)
@@ -3709,10 +3734,6 @@ struct FImplementedInterface
 		, bImplementedByK2(InImplementedByK2)
 	{}
 
-#if UE_WITH_CONSTINIT_UOBJECT
-	explicit FImplementedInterface(UE::CodeGen::ConstInit::FClassImplementedInterface InCompiledIn);
-#endif // UE_WITH_CONSTINIT_UOBJECT
-
 	friend COREUOBJECT_API FArchive& operator<<(FArchive& Ar, FImplementedInterface& A);
 };
 
@@ -4022,8 +4043,9 @@ public:
 		/**
 		 * Static data about compiled in interfaces from UHT
 		 * This union member is active until full compiled-in UObject construction in ProcessNewlyLoadedUObjects
+		 * It's mutable because pointers are fixed up here before copying to Interfaces.
 		 */
-		TConstArrayView<UE::CodeGen::ConstInit::FClassImplementedInterface> CompiledInInterfaces;
+		TArrayView<UE::CodeGen::ConstInit::FClassImplementedInterface> CompiledInInterfaces;
 #endif // UE_WITH_CONSTINIT_UOBJECT
 	};
 
@@ -4059,6 +4081,22 @@ public:
 		UE::CodeGen::ConstInit::FStructParams StructParams,
 		UE::CodeGen::ConstInit::FClassParams Params
 	);
+	// Special constructor for the UClass representing UClass itself
+	consteval UClass(
+		UE::CodeGen::ConstInit::EClass,
+		UE::CodeGen::ConstInit::FObjectParams ObjectParams,
+		UE::CodeGen::ConstInit::FUFieldParams FieldParams,
+		UE::CodeGen::ConstInit::FStructParams StructParams,
+		UE::CodeGen::ConstInit::FClassParams Params
+	);
+	// Special constructor for the UClass representing UObject itself
+	consteval UClass(
+		UE::CodeGen::ConstInit::EObject,
+		UE::CodeGen::ConstInit::FObjectParams ObjectParams,
+		UE::CodeGen::ConstInit::FUFieldParams FieldParams,
+		UE::CodeGen::ConstInit::FStructParams StructParams,
+		UE::CodeGen::ConstInit::FClassParams Params
+	);
 #endif
 
 	// Destructor
@@ -4165,6 +4203,9 @@ public:
 	COREUOBJECT_API UFunction* FindFunctionByName(FName InName, EIncludeSuperFlag::Type IncludeSuper = EIncludeSuperFlag::IncludeSuper) const;
 
 	// UObject interface.
+#if UE_WITH_CONSTINIT_UOBJECT && !IS_MONOLITHIC 
+	COREUOBJECT_API virtual void LinkCompiledInPointerFields(UE::CoreUObject::Private::FCompiledInObjectLinker& Linker) override;
+#endif // UE_WITH_CONSTINIT_UOBJECT && !IS_MONOLITHIC
 	COREUOBJECT_API virtual void Serialize(FArchive& Ar) override;
 	COREUOBJECT_API virtual void PostLoad() override;
 	COREUOBJECT_API virtual void FinishDestroy() override;
@@ -4768,8 +4809,28 @@ protected:
 };
 
 #if UE_WITH_CONSTINIT_UOBJECT
-extern COREUOBJECT_API TNoDestroy<UClass> Z_ConstInit_UClass_UClass;
-extern COREUOBJECT_API TNoDestroy<UClass> Z_ConstInit_UClass_UObject;
+extern TNoDestroy<UClass> Z_ConstInit_UClass_UObject;
+extern TNoDestroy<UClass> Z_ConstInit_UClass_UClass;
+consteval UClass::UClass(
+	UE::CodeGen::ConstInit::EClass,
+	UE::CodeGen::ConstInit::FObjectParams ObjectParams,
+	UE::CodeGen::ConstInit::FUFieldParams FieldParams,
+	UE::CodeGen::ConstInit::FStructParams StructParams,
+	UE::CodeGen::ConstInit::FClassParams Params
+)
+	: UClass(ObjectParams.SetClass(&Z_ConstInit_UClass_UClass), FieldParams, StructParams, Params)
+{
+}
+consteval UClass::UClass(
+	UE::CodeGen::ConstInit::EObject,
+	UE::CodeGen::ConstInit::FObjectParams ObjectParams,
+	UE::CodeGen::ConstInit::FUFieldParams FieldParams,
+	UE::CodeGen::ConstInit::FStructParams StructParams,
+	UE::CodeGen::ConstInit::FClassParams Params
+)
+	: UClass(ObjectParams, FieldParams, StructParams, Params.SetWithin(&Z_ConstInit_UClass_UObject))
+{
+}
 consteval UClass::UClass(
 	UE::CodeGen::ConstInit::FObjectParams ObjectParams,
 	UE::CodeGen::ConstInit::FUFieldParams FieldParams,
@@ -4777,7 +4838,7 @@ consteval UClass::UClass(
 	UE::CodeGen::ConstInit::FClassParams Params
 )
 	// Initialization of class for UClass passes null because of constexpr evaluation order and TNoDestroy wrapper
-	: Super(ObjectParams.SetClassIfUnset(&Z_ConstInit_UClass_UClass), FieldParams, StructParams)
+	: Super(ObjectParams, FieldParams, StructParams)
 	, ClassConstructor(Params.ClassConstructor)
 	, ClassVTableHelperCtorCaller(Params.ClassVTableHelperCtorCaller)
 	, CppClassStaticFunctions(Params.CppClassStaticFunctions)
@@ -4786,7 +4847,7 @@ consteval UClass::UClass(
 	, bLayoutChanging(false)
 	, ClassFlags(Params.ClassFlags | CLASS_Native)
 	, ClassCastFlags(Params.ClassCastFlags)
-	, ClassWithin(ConstEval, Params.ClassWithin ? Params.ClassWithin : &Z_ConstInit_UClass_UObject)
+	, ClassWithin(Params.ClassWithin)
 #if WITH_EDITORONLY_DATA
 	, PropertiesPendingDestruction(nullptr)
 #endif
@@ -4796,7 +4857,7 @@ consteval UClass::UClass(
 	, SparseClassData(nullptr)
 	, SparseClassDataStruct(Params.SparseClassDataStruct)
 #if WITH_EDITOR
-	, CppTypeInfo(InClassParams.Place, InCppClassTypeInfo)
+	, CppTypeInfo(InPlace, Params.CppClassTypeInfo)
 #endif
 	, FuncMap(ConstEval)
 	, FuncMapLock()
@@ -4810,13 +4871,6 @@ consteval UClass::UClass(
 		checkf(Params.NativeFunctions.Num() == 0, TEXT("UObject must have no native functions because ScriptCore will try and override it"));
 	}
 }
-
-inline FImplementedInterface::FImplementedInterface(UE::CodeGen::ConstInit::FClassImplementedInterface InCompiledIn)
-	: Class(InCompiledIn.Class)
-	, PointerOffset(InCompiledIn.PointerOffset)
-	, bImplementedByK2(InCompiledIn.bImplementedByK2)
-{
-}
 #endif // UE_WITH_CONSTINIT_UOBJECT
 
 /**
diff --git a/Engine/Source/Runtime/CoreUObject/Public/UObject/CompiledInObjectPtr.h b/Engine/Source/Runtime/CoreUObject/Public/UObject/CompiledInObjectPtr.h
new file mode 100644
--- /dev/null
+++ b/Engine/Source/Runtime/CoreUObject/Public/UObject/CompiledInObjectPtr.h
@@ -0,0 +1,225 @@
+// Copyright Epic Games, Inc. All Rights Reserved.
+
+#pragma once
+
+/*
+ * This header contains types used to store and decode references between compiled-in objects in
+ * different modules.
+ * e.g. a UClass in a project module which extends Actor needs a pointer to the AActor UClass
+ * instance in its UStruct::SuperStruct field. Some compilers and linkers will not allow this 
+ * field to be constinit-initialized to the address of a constinit variable in another module.
+ *
+ * The values encoded are produced by UHT such that they correspond to UObject initialization 
+ * data handled by FCompiledInObjectRegistry.
+ *
+ * Incomplete features:
+ *	- Stability of references across live coding/hot reload recompiles (if UHT runs)
+ * 	- Stability of references for installed-engine scenarios (i.e. references for engine types
+ *	  must be fixed when the engine build is created so that plugins/projects can link against
+ *	  the right objects at runtime)
+ *
+ * Non-features:
+ *	- There is no ability to distinguish the encoding of these compiled in references from
+ *	  other encoded pointers used by TObjectPtr such as late-resolved object handles or remote
+ *	  object handles.
+ *  - Therefore there is no ability to defer resolution of compiled in references until an
+ *	  arbitrary later point or resolve them on demand. They can only be resolved by the
+ * 	  CoreUObject module as part of initializing a module which contains UObjects.
+ *
+ */
+
+#include "CoreTypes.h"
+#include "Concepts/ConvertibleTo.h"
+#include "UObject/ObjectMacros.h"
+
+#if UE_WITH_CONSTINIT_UOBJECT
+
+namespace UE::CodeGen::ConstInit
+{
+
+// Decoded reference to a compiled in object
+struct FCompiledInObjectReference
+{
+	uint32 ModuleId = 0;
+	uint32 ObjectId = 0;
+
+	static inline constexpr int32 NumObjectIdBits = 32;
+	static inline constexpr uint32 MaskObjectId = 0xFFFF'FFFF;
+	static inline constexpr uint32 MaxModuleId = 0x7FFF'FFFF;
+
+	// To distinguish compiled-in references from pointers the low bit is always 1
+	// UObjects should have a natural alignment greater than 1, allowing this.
+	static inline constexpr uint32 CompiledInRefMask = 0x1;
+	static inline constexpr int32 CompiledInRefShift = 1;
+
+	static_assert(sizeof(UPTRINT) >= 8, "FCompiledInObjectReference requires pointers to be 8 bytes or more");
+
+	/* Check whether a pointer is an encoded reference or a raw pointer */
+	[[nodiscard]] static constexpr bool IsEncodedRef(UPTRINT InEncoded)
+	{
+		return (InEncoded & CompiledInRefMask) == CompiledInRefMask;
+	}
+
+	/* Decode an encoded pointer into this structure type */
+	[[nodiscard]] static constexpr FCompiledInObjectReference FromEncoded(UPTRINT InEncoded)
+	{
+		InEncoded >>= CompiledInRefShift;
+		uint32 ModuleId = InEncoded >> NumObjectIdBits;
+		uint32 ObjectId = InEncoded & MaskObjectId;
+		return { .ModuleId = ModuleId, .ObjectId = ObjectId };
+	}
+
+	/* Return this structure as an encoded pointer */
+	[[nodiscard]] constexpr UPTRINT GetEncoded() const
+	{
+		const UPTRINT CombinedId = ((UPTRINT)ModuleId << NumObjectIdBits) | ObjectId;
+		return (CombinedId << CompiledInRefShift) | CompiledInRefMask;
+	}
+};
+	
+// Encoded pointer to another compiled-in object for runtime linking
+struct FCompiledInObjectPtr
+{
+	[[nodiscard]] constexpr FCompiledInObjectPtr() = default;
+	[[nodiscard]] explicit constexpr FCompiledInObjectPtr(FCompiledInObjectReference InRef)
+		: Encoded(InRef.GetEncoded())
+	{
+	}
+	constexpr FCompiledInObjectPtr(const FCompiledInObjectPtr&) = default;
+	constexpr FCompiledInObjectPtr& operator=(const FCompiledInObjectPtr&) = default;
+
+	[[nodiscard]] constexpr bool IsUnset() const
+	{
+		return !FCompiledInObjectReference::IsEncodedRef(Encoded);
+	}
+
+	[[nodiscard]] constexpr bool IsEncodedRef() const
+	{
+		return FCompiledInObjectReference::IsEncodedRef(Encoded);
+	}
+
+	[[nodiscard]] constexpr FCompiledInObjectReference GetDecoded() const
+	{
+		return FCompiledInObjectReference::FromEncoded(Encoded);
+	}
+
+private:
+	uint64 Encoded = 0;
+};
+
+#if IS_MONOLITHIC 
+template<typename T>
+using TCompiledInObjectPtr = T*;
+#else // IS_MONOLITHIC
+
+// Wrapper for a compiled-in object pointer which is not a TObjectPtr to be used in other UHT-generated
+// structures
+// T is not required to derive from UObject, e.g. it may be FStructBaseChain. 
+// In that case if the pointer is an encoded reference, it will resolve to a UObject which can be cast to the required type.
+// e.g. lists of implemented interfaces, struct base chains
+template<typename T>
+struct TCompiledInObjectPtr
+{
+	union
+	{
+		FCompiledInObjectPtr CompiledInPtr;
+		T* Pointer = nullptr; 
+	};
+
+	constexpr TCompiledInObjectPtr() = default;
+	constexpr TCompiledInObjectPtr(const TCompiledInObjectPtr&) = default;
+	constexpr TCompiledInObjectPtr& operator=(const TCompiledInObjectPtr&) = default;
+
+	// Implicit compile-time construction from object pointer
+	template<UE::CConvertibleTo<T*> U> requires (!UE::CSameAs<T, U>)
+	[[nodiscard]] constexpr TCompiledInObjectPtr(U&& InPtr)
+		: Pointer(InPtr)
+	{
+	}
+
+	// Implicit compile-time construction from UHT encoded reference
+	[[nodiscard]] constexpr TCompiledInObjectPtr(FCompiledInObjectReference InRef)
+		: CompiledInPtr(InRef)
+	{
+	}
+
+	// Operators to allow this or T* to be used depending on IS_MONOLITHIC
+	operator T*() const
+	{
+		return GetPointer();
+	}
+
+	T* operator->()
+	{
+		return GetPointer();
+	}
+
+	const T* operator->() const
+	{
+		return GetPointer();
+	}
+
+	T& operator*()
+	{
+		return *GetPointer();
+	}
+
+	const T& operator*() const
+	{
+		return *GetPointer();
+	}
+
+	[[nodiscard]] T* GetPointer() const
+	{
+		check(!IsEncodedRef());
+		return reinterpret_cast<T*>(Pointer);
+	}
+
+	[[nodiscard]] constexpr bool IsEncodedRef() const
+	{
+		return CompiledInPtr.IsEncodedRef();
+	}
+
+	[[nodiscard]] constexpr FCompiledInObjectReference GetDecoded() const
+	{
+		return CompiledInPtr.GetDecoded();
+	}
+	
+	[[nodiscard]] constexpr bool IsUnset() const
+	{
+		return Pointer == nullptr;
+	}
+
+	[[nodiscard]] constexpr explicit operator bool() const 
+	{
+		return !IsUnset();
+	}
+
+	[[nodiscard]] constexpr bool operator==(TYPE_OF_NULLPTR)
+	{
+		return IsUnset();
+	}
+	
+	[[nodiscard]] constexpr bool operator!=(TYPE_OF_NULLPTR)
+	{
+		return !IsUnset();
+	}
+};
+
+#endif // IS_MONOLITHIC
+}
+
+#if !IS_MONOLITHIC
+class FStructBaseChain;
+
+namespace UE::Private
+{
+// Utility to allow both raw and encoded pointers to be wrapped in AsStructBaseChain to handle offsetting to private base class
+[[nodiscard]] consteval UE::CodeGen::ConstInit::TCompiledInObjectPtr<const FStructBaseChain> AsStructBaseChain(UE::CodeGen::ConstInit::FCompiledInObjectReference Ptr)
+{
+	return UE::CodeGen::ConstInit::TCompiledInObjectPtr<const FStructBaseChain>(Ptr);
+}
+}
+
+#endif // !IS_MONOLITHIC
+#endif // UE_WITH_CONSTINIT_UOBJECT
\ No newline at end of file
diff --git a/Engine/Source/Runtime/CoreUObject/Public/UObject/EnumProperty.h b/Engine/Source/Runtime/CoreUObject/Public/UObject/EnumProperty.h
--- a/Engine/Source/Runtime/CoreUObject/Public/UObject/EnumProperty.h
+++ b/Engine/Source/Runtime/CoreUObject/Public/UObject/EnumProperty.h
@@ -11,6 +11,7 @@
 #include "UObject/NameTypes.h"
 #include "UObject/ObjectMacros.h"
 #include "UObject/UnrealType.h"
+#include "UObjectGlobals.h"
 
 class FArchive;
 class FBlake3;
@@ -44,7 +45,7 @@ public:
 		UE::CodeGen::ConstInit::FPropertyParams InBaseParams,
 		FNumericProperty* InUnderlyingProperty,
 		int32 InUnderlyingSize,
-		UEnum* InEnum
+		UE::CodeGen::ConstInit::TObjectRefParam<UEnum> InEnum
 	)
 		: Super(InBaseParams.SetElementSize(InUnderlyingSize))
 		, UnderlyingProp(InUnderlyingProperty)
@@ -63,6 +64,9 @@ public:
 	COREUOBJECT_API virtual void Serialize( FArchive& Ar ) override;
 	COREUOBJECT_API virtual void AddReferencedObjects(FReferenceCollector& Collector) override;
 	COREUOBJECT_API virtual void GetPreloadDependencies(TArray<UObject*>& OutDeps) override;
+#if UE_WITH_CONSTINIT_UOBJECT && !IS_MONOLITHIC
+	COREUOBJECT_API virtual void LinkCompiledInPointerFields(UE::CoreUObject::Private::FCompiledInObjectLinker& Linker);
+#endif // UE_WITH_CONSTINIT_UOBJECT && !IS_MONOLITHIC
 	// End of UObject interface
 
 	// Field interface
diff --git a/Engine/Source/Runtime/CoreUObject/Public/UObject/Field.h b/Engine/Source/Runtime/CoreUObject/Public/UObject/Field.h
--- a/Engine/Source/Runtime/CoreUObject/Public/UObject/Field.h
+++ b/Engine/Source/Runtime/CoreUObject/Public/UObject/Field.h
@@ -6,6 +6,7 @@ Field.h: Declares FField property system fundamentals
 
 #pragma once
 
+#include "Concepts/DerivedFrom.h"
 #include "Containers/Array.h"
 #include "Containers/LinkedListBuilder.h"
 #include "Containers/Map.h"
@@ -97,11 +98,11 @@ public:
 	static COREUOBJECT_API TMap<FName, FFieldClass*>& GetNameToFieldClassMap();
 
 	COREUOBJECT_API explicit FFieldClass(const TCHAR* InCPPName, uint64 InId, EClassCastFlags InCastFlags, FFieldClass* InSuperClass, FConstructFunction* InConstructFn);
-	COREUOBJECT_API ~FFieldClass();
+	COREUOBJECT_API virtual ~FFieldClass();
 
+#if UE_WITH_CONSTINIT_UOBJECT
 	consteval explicit FFieldClass(
 		EConstEval,
-		const TCHAR* InCPPName,
 		uint64 InId,
 		EClassCastFlags InCastFlags,
 		FFieldClass* InSuperClass,
@@ -116,6 +117,7 @@ public:
 		, ConstructFn(InConstructFn)
 	{
 	}
+#endif // UE_WITH_CONSTINIT_UOBJECT 
 
 #if UE_WITH_CONSTINIT_UOBJECT
 	bool IsFullyConstructed() const
@@ -198,6 +200,11 @@ public:
 	}
 	COREUOBJECT_API friend FArchive& operator << (FArchive& Ar, FFieldClass*& InOutFieldClass);
 
+protected:
+#if UE_WITH_CONSTINIT_UOBJECT
+	virtual FFieldClass* GetSuperStaticClass() = 0;
+#endif // UE_WITH_CONSTINIT_UOBJECT
+
 private:
 	bool IsChildOf_Walk(const FFieldClass* InBaseClass) const
 	{
@@ -216,6 +223,45 @@ private:
 #endif // UE_WITH_CONSTINIT_UOBJECT
 };
 
+/** 
+ * Helper subclass to initialize an FFieldClass singleton for a particular subtype of FField 
+ * e.g. helps with linking FFieldClass to its super for modular constinit-uobject builds.
+ * Some of this code could be removed if UHT generated some of the code for FField/FFieldClass
+ */
+#if UE_WITH_CONSTINIT_UOBJECT
+template<typename FieldType /* Type derived from FField */>
+class TFieldClass : public FFieldClass
+{
+public: 
+	consteval TFieldClass()
+		: FFieldClass(ConstEval, 
+			FieldType::StaticClassCastFlagsPrivate(),
+			FieldType::StaticClassCastFlags(),
+#if IS_MONOLITHIC
+			&FieldType::ConstInitClass,
+#else
+			nullptr,
+#endif // IS_MONOLITHIC
+			&FieldType::Construct
+		)
+	{
+	}
+
+protected:
+	FFieldClass* GetSuperStaticClass() override final
+	{
+		if constexpr (std::is_same_v<typename FieldType::Super, typename FieldType::ThisClass>)
+		{
+			return nullptr;
+		}
+		else
+		{
+			return FieldType::Super::StaticClass();
+		}
+	}
+};
+#endif // UE_WITH_CONSTINIT_UOBJECT
+
 #if UE_WITH_CONSTINIT_UOBJECT
 struct FRegisterFieldClass
 {
@@ -235,19 +281,14 @@ struct FRegisterFieldClass
 #endif
 
 #if UE_WITH_CONSTINIT_UOBJECT
-#define UE_DEFINE_GET_FIELD_CLASS_PRIVATE virtual FFieldClass* GetFieldClassPrivate() override { return &ConstInitClass; }
-#define UE_DECLARE_FIELD_CLASS_STATIC_CONSTINIT static constinit TNoDestroy<FFieldClass> ConstInitClass;
-#define UE_DEFINE_FIELD_CLASS_STATIC_CONSTINIT(TClass)                                               \
-	constinit TNoDestroy<FFieldClass> TClass::ConstInitClass{ NoDestroyConstEval,                    \
-		                                                      ConstEval,                             \
-		                                                      TEXT(UE_STRINGIZE(TClass)),            \
-		                                                      TClass::StaticClassCastFlagsPrivate(), \
-		                                                      TClass::StaticClassCastFlags(),        \
-		                                                      &TClass::Super::ConstInitClass,        \
-		                                                      &TClass::Construct };
+#define UE_DECLARE_GET_FIELD_CLASS_PRIVATE(TRequiredAPI) 				TRequiredAPI virtual FFieldClass* GetFieldClassPrivate() override;
+#define UE_DEFINE_GET_FIELD_CLASS_PRIVATE(TClass) 						FFieldClass* TClass::GetFieldClassPrivate() { return &TClass::ConstInitClass; }
+#define UE_DECLARE_FIELD_CLASS_STATIC_CONSTINIT(TClass, TRequiredAPI) 	TRequiredAPI static constinit TNoDestroy<TFieldClass<TClass>> ConstInitClass;
+#define UE_DEFINE_FIELD_CLASS_STATIC_CONSTINIT(TClass)					constinit TNoDestroy<TFieldClass<TClass>> TClass::ConstInitClass;
 #else
-#define UE_DEFINE_GET_FIELD_CLASS_PRIVATE 
-#define UE_DECLARE_FIELD_CLASS_STATIC_CONSTINIT 
+#define UE_DECLARE_GET_FIELD_CLASS_PRIVATE(TRequiredAPI) 
+#define UE_DEFINE_GET_FIELD_CLASS_PRIVATE(TClass) 
+#define UE_DECLARE_FIELD_CLASS_STATIC_CONSTINIT(TClass, TRequiredAPI)
 #define UE_DEFINE_FIELD_CLASS_STATIC_CONSTINIT(...)
 #endif
 
@@ -258,7 +299,7 @@ struct FRegisterFieldClass
 private: \
 	TClass& operator=(TClass&&) = delete;   \
 	TClass& operator=(const TClass&) = delete;   \
-	UE_DEFINE_GET_FIELD_CLASS_PRIVATE \
+	UE_DECLARE_GET_FIELD_CLASS_PRIVATE(TRequiredAPI) \
 public: \
 	using Super = UE_REMOVE_OPTIONAL_PARENS(TSuperClass);\
 	using ThisClass = TClass;\
@@ -267,7 +308,7 @@ public: \
 	{ \
 	} \
 	static TRequiredAPI FFieldClass* StaticClass(); \
-	UE_DECLARE_FIELD_CLASS_STATIC_CONSTINIT \
+	UE_DECLARE_FIELD_CLASS_STATIC_CONSTINIT(TClass, TRequiredAPI) \
 	static TRequiredAPI FField* Construct(const FFieldVariant& InOwner, const FName& InName, EObjectFlags InObjectFlags); \
 	inline static constexpr EClassCastFlags StaticClassCastFlagsPrivate() \
 	{ \
@@ -317,6 +358,7 @@ public: \
 	{                                                                                                     \
 		IMPLEMENT_FIELD_CONSTRUCT_IMPLEMENTATION(TClass)                                                  \
 	}                                                                                                     \
+	UE_DEFINE_GET_FIELD_CLASS_PRIVATE(TClass) 														      \
 	UE_DEFINE_FIELD_CLASS_STATIC_CONSTINIT(TClass)                                                        \
 	static FRegisterFieldClass Register_FieldClass_##TClass(                                              \
 		&TClass::ConstInitClass, TEXT(UE_STRINGIZE(TClass)));                                             \
@@ -579,7 +621,7 @@ public:
 
 	static COREUOBJECT_API FFieldClass* StaticClass();
 #if UE_WITH_CONSTINIT_UOBJECT
-	static constinit TNoDestroy<FFieldClass> ConstInitClass;
+	static constinit TNoDestroy<TFieldClass<FField>> ConstInitClass;
 #endif // UE_WITH_CONSTINIT_UOBJECT
 
 	virtual SIZE_T GetFieldSize() const
@@ -617,7 +659,7 @@ public:
 	COREUOBJECT_API virtual ~FField();
 
 #if UE_WITH_CONSTINIT_UOBJECT
-	explicit consteval FField(EConstEval, FFieldClass* InFieldClass, UE::CodeGen::ConstInit::FFieldOwnerParam InOwner, FField* InNext, EObjectFlags InFlags, const UTF8CHAR* InNameUTF8)
+	explicit consteval FField(EConstEval, UE::CodeGen::ConstInit::FFieldOwnerParam InOwner, FField* InNext, EObjectFlags InFlags, const UTF8CHAR* InNameUTF8)
 	: NameTempUTF8(InNameUTF8)
 	, Owner(InOwner.bIsObject ? FFieldVariant(ConstEval, InOwner.Value.Object) : FFieldVariant(ConstEval, InOwner.Value.Field))
 	, Next(InNext)
@@ -642,6 +684,9 @@ public:
 	COREUOBJECT_API bool IsIn(const UObject* InOwner) const;
 	COREUOBJECT_API bool IsIn(const FField* InOwner) const;
 	COREUOBJECT_API FLinkerLoad* GetLinker() const;	
+#if UE_WITH_CONSTINIT_UOBJECT && !IS_MONOLITHIC
+	COREUOBJECT_API virtual void LinkCompiledInPointerFields(UE::CoreUObject::Private::FCompiledInObjectLinker& Linker);
+#endif // UE_WITH_CONSTINIT_UOBJECT && !IS_MONOLITHIC
 	// End UObject interface
 
 	// Begin UField interface.
diff --git a/Engine/Source/Runtime/CoreUObject/Public/UObject/FieldPathProperty.h b/Engine/Source/Runtime/CoreUObject/Public/UObject/FieldPathProperty.h
--- a/Engine/Source/Runtime/CoreUObject/Public/UObject/FieldPathProperty.h
+++ b/Engine/Source/Runtime/CoreUObject/Public/UObject/FieldPathProperty.h
@@ -40,21 +40,38 @@ public:
 	COREUOBJECT_API FFieldPathProperty(FFieldVariant InOwner, const UECodeGen_Private::FFieldPathPropertyParams& Prop);
 
 #if UE_WITH_CONSTINIT_UOBJECT
-	consteval FFieldPathProperty(
-		UE::CodeGen::ConstInit::FPropertyParams InBaseParams,
-		FFieldClass* InPropertyClass
-	)
+	consteval FFieldPathProperty(UE::CodeGen::ConstInit::FPropertyParams InBaseParams, FFieldClass* InPropertyClass)
 		: Super(InBaseParams)
 		, PropertyClass(InPropertyClass)
 	{
 	}
-#endif
+#if !IS_MONOLITHIC
+	consteval FFieldPathProperty(UE::CodeGen::ConstInit::FPropertyParams InBaseParams, const UTF8CHAR* InPropertyClass)
+		: Super(InBaseParams)
+		, PropertyClassName(InPropertyClass)
+	{
+	}
+#endif // !IS_MONOLITHIC
+#endif // UE_WITH_CONSTINIT_UOBJECT
 
 #if WITH_EDITORONLY_DATA
 	COREUOBJECT_API explicit FFieldPathProperty(UField* InField);
 #endif // WITH_EDITORONLY_DATA
 
-	FFieldClass* PropertyClass;
+	union
+	{
+		FFieldClass* PropertyClass = nullptr;
+#if UE_WITH_CONSTINIT_UOBJECT && !IS_MONOLITHIC
+		// Until LinkCompiledInPointerFields holds the name of the FFieldClass that should be assigned to PropertyClass.
+		// TODO: This could be linked via an encoded pointer like objects pointers, but UHT does not parse FFieldClass.
+		// Also note that this pointer cannot be distinguished from PropertyClass unlike other fields handled by LinkCompiledInPointerFields
+		const UTF8CHAR* PropertyClassName;
+#endif // UE_WITH_CONSTINIT_UOBJECT && !IS_MONOLITHIC
+	};
+
+#if UE_WITH_CONSTINIT_UOBJECT && !IS_MONOLITHIC
+	COREUOBJECT_API virtual void LinkCompiledInPointerFields(UE::CoreUObject::Private::FCompiledInObjectLinker& Linker);
+#endif // UE_WITH_CONSTINIT_UOBJECT && !IS_MONOLITHIC
 
 	// UHT interface
 	COREUOBJECT_API virtual FString GetCPPMacroType(FString& ExtendedTypeText) const  override;
diff --git a/Engine/Source/Runtime/CoreUObject/Public/UObject/IntrinsicClassBaseChain.h b/Engine/Source/Runtime/CoreUObject/Public/UObject/IntrinsicClassBaseChain.h
--- a/Engine/Source/Runtime/CoreUObject/Public/UObject/IntrinsicClassBaseChain.h
+++ b/Engine/Source/Runtime/CoreUObject/Public/UObject/IntrinsicClassBaseChain.h
@@ -5,6 +5,7 @@
 #if defined(UE_WITH_CONSTINIT_UOBJECT) && UE_WITH_CONSTINIT_UOBJECT
 
 #include "UObject/Class.h"
+#include "UObject/CompiledInObjectPtr.h"
 
 namespace UE::Private
 {
@@ -21,11 +22,11 @@ namespace UE::Private
 	template<typename... Classes>
 	struct TIntrinsicClassBaseChainImpl
 	{
-		using BaseChainPtr = const FStructBaseChain*;
+		using BaseChainPtr = CodeGen::ConstInit::TCompiledInObjectPtr<const FStructBaseChain>;
 		static constexpr SIZE_T Num = sizeof...(Classes);
-		static inline constinit BaseChainPtr Bases[] = { static_cast<BaseChainPtr>(GetConstInitClass<Classes>())... };
+		static inline constinit BaseChainPtr Bases[] = { BaseChainPtr{UE::Private::AsStructBaseChain((GetConstInitClass<Classes>()))}... };
 
-		static consteval TArrayView<const FStructBaseChain*> ToArrayView()
+		static consteval TArrayView<BaseChainPtr> ToArrayView()
 		{
 			return MakeArrayView(Bases, Num);
 		}
diff --git a/Engine/Source/Runtime/CoreUObject/Public/UObject/IntrinsicClassFwd.h b/Engine/Source/Runtime/CoreUObject/Public/UObject/IntrinsicClassFwd.h
--- a/Engine/Source/Runtime/CoreUObject/Public/UObject/IntrinsicClassFwd.h
+++ b/Engine/Source/Runtime/CoreUObject/Public/UObject/IntrinsicClassFwd.h
@@ -4,6 +4,7 @@
 
 #include "UObject/ObjectMacros.h"
 
+// TODO: These need to be emitted by UHT if we are going to support intrinsic classes outside of CoreUObject in modular constinit-uobject builds 
 #if UE_WITH_CONSTINIT_UOBJECT
 
 #include "Templates/NoDestroy.h"
@@ -12,7 +13,7 @@
 
 #define UE_DEFINE_INTRINSIC_CLASS_CONSTINIT_ACCESSOR(ClassName, RequiredAPI) \
     class ClassName; \
-    extern RequiredAPI constinit TNoDestroy<UClass> Z_ConstInit_UClass_##ClassName; \
+    extern constinit TNoDestroy<UClass> Z_ConstInit_UClass_##ClassName; \
     namespace UE::Private \
     { \
         template<> \
diff --git a/Engine/Source/Runtime/CoreUObject/Public/UObject/ObjectHandle.h b/Engine/Source/Runtime/CoreUObject/Public/UObject/ObjectHandle.h
--- a/Engine/Source/Runtime/CoreUObject/Public/UObject/ObjectHandle.h
+++ b/Engine/Source/Runtime/CoreUObject/Public/UObject/ObjectHandle.h
@@ -10,6 +10,7 @@
 #include "Misc/AssertionMacros.h"
 #include "Misc/UEOps.h"
 #include "Templates/TypeHash.h"
+#include "UObject/CompiledInObjectPtr.h"
 #include "UObject/NameTypes.h"
 #include "UObject/ObjectHandleTracking.h"
 #include "UObject/ObjectMacros.h"
@@ -24,7 +25,7 @@ class UPackage;
  * FObjectHandle is either a packed object ref or the resolved pointer to an object.  Depending on configuration
  * when you create a handle, it may immediately be resolved to a pointer.
  */
-#if UE_WITH_OBJECT_HANDLE_LATE_RESOLVE
+#if UE_WITH_OBJECT_HANDLE_LATE_RESOLVE || (UE_WITH_CONSTINIT_UOBJECT && !IS_MONOLITHIC)
 
 namespace UE::CoreUObject::Private
 {
@@ -56,7 +57,7 @@ inline bool IsObjectHandleTypeSafe(FObjectHandle Handle);
 //Private functions that forced public due to inlining.
 namespace UE::CoreUObject::Private
 {
-#if UE_WITH_OBJECT_HANDLE_LATE_RESOLVE
+#if UE_WITH_OBJECT_HANDLE_LATE_RESOLVE || (UE_WITH_CONSTINIT_UOBJECT && !IS_MONOLITHIC)
 
 	struct FObjectHandlePrivate
 	{
@@ -64,10 +65,16 @@ namespace UE::CoreUObject::Private
 		// Pointer union member is for constinit initialization where we want the linker to be able to do relocations for us within a binary.
 		union 
 		{
+#if UE_WITH_OBJECT_HANDLE_LATE_RESOLVE
 			TSAN_ATOMIC(UPTRINT) PointerOrRef;
+#endif
+#if UE_WITH_CONSTINIT_UOBJECT && !IS_MONOLITHIC
+			UE::CodeGen::ConstInit::FCompiledInObjectPtr CompiledInPtr;
+#endif
 			const void* Pointer;
 		};
 
+#if UE_WITH_OBJECT_HANDLE_LATE_RESOLVE
 		UE_FORCEINLINE_HINT UPTRINT GetPointerOrRef() const
 		{
 			return PointerOrRef;
@@ -77,11 +84,17 @@ namespace UE::CoreUObject::Private
 		{
 			return PointerOrRef != 0;
 		}
+#else 
+		explicit inline operator bool() const
+		{
+			return Pointer != nullptr;
+		}
+#endif
 
 		// When using TSAN, PointerOrRef has non-trivial copy construction so we need to explicitly implement constructors to handle the union.
 		// This means this type can't use aggregate initialization even when TSAN is disabled
 		FObjectHandlePrivate()
-		: PointerOrRef(0)
+		: Pointer(nullptr)
 		{
 		}
 		~FObjectHandlePrivate() = default;
@@ -89,11 +102,19 @@ namespace UE::CoreUObject::Private
 			: Pointer(InPointer)
 		{
 		}
+#if UE_WITH_OBJECT_HANDLE_LATE_RESOLVE
 		[[nodiscard]] explicit FObjectHandlePrivate(UPTRINT Packed)
 			: PointerOrRef(Packed)
 		{
 		}
-#if USING_THREAD_SANITISER || USING_INSTRUMENTATION
+#endif
+#if UE_WITH_CONSTINIT_UOBJECT && !IS_MONOLITHIC
+		[[nodiscard]] explicit consteval FObjectHandlePrivate(UE::CodeGen::ConstInit::FCompiledInObjectPtr InCompiledInPtr)
+			: CompiledInPtr(InCompiledInPtr)
+		{
+		}
+#endif // UE_WITH_CONSTINIT_UOBJECT && !IS_MONOLITHIC
+#if UE_WITH_OBJECT_HANDLE_LATE_RESOLVE && (USING_THREAD_SANITISER || USING_INSTRUMENTATION)
 		[[nodiscard]] constexpr FObjectHandlePrivate(const FObjectHandlePrivate& Other)
 		{
 			UE_IF_CONSTEVAL
@@ -128,11 +149,13 @@ namespace UE::CoreUObject::Private
 	// Allow generated operators to be accessible from comparisons in other namespaces.
 	UE_OPS_NAMESPACE_VISIBLE(FObjectHandlePrivate)
 
+#if UE_WITH_OBJECT_HANDLE_LATE_RESOLVE
 	/* Returns the packed object ref for this object IF one exists otherwise returns a null PackedObjectRef */
 	COREUOBJECT_API FPackedObjectRef FindExistingPackedObjectRef(const UObject* Object);
 
 	/* Creates and ObjectRef from a packed object ref*/
 	COREUOBJECT_API FObjectRef MakeObjectRef(FPackedObjectRef Handle);
+#endif // UE_WITH_OBJECT_HANDLE_LATE_RESOLVE
 
 
 #elif UE_WITH_REMOTE_OBJECT_HANDLE
@@ -196,6 +219,16 @@ namespace UE::CoreUObject::Private
 
 	/** Read the handle as a pointer without checking if it is resolved. Invalid to call for unresolved handles. */
 	inline UObject* ReadObjectHandlePointerNoCheck(FObjectHandle Handle);
+
+#if UE_WITH_CONSTINIT_UOBJECT && !IS_MONOLITHIC
+	/** 
+	 * Makes an object handle from generated parameters from UHT.
+	 * If Ptr is non-null, makes a direct object ptr that can be linked by the linker.
+	 * Otherwise if CompiledInPtr is an encoded reference, makes an encoded pointer that can be linked at runtime.
+	 * Otherwise, makes a null handle.
+	 */
+	inline consteval FObjectHandle MakeCompiledInObjectHandle(void* Ptr, UE::CodeGen::ConstInit::FCompiledInObjectPtr CompiledInPtr);
+#endif // UE_WITH_CONSTINIT_UOBJECT && !IS_MONOLITHIC
 }
 
 /////////////////////////////////////////////////////////////////////////////////////////////////////////////////////////
@@ -457,6 +490,24 @@ namespace UE::CoreUObject::Private
 	}
 #endif
 
+#if UE_WITH_CONSTINIT_UOBJECT && !IS_MONOLITHIC
+	inline consteval FObjectHandle MakeCompiledInObjectHandle(void* Ptr, UE::CodeGen::ConstInit::FCompiledInObjectPtr CompiledInPtr)
+	{
+		if (Ptr)
+		{
+			return FObjectHandle(Ptr);
+		}
+		else if (!CompiledInPtr.IsUnset())
+		{
+			return FObjectHandle(CompiledInPtr);
+		}
+		else 
+		{
+			return FObjectHandle(nullptr);
+		}
+	}
+#endif // UE_WITH_CONSTINIT_UOBJECT && !IS_MONOLITHIC
+
 	inline UObject* ResolveObjectHandle(FObjectHandle& Handle)
 	{
 #if UE_WITH_OBJECT_HANDLE_LATE_RESOLVE || UE_WITH_OBJECT_HANDLE_TRACKING
diff --git a/Engine/Source/Runtime/CoreUObject/Public/UObject/ObjectMacros.h b/Engine/Source/Runtime/CoreUObject/Public/UObject/ObjectMacros.h
--- a/Engine/Source/Runtime/CoreUObject/Public/UObject/ObjectMacros.h
+++ b/Engine/Source/Runtime/CoreUObject/Public/UObject/ObjectMacros.h
@@ -1874,6 +1874,7 @@ namespace UM
 #define DECLARE_CLASS2( TClass, TSuperClass, TStaticFlags, TStaticCastFlags, TPackage, TPrivateAccessor ) \
 private: \
 	friend struct FUObjectCppClassStaticFunctions; \
+	friend struct TUObjectCppClassStaticThunks<TClass>; \
 public: \
 	/** Bitwise union of #EClassFlags pertaining to this class.*/ \
 	static constexpr EClassFlags StaticClassFlags=EClassFlags(TStaticFlags); \
@@ -2086,6 +2087,35 @@ public: \
 		UOBJECT_CPPCLASS_STATICFUNCTIONS_WITHEDITORONLYDATA(TClass) \
 	)
 
+// Used to allow constinit initialization UClass variables outside of CoreUObject in modular 
+// builds to take pointers to these methods rather than the methods of UObject, which can't 
+// be used in constant expressions in a different dll.
+template<typename T>
+struct TUObjectCppClassStaticThunks
+{
+#if UE_WITH_CONSTINIT_UOBJECT && !IS_MONOLITHIC
+	static void AddReferencedObjects(UObject* ThisObject, FReferenceCollector& Ar)
+	{
+		T::AddReferencedObjects(ThisObject, Ar);
+	}
+
+#if WITH_EDITORONLY_DATA
+	static void DeclareCustomVersions(FArchive& Ar, const UClass* SpecificSubclass)
+	{
+		T::DeclareCustomVersions(Ar, SpecificSubclass);
+	}
+	static void AppendToClassSchema(FAppendToClassSchemaContext& Context)
+	{
+		T::AppendToClassSchema(Context);
+	}
+	static void DeclareConstructClasses(TArray<FTopLevelAssetPath>& OutConstructClasses, const UClass* SpecificSubclass)
+	{
+		T::DeclareConstructClasses(OutConstructClasses, SpecificSubclass);
+	}
+#endif // WITH_EDITORONLY_DATA
+#endif // UE_WITH_CONSTINIT_UOBJECT && !IS_MONOLITHIC
+};
+
 /**
  * Collection of the pointers to our specified list of static functions that are defined on a specific c++ class,
  * for reference by the corresponding UClass. The pointers in this structure point to the class's version of that
@@ -2103,12 +2133,23 @@ public:
 #endif
 	// UObjectCppClassStaticFunctions: Extend this list of types with the type of your new static function.
 
-	// Take a null ptr of the class for template type deduction, then this class can be declared a friend of generated classes to access their member function pointers
+#if UE_WITH_CONSTINIT_UOBJECT
 	template<typename TClass>
-	explicit consteval FUObjectCppClassStaticFunctions(TClass* Null)
-	: UOBJECT_CPPCLASS_STATICFUNCTIONS_FORCLASS(TClass)
+	static constexpr FUObjectCppClassStaticFunctions ForClass()
 	{
+#if IS_MONOLITHIC
+		using DispatchType = TClass;
+#else
+		using DispatchType = TUObjectCppClassStaticThunks<TClass>;
+#endif // IS_MONOLITHIC
+		return FUObjectCppClassStaticFunctions(
+			&DispatchType::AddReferencedObjects
+	#if WITH_EDITORONLY_DATA
+			, &DispatchType::DeclareCustomVersions, &DispatchType::AppendToClassSchema, &DispatchType::DeclareConstructClasses
+	#endif
+		);
 	}
+#endif // UE_WITH_CONSTINIT_UOBJECT
 
 	explicit constexpr FUObjectCppClassStaticFunctions(AddReferencedObjectsType InAddReferencedObjects
 #if WITH_EDITORONLY_DATA
@@ -2136,6 +2177,8 @@ public:
 	}
 public:
 	FUObjectCppClassStaticFunctions() = default;
+	constexpr FUObjectCppClassStaticFunctions(const FUObjectCppClassStaticFunctions&) = default;
+	FUObjectCppClassStaticFunctions& operator=(const FUObjectCppClassStaticFunctions&) = default;
 	bool IsInitialized() const
 	{
 		// All methods that set the functions in this class guarantee that if any function is non-null, they all are.
@@ -2320,8 +2363,8 @@ private:
 			InitCode; \
 		} \
 	}; \
-	static FRegisterCompiledInObjects Z_Register_UClass_##TClass{ &Z_ConstInit_UClass_##TClass, &Z_Construct_UClass_##TClass##_Statics::Construct }; \
-	TRequiredAPI UClass* TClass::GetPrivateStaticClass() \
+	static FRegisterIntrinsicClass Z_Register_UClass_##TClass{ &Z_ConstInit_UClass_##TClass, &Z_Construct_UClass_##TClass##_Statics::Construct }; \
+	UClass* TClass::GetPrivateStaticClass() \
 	{ \
 		return &Z_ConstInit_UClass_##TClass; \
 	}
diff --git a/Engine/Source/Runtime/CoreUObject/Public/UObject/ObjectPtr.h b/Engine/Source/Runtime/CoreUObject/Public/UObject/ObjectPtr.h
--- a/Engine/Source/Runtime/CoreUObject/Public/UObject/ObjectPtr.h
+++ b/Engine/Source/Runtime/CoreUObject/Public/UObject/ObjectPtr.h
@@ -82,6 +82,18 @@ public:
 	{
 	}
 
+#if UE_WITH_CONSTINIT_UOBJECT
+	template<typename U>
+	[[nodiscard]] explicit consteval FObjectPtr(UE::CodeGen::ConstInit::TObjectRefParam<U> Ptr)
+#if IS_MONOLITHIC
+		: Handle(UE::CoreUObject::Private::MakeObjectHandle(Ptr.ObjectPtr))
+#else
+		: Handle(UE::CoreUObject::Private::MakeCompiledInObjectHandle(Ptr.ObjectPtr, Ptr.CompiledInPtr))
+#endif // IS_MONOLITHIC
+	{
+	}
+#endif // UE_WITH_CONSTINIT_UOBJECT
+
 #if UE_WITH_OBJECT_HANDLE_LATE_RESOLVE || UE_WITH_REMOTE_OBJECT_HANDLE
 	[[nodiscard]] explicit FORCEINLINE FObjectPtr(FObjectHandle Handle)
 		: Handle(Handle)
@@ -538,6 +550,13 @@ public:
 	{
 	}
 
+#if UE_WITH_CONSTINIT_UOBJECT
+	[[nodiscard]] explicit consteval TObjectPtr(UE::CodeGen::ConstInit::TObjectRefParam<T> Ptr)
+		: ObjectPtr(Ptr)
+	{
+	}
+#endif
+
 	template <
 		typename U
 		UE_REQUIRES(std::is_convertible_v<U*, T*>)
@@ -801,6 +820,11 @@ private:
 	FORCEINLINE T* GetNoReadNoCheck() const { return (T*)(UE::CoreUObject::Private::ResolveObjectHandleNoReadNoCheck(ObjectPtr.GetHandleRef())); }
 	FORCEINLINE T* GetNoResolveNoCheck() const { return (T*)(UE::CoreUObject::Private::ReadObjectHandlePointerNoCheck(ObjectPtr.GetHandleRef())); }
 
+	[[nodiscard]] UE_REWRITE FObjectHandle& GetHandleRef() const 
+	{
+		return ObjectPtr.GetHandleRef();
+	}
+
 	// @TODO: OBJPTR: There is a risk of a gap in access tracking here.  The caller may get a mutable pointer, write to it, then
 	//			read from it.  That last read would happen without an access being recorded.  Not sure if there is a good way
 	//			to handle this case without forcing the calling code to be modified.
@@ -842,6 +866,8 @@ namespace ObjectPtr_Private
 	template <typename T> struct TRawPointerType<const          TObjectPtr<T>> { using Type = T*; };
 	template <typename T> struct TRawPointerType<      volatile TObjectPtr<T>> { using Type = T*; };
 	template <typename T> struct TRawPointerType<const volatile TObjectPtr<T>> { using Type = T*; };
+	template <typename T>
+	class TNonAccessTrackedObjectPtr;
 	struct Friend
 	{
 		template <typename T>
@@ -890,6 +916,21 @@ namespace ObjectPtr_Private
 		{
 			return reinterpret_cast<T*>(UE::CoreUObject::Private::NoResolveObjectHandleNoRead(Ptr.ObjectPtr.GetHandleRef()));
 		}
+
+		template<typename T>
+		UE_REWRITE static FObjectHandle& GetHandleRef(FObjectPtr& Ptr)
+		{
+			return Ptr.GetHandleRef();
+		}
+
+		template<typename T>
+		UE_REWRITE static FObjectHandle& GetHandleRef(TObjectPtr<T>& Ptr)
+		{
+			return Ptr.GetHandleRef();
+		}
+		
+		template<typename T>
+		static FObjectHandle& GetHandleRef(TNonAccessTrackedObjectPtr<T>& Ptr);
 	};
 	
 	template <typename T>
@@ -914,6 +955,14 @@ namespace ObjectPtr_Private
 		{
 		}
 
+#if UE_WITH_CONSTINIT_UOBJECT
+		template<typename U>
+		explicit consteval TNonAccessTrackedObjectPtr(UE::CodeGen::ConstInit::TObjectRefParam<U> Ptr)
+			: ObjectPtr{FObjectPtr(Ptr)}
+		{
+		}
+#endif // UE_WITH_CONSTINIT_UOBJECT
+
 		TNonAccessTrackedObjectPtr& operator=(T* Value)
 		{
 			ObjectPtr = Value;
@@ -930,6 +979,11 @@ namespace ObjectPtr_Private
 			return ObjectPtr_Private::Friend::NoAccessTrackingGetNoResolve(ObjectPtr);
 		}
 
+		[[nodiscard]] UE_REWRITE FObjectHandle GetHandle() const 
+		{
+			return ObjectPtr.GetHandle();
+		}
+
 		explicit operator UPTRINT() const
 		{
 			return BitCast<UPTRINT>(Get());
@@ -956,8 +1010,15 @@ namespace ObjectPtr_Private
 		}		
 	
 	private:
+		friend struct Friend;
 		TObjectPtr<T> ObjectPtr;
 	};
+
+	template<typename T>
+	inline FObjectHandle& Friend::GetHandleRef(TNonAccessTrackedObjectPtr<T>& Ptr)
+	{
+		return Ptr.ObjectPtr.GetHandleRef();
+	}
 }
 
 template <typename T>
diff --git a/Engine/Source/Runtime/CoreUObject/Public/UObject/PropertyOptional.h b/Engine/Source/Runtime/CoreUObject/Public/UObject/PropertyOptional.h
--- a/Engine/Source/Runtime/CoreUObject/Public/UObject/PropertyOptional.h
+++ b/Engine/Source/Runtime/CoreUObject/Public/UObject/PropertyOptional.h
@@ -210,6 +210,9 @@ public:
 	// UObject interface
 	UE_API virtual void Serialize(FArchive& Ar) override;
 	UE_API virtual void GetPreloadDependencies(TArray<UObject*>& OutDeps) override;
+#if UE_WITH_CONSTINIT_UOBJECT && !IS_MONOLITHIC
+	COREUOBJECT_API virtual void LinkCompiledInPointerFields(UE::CoreUObject::Private::FCompiledInObjectLinker& Linker);
+#endif // UE_WITH_CONSTINIT_UOBJECT && !IS_MONOLITHIC
 	// End of UObject interface
 
 	// Field interface
diff --git a/Engine/Source/Runtime/CoreUObject/Public/UObject/RegisterCompiledInObjects.h b/Engine/Source/Runtime/CoreUObject/Public/UObject/RegisterCompiledInObjects.h
--- a/Engine/Source/Runtime/CoreUObject/Public/UObject/RegisterCompiledInObjects.h
+++ b/Engine/Source/Runtime/CoreUObject/Public/UObject/RegisterCompiledInObjects.h
@@ -8,6 +8,7 @@
 
 #include "Containers/ArrayView.h"
 #include "Misc/TVariant.h"
+#include "UObject/CompiledInObjectPtr.h"
 #include "UObject/UObjectGlobals.h"
 
 class UClass;
@@ -15,99 +16,84 @@ class UDelegateFunction;
 class UObject;
 class UPackage;
 
-// Register partially initialized objects such as packages, classes, structs and enums for runtime construction
-// This structure is stored as a linked list where possible to avoid dynamic allocation during startup 
-struct FRegisterCompiledInObjects
+// Utility struct for passing many arrays with designated initializer for clarity 
+struct FCompiledInObjects
 {
-	using FIntrinsicClassConstructor = void();
+#if !IS_MONOLITHIC
+	// Tables of objects in other modules for resolving encoding pointers in these objects
+	TConstArrayView<TConstArrayView<UObject*>> LinkedModules;
 
-	// Register an intrinsic class with a special initialization function
-	FRegisterCompiledInObjects(UClass* InClass, FIntrinsicClassConstructor* InConstructor)
-		: FRegisterCompiledInObjects(FIntrinsicClassRegistrant{InClass, InConstructor})
-	{
-	}
+	// Objects which can be linked by other objects using the package's link id and the index of the
+	// object in this array.
+	// The following arrays may bes slices of this array produced by UHT to reduce duplication.
+	// This array and the following arrays may contain null entries if objects are declared inside #if scopes in headers
+	// This means that link indices are stable regardless of the value of those preprocessor definitions
+	TConstArrayView<UObject*> LinkableObjects;
+#endif // !IS_MONOLITHIC
 
-	// Register a package and some objects that might otherwise not be registered - e.g. noexport structs and delegate
-	// function signatures. ArrayView should point to static memory
-	FRegisterCompiledInObjects(UPackage* InPackage, TConstArrayView<UObject*> Others)
-		: FRegisterCompiledInObjects(FPackageRegistrants{InPackage, Others})
-	{
-	}
+	// Packages that are part of this module - usually 1 for UObject modules or 2 for modules with Verse
+	TConstArrayView<UPackage*> Packages;
 
-	// Primary entry point - register a batch of objects from a single generated cpp file all together
-	// ArrayView should point to static memory
-	// Constructors for each possible set of parameters are provided to reduce code size at the call site
-	FRegisterCompiledInObjects(
-		TConstArrayView<UClass*> InClasses,
-		TConstArrayView<UScriptStruct*> InStructs,
-		TConstArrayView<UEnum*> InEnums)
-		: FRegisterCompiledInObjects(FClassRegistrants{InClasses, InStructs, InEnums})
-	{
-	}
-	FRegisterCompiledInObjects(
-		TConstArrayView<UClass*> InClasses)
-		: FRegisterCompiledInObjects(FClassRegistrants{InClasses, {}, {}})
-	{
-	}
-	FRegisterCompiledInObjects(
-		TConstArrayView<UClass*> InClasses,
-		TConstArrayView<UScriptStruct*> InStructs)
-		: FRegisterCompiledInObjects(FClassRegistrants{InClasses, InStructs, {}})
-	{
-	}
-	FRegisterCompiledInObjects(
-		TConstArrayView<UClass*> InClasses,
-		TConstArrayView<UEnum*> InEnums)
-		: FRegisterCompiledInObjects(FClassRegistrants{InClasses, {}, InEnums})
+	// Arrays of top level objects to initialize segmented by type to avoid casting.
+	// Typed as UObject* because they are slices of a single array
+	TConstArrayView<UObject*> Classes;
+	TConstArrayView<UObject*> Structs;
+	TConstArrayView<UObject*> Enums;
+	// Other types of objects (e.g. UDelegateFunction) which can be initialized by casting after
+	// classes are propery initialized.
+	TConstArrayView<UObject*> Others;
+
+	inline TConstArrayView<UClass*> GetClasses() const
 	{
+		return MakeArrayView(reinterpret_cast<UClass* const*>(Classes.GetData()), Classes.Num());
 	}
-	FRegisterCompiledInObjects(
-		TConstArrayView<UScriptStruct*> InStructs,
-		TConstArrayView<UEnum*> InEnums)
-		: FRegisterCompiledInObjects(FClassRegistrants{{}, InStructs, InEnums})
+	inline TConstArrayView<UScriptStruct*> GetStructs() const
 	{
+		return MakeArrayView(reinterpret_cast<UScriptStruct* const*>(Structs.GetData()), Structs.Num());
 	}
-	FRegisterCompiledInObjects(
-		TConstArrayView<UScriptStruct*> InStructs)
-		: FRegisterCompiledInObjects(FClassRegistrants{{}, InStructs, {}})
+	inline TConstArrayView<UEnum*> GetEnums() const
 	{
+		return MakeArrayView(reinterpret_cast<UEnum* const*>(Enums.GetData()), Enums.Num());
 	}
-	FRegisterCompiledInObjects(
-		TConstArrayView<UEnum*> InEnums)
-		: FRegisterCompiledInObjects(FClassRegistrants{{}, {}, InEnums})
+};
+
+// Register partially initialized objects such as packages, classes, structs and enums for runtime construction
+// This structure is stored as a linked list where possible to avoid dynamic allocation during startup 
+// UHT emits one of these per module
+struct FRegisterCompiledInObjects
+{
+	// Register a package and all the objects which need to be registered directly.
+	// Other objects may be registered as subobjects, e.g. Functions
+	[[nodiscard]] explicit FRegisterCompiledInObjects(FCompiledInObjects InObjects)
+		: Objects(InObjects)
 	{
+		Register();
 	}
 
-	struct FPackageRegistrants
-	{
-		UPackage* Package = nullptr;
-		TConstArrayView<UObject*> Objects;
-	};
+	FCompiledInObjects Objects;
 
-	struct FClassRegistrants
-	{
-		TConstArrayView<UClass*> Classes;
-		TConstArrayView<UScriptStruct*> Structs;
-		TConstArrayView<UEnum*> Enums;
-	};
+	FRegisterCompiledInObjects* ListNext = nullptr;
 
-	struct FIntrinsicClassRegistrant
-	{
-		UClass* IntrinsicClass = nullptr;
-		FIntrinsicClassConstructor* Constructor = nullptr;
-	};
+	// Register all our objects for deferred construction
+	COREUOBJECT_API void Register();
+};
 
-	TVariant<FPackageRegistrants, FClassRegistrants, FIntrinsicClassRegistrant> Registrants;
-	FRegisterCompiledInObjects* ListNext = nullptr;
-	
-	// Delegating constructor
-	template<typename VARIANT>
-	FRegisterCompiledInObjects(VARIANT&& Variant)
-	: Registrants(TInPlaceType<VARIANT>{}, Forward<VARIANT>(Variant))
+struct FRegisterIntrinsicClass
+{
+	using FIntrinsicClassConstructor = void();
+
+	[[nodiscard]] explicit FRegisterIntrinsicClass(UClass* InClass, FIntrinsicClassConstructor* InConstructor)
+		: Class(InClass)
+		, Constructor(InConstructor)
 	{
 		Register();
 	}
 
+	UClass* Class = nullptr;
+	FIntrinsicClassConstructor* Constructor = nullptr;
+
+	FRegisterIntrinsicClass* ListNext = nullptr;
+
 	// Register all our objects for deferred construction
 	COREUOBJECT_API void Register();
 };
diff --git a/Engine/Source/Runtime/CoreUObject/Public/UObject/UObjectBase.h b/Engine/Source/Runtime/CoreUObject/Public/UObject/UObjectBase.h
--- a/Engine/Source/Runtime/CoreUObject/Public/UObject/UObjectBase.h
+++ b/Engine/Source/Runtime/CoreUObject/Public/UObject/UObjectBase.h
@@ -32,6 +32,8 @@ class UScriptStruct;
 namespace UE::CoreUObject::Private
 {
 	struct FObjectHandleUtils;
+
+	class FCompiledInObjectLinker;
 }
 
 // IWYU pragma: end_exports
@@ -102,7 +104,7 @@ public:
 	explicit consteval UObjectBase(UE::CodeGen::ConstInit::FObjectParams InParams)
 		: ObjectFlags(InParams.Flags | RF_MarkAsNative | RF_MarkAsRootSet | RF_NeedInitialization)
 		, InternalIndex(INDEX_NONE)
-		, ClassPrivate(ConstEval, InParams.Class)
+		, ClassPrivate(InParams.Class)
 		, UninitializedNameUTF8(InParams.NameUTF8)
 		, OuterPrivate(InParams.Outer)
 	{
@@ -464,6 +466,12 @@ private:
 	/** This is used by the reinstancer to re-class and re-archetype the current instances of a class before recompiling */
 	COREUOBJECT_API void SetClass(UClass* NewClass);
 #endif
+
+protected:
+#if UE_WITH_CONSTINIT_UOBJECT && !IS_MONOLITHIC
+	// Overridden for types generated by UHT which may contain encoded references to compiled-in objects in other modules
+	COREUOBJECT_API virtual void LinkCompiledInPointerFields(UE::CoreUObject::Private::FCompiledInObjectLinker& Linker);
+#endif // UE_WITH_CONSTINIT_UOBJECT && !IS_MONOLITHIC
 };
 
 /**
diff --git a/Engine/Source/Runtime/CoreUObject/Public/UObject/UObjectGlobals.h b/Engine/Source/Runtime/CoreUObject/Public/UObject/UObjectGlobals.h
--- a/Engine/Source/Runtime/CoreUObject/Public/UObject/UObjectGlobals.h
+++ b/Engine/Source/Runtime/CoreUObject/Public/UObject/UObjectGlobals.h
@@ -37,6 +37,7 @@
 #include "Templates/UnrealTemplate.h"
 #include "Templates/IsTObjectPtr.h"
 #include "Traits/IsCharEncodingCompatibleWith.h"
+#include "UObject/CompiledInObjectPtr.h"
 #include "UObject/FindObjectFlags.h"
 #include "UObject/NameTypes.h"
 #include "UObject/ObjectMacros.h"
@@ -4267,15 +4268,23 @@ struct FVerseCallableThunk;
 
 namespace UE::CodeGen::ConstInit
 {
+	// Special tags for constructing classes such as UClass that are self-referential
+	// i.e. the class of class is class
+	enum class EClass { Class };
+	enum class EObject { Object };
+
 #if WITH_METADATA
 	struct FMetaData
 	{
-		const UTF8CHAR* NameUTF8;
-		const UTF8CHAR* ValueUTF8;
+		// NOTE: These should be changed to UTF8CHAR when UHT is modified to emit \u and \U
+		// for unicode characters rather than \x
+		const char* NameUTF8;
+		const char* ValueUTF8;
 	};
 #endif
 
 	// Parameter wrapper for FField* or UObject* which can be used as a parameter in a constexpr context
+	// FFieldOwner cannot be used directly because we can't use the tagged pointer/union code in constexpr
 	struct FFieldOwnerParam
 	{
 		union FValue
@@ -4304,10 +4313,68 @@ namespace UE::CodeGen::ConstInit
 		}
 	};
 
+	// Parameter wrapper for either an object pointer or an encoded reference to an object in another module to be linked
+	// TCompiledInObjectPtr cannot be used directly because we can't use the tagged pointer/union code in constexpr
+	template<typename T>
+	struct TObjectRefParam
+	{
+		T* ObjectPtr = nullptr;
+#if !IS_MONOLITHIC
+		FCompiledInObjectPtr CompiledInPtr;
+#endif // !IS_MONOLITHIC
+
+		constexpr TObjectRefParam() = default;
+		constexpr TObjectRefParam(const TObjectRefParam&) = default;
+		constexpr TObjectRefParam& operator=(const TObjectRefParam&) = default;
+
+		// These constructors are not explicit to simplify UHT generated code 
+		constexpr TObjectRefParam(T* InPtr)
+			: ObjectPtr(InPtr)
+		{
+		}
+
+#if IS_MONOLITHIC
+		constexpr bool IsUnset() const
+		{
+			return ObjectPtr == nullptr;
+		}
+#else // IS_MONOLITHIC
+		constexpr TObjectRefParam(uint32 InModuleId, uint32 InObjectId)
+			: CompiledInPtr(FCompiledInObjectReference{InModuleId, InObjectId})
+		{
+		}
+
+		constexpr TObjectRefParam(FCompiledInObjectReference InRef)
+			: CompiledInPtr(InRef)
+		{
+		}
+		constexpr bool IsUnset() const
+		{
+			return ObjectPtr == nullptr && CompiledInPtr.IsUnset();
+		}
+#endif // IS_MONOLITHIC
+
+		constexpr explicit operator bool() const 
+		{
+			return !IsUnset();
+		}
+
+		constexpr bool operator==(TYPE_OF_NULLPTR)
+		{
+			return IsUnset();
+		}
+		
+		constexpr bool operator!=(TYPE_OF_NULLPTR)
+		{
+			return !IsUnset();
+		}
+	};
+
+	using FObjectRefParam = TObjectRefParam<UObject>;
+
 	// Shared parameters for all properties for constant-initialization
 	struct FPropertyParams
 	{
-		FFieldClass*		FieldClass;
 		FFieldOwnerParam 	Owner;
 		FProperty* 	    	NextProperty; // Next property in this exact struct, no link to inherited properties
 		const UTF8CHAR* NameUTF8;
@@ -4338,13 +4405,13 @@ namespace UE::CodeGen::ConstInit
 		EObjectFlags Flags 		= RF_NoFlags;
 		// UObject* rather than UClass* because consteval constructors must be inlined and TObjectPtr
 		// cannot see inheritance hierarchy in definition of UObjectBase
-		UObject* Class 			= nullptr; 
+		FObjectRefParam Class;
 		const UTF8CHAR* NameUTF8 = nullptr;
-		UObject* Outer 			= nullptr;
+		FObjectRefParam Outer;
 
-		consteval FObjectParams& SetClassIfUnset(UObject* InClass)
+		consteval FObjectParams& SetClass(UObject* InClass)
 		{ 
-			Class = InClass;
+			Class.ObjectPtr = InClass;
 			return *this;
 		}
 	};
@@ -4362,14 +4429,14 @@ namespace UE::CodeGen::ConstInit
 
 	struct FUFieldParams
 	{
-		UField* Next;
+		TObjectRefParam<UField> Next;
 	};
 
 	struct FStructParams
 	{
-		UStruct* Super;
-		UField* FirstChild;
-		TArrayView<const FStructBaseChain*> BaseChain;
+		TObjectRefParam<UStruct> Super;
+		UField* FirstChild; // No cross-module linking necessary 
+		TArrayView<TCompiledInObjectPtr<const FStructBaseChain>> BaseChain;
 		FField* ChildProperties;
 		int32 PropertiesSize;
 		int16 MinAlignment;
@@ -4380,7 +4447,7 @@ namespace UE::CodeGen::ConstInit
 
 	struct FClassImplementedInterface
 	{
-		UClass* Class;
+		TCompiledInObjectPtr<UClass> Class;
 		int32 PointerOffset;
 		bool bImplementedByK2;
 		bool bVerseDirectInterface;
@@ -4396,9 +4463,15 @@ namespace UE::CodeGen::ConstInit
 		EClassCastFlags ClassCastFlags;
 		const FCppClassTypeInfoStatic* CppClassTypeInfo;
 		TConstArrayView<FClassNativeFunction> NativeFunctions;
-		UClass* ClassWithin;
-		TConstArrayView<FClassImplementedInterface> Interfaces;
-		UScriptStruct* SparseClassDataStruct;
+		TObjectRefParam<UClass> ClassWithin;
+		TArrayView<FClassImplementedInterface> Interfaces;
+		TObjectRefParam<UScriptStruct> SparseClassDataStruct;
+
+		consteval FClassParams& SetWithin(UClass* WithinClass)
+		{ 
+			ClassWithin.ObjectPtr = WithinClass;
+			return *this;
+		}
 	};
 
 	struct FVerseClassParams
diff --git a/Engine/Source/Runtime/CoreUObject/Public/UObject/UnrealType.h b/Engine/Source/Runtime/CoreUObject/Public/UObject/UnrealType.h
--- a/Engine/Source/Runtime/CoreUObject/Public/UObject/UnrealType.h
+++ b/Engine/Source/Runtime/CoreUObject/Public/UObject/UnrealType.h
@@ -260,7 +260,7 @@ public:
 #if UE_WITH_CONSTINIT_UOBJECT
 	PRAGMA_DISABLE_DEPRECATION_WARNINGS
 	explicit consteval FProperty(UE::CodeGen::ConstInit::FPropertyParams InParams)
-		: Super(ConstEval, InParams.FieldClass, InParams.Owner, InParams.NextProperty, InParams.ObjectFlags, InParams.NameUTF8)
+		: Super(ConstEval, InParams.Owner, InParams.NextProperty, InParams.ObjectFlags, InParams.NameUTF8)
 		, ArrayDim(InParams.ArrayDim)
 		, ElementSize(InParams.ElementSize)
 		, PropertyFlags(InParams.PropertyFlags)
@@ -2180,9 +2180,9 @@ class FByteProperty : public TProperty_Numeric<uint8>
 	UE_API FByteProperty(FFieldVariant InOwner, const UECodeGen_Private::FBytePropertyParams& Prop);
 
 #if UE_WITH_CONSTINIT_UOBJECT
-	explicit consteval FByteProperty(UE::CodeGen::ConstInit::FPropertyParams InBaseParams, UEnum* InEnum)
+	explicit consteval FByteProperty(UE::CodeGen::ConstInit::FPropertyParams InBaseParams, UE::CodeGen::ConstInit::TObjectRefParam<UEnum> InEnum)
 		: Super(InBaseParams)
-		, Enum(ConstEval, InEnum)
+		, Enum(InEnum)
 	{
 	}
 #endif
@@ -2195,6 +2195,9 @@ class FByteProperty : public TProperty_Numeric<uint8>
 	UE_API virtual void Serialize( FArchive& Ar ) override;
 	UE_API virtual void AddReferencedObjects(FReferenceCollector& Collector) override;
 	UE_API virtual void GetPreloadDependencies(TArray<UObject*>& OutDeps) override;
+#if UE_WITH_CONSTINIT_UOBJECT && !IS_MONOLITHIC
+	UE_API virtual void LinkCompiledInPointerFields(UE::CoreUObject::Private::FCompiledInObjectLinker& Linker);
+#endif // UE_WITH_CONSTINIT_UOBJECT && !IS_MONOLITHIC
 	// End of UObject interface
 
 	// Field interface
@@ -2742,12 +2745,9 @@ public:
 
 #if UE_WITH_CONSTINIT_UOBJECT
 protected:
-	explicit consteval FObjectPropertyBase(
-		UE::CodeGen::ConstInit::FPropertyParams InBaseParams, 
-		UClass* InPropertyClass
-	)
+	explicit consteval FObjectPropertyBase(UE::CodeGen::ConstInit::FPropertyParams InBaseParams, UE::CodeGen::ConstInit::TObjectRefParam<UClass> InPropertyClass)
 		: Super(InBaseParams)
-		, PropertyClass(ConstEval, InPropertyClass)
+		, PropertyClass(InPropertyClass)
 	{
 	}
 public:
@@ -2761,6 +2761,9 @@ public:
 	UE_API virtual void Serialize( FArchive& Ar ) override;
 	UE_API virtual void AddReferencedObjects(FReferenceCollector& Collector) override;
 	UE_API virtual void BeginDestroy() override;
+#if UE_WITH_CONSTINIT_UOBJECT && !IS_MONOLITHIC
+	UE_API virtual void LinkCompiledInPointerFields(UE::CoreUObject::Private::FCompiledInObjectLinker& Linker);
+#endif // UE_WITH_CONSTINIT_UOBJECT && !IS_MONOLITHIC
 	// End of UObject interface
 
 	// Field interface
@@ -3046,10 +3049,7 @@ public:
 
 #if UE_WITH_CONSTINIT_UOBJECT
 protected:
-	explicit consteval TFObjectPropertyBase(
-		UE::CodeGen::ConstInit::FPropertyParams InBaseParams,
-		UClass* InPropertyClass
-	)
+	explicit consteval TFObjectPropertyBase(UE::CodeGen::ConstInit::FPropertyParams InBaseParams, UE::CodeGen::ConstInit::TObjectRefParam<UClass> InPropertyClass)
 		: Super(InBaseParams, InPropertyClass)
 	{
 	}
@@ -3097,7 +3097,7 @@ class FObjectProperty : public TFObjectPropertyBase<TObjectPtr<UObject>>
 	UE_API FObjectProperty(FFieldVariant InOwner, const UECodeGen_Private::FObjectPropertyParams& Prop);
 
 #if UE_WITH_CONSTINIT_UOBJECT
-	explicit consteval FObjectProperty(UE::CodeGen::ConstInit::FPropertyParams InBaseParams, UClass* PropertyClass)
+	explicit consteval FObjectProperty(UE::CodeGen::ConstInit::FPropertyParams InBaseParams, UE::CodeGen::ConstInit::TObjectRefParam<UClass> PropertyClass)
 		: Super(InBaseParams, PropertyClass)
 	{
 	}
@@ -3214,7 +3214,7 @@ class FWeakObjectProperty : public TFObjectPropertyBase<FWeakObjectPtr>
 	UE_API FWeakObjectProperty(FFieldVariant InOwner, const UECodeGen_Private::FWeakObjectPropertyParams& Prop);
 
 #if UE_WITH_CONSTINIT_UOBJECT
-	explicit consteval FWeakObjectProperty(UE::CodeGen::ConstInit::FPropertyParams InBaseParams, UClass* PropertyClass)
+	explicit consteval FWeakObjectProperty(UE::CodeGen::ConstInit::FPropertyParams InBaseParams, UE::CodeGen::ConstInit::TObjectRefParam<UClass> PropertyClass)
 		: Super(InBaseParams, PropertyClass)
 	{
 	}
@@ -3283,7 +3283,7 @@ class FLazyObjectProperty : public TFObjectPropertyBase<FLazyObjectPtr>
 #endif // WITH_EDITORONLY_DATA
 
 #if UE_WITH_CONSTINIT_UOBJECT
-	explicit consteval FLazyObjectProperty(UE::CodeGen::ConstInit::FPropertyParams InBaseParams, UClass* PropertyClass)
+	explicit consteval FLazyObjectProperty(UE::CodeGen::ConstInit::FPropertyParams InBaseParams, UE::CodeGen::ConstInit::TObjectRefParam<UClass> PropertyClass)
 		: Super(InBaseParams, PropertyClass)
 	{
 	}
@@ -3349,7 +3349,7 @@ class FSoftObjectProperty : public TFObjectPropertyBase<FSoftObjectPtr>
 	UE_API FSoftObjectProperty(FFieldVariant InOwner, const UECodeGen_Private::FObjectPropertyParamsWithoutClass& Prop, UClass* InClass);
 
 #if UE_WITH_CONSTINIT_UOBJECT
-	explicit consteval FSoftObjectProperty(UE::CodeGen::ConstInit::FPropertyParams InBaseParams, UClass* PropertyClass)
+	explicit consteval FSoftObjectProperty(UE::CodeGen::ConstInit::FPropertyParams InBaseParams, UE::CodeGen::ConstInit::TObjectRefParam<UClass> PropertyClass)
 		: Super(InBaseParams, PropertyClass)
 	{
 	}
@@ -3425,9 +3425,13 @@ public:
 	UE_API FClassProperty(FFieldVariant InOwner, const UECodeGen_Private::FClassPropertyParams& Prop);
 
 #if UE_WITH_CONSTINIT_UOBJECT
-	explicit consteval FClassProperty(UE::CodeGen::ConstInit::FPropertyParams InBaseParams, UClass* InPropertyClass, UClass* InMetaClass)
+	explicit consteval FClassProperty(
+		UE::CodeGen::ConstInit::FPropertyParams InBaseParams,
+		UE::CodeGen::ConstInit::TObjectRefParam<UClass> InPropertyClass,
+		UE::CodeGen::ConstInit::TObjectRefParam<UClass> InMetaClass
+	)
 		: Super(InBaseParams, InPropertyClass)
-		, MetaClass(ConstEval, InMetaClass)
+		, MetaClass(InMetaClass)
 	{
 	}
 #endif
@@ -3440,6 +3444,9 @@ public:
 	UE_API virtual void Serialize( FArchive& Ar ) override;
 	UE_API virtual void AddReferencedObjects(FReferenceCollector& Collector) override;
 	UE_API virtual void BeginDestroy() override;
+#if UE_WITH_CONSTINIT_UOBJECT && !IS_MONOLITHIC
+	UE_API virtual void LinkCompiledInPointerFields(UE::CoreUObject::Private::FCompiledInObjectLinker& Linker);
+#endif // UE_WITH_CONSTINIT_UOBJECT && !IS_MONOLITHIC
 	// End of UObject interface
 
 	// Field Interface
@@ -3502,9 +3509,13 @@ public:
 	UE_API FSoftClassProperty(FFieldVariant InOwner, const UECodeGen_Private::FSoftClassPropertyParams& Prop);
 
 #if UE_WITH_CONSTINIT_UOBJECT
-	explicit consteval FSoftClassProperty(UE::CodeGen::ConstInit::FPropertyParams InBaseParams, UClass* ClassClass, UClass* InMetaClass)
+	explicit consteval FSoftClassProperty(
+		UE::CodeGen::ConstInit::FPropertyParams InBaseParams,
+		UE::CodeGen::ConstInit::TObjectRefParam<UClass> ClassClass,
+		UE::CodeGen::ConstInit::TObjectRefParam<UClass> InMetaClass
+	)
 		: Super(InBaseParams, ClassClass)
-		, MetaClass(ConstEval, InMetaClass)
+		, MetaClass(InMetaClass)
 	{
 	}
 #endif
@@ -3525,6 +3536,9 @@ public:
 	UE_API virtual void Serialize( FArchive& Ar ) override;
 	UE_API virtual void AddReferencedObjects(FReferenceCollector& Collector) override;
 	UE_API virtual void BeginDestroy() override;
+#if UE_WITH_CONSTINIT_UOBJECT && !IS_MONOLITHIC
+	UE_API virtual void LinkCompiledInPointerFields(UE::CoreUObject::Private::FCompiledInObjectLinker& Linker);
+#endif // UE_WITH_CONSTINIT_UOBJECT && !IS_MONOLITHIC
 	// End of UObject interface
 
 	// FProperty interface
@@ -3578,9 +3592,9 @@ public:
 	UE_API FInterfaceProperty(FFieldVariant InOwner, const UECodeGen_Private::FInterfacePropertyParams& Prop);
 
 #if UE_WITH_CONSTINIT_UOBJECT
-	explicit consteval FInterfaceProperty(UE::CodeGen::ConstInit::FPropertyParams InBaseParams, UClass* InInterfaceClass)
+	explicit consteval FInterfaceProperty(UE::CodeGen::ConstInit::FPropertyParams InBaseParams, UE::CodeGen::ConstInit::TObjectRefParam<UClass> InInterfaceClass)
 		: Super(InBaseParams)
-		, InterfaceClass(ConstEval, InInterfaceClass)
+		, InterfaceClass(InInterfaceClass)
 	{
 	}
 #endif
@@ -3619,6 +3633,9 @@ public:
 	UE_API virtual void EmitReferenceInfo(UE::GC::FSchemaBuilder& Schema, int32 BaseOffset, TArray<const FStructProperty*>& EncounteredStructProps, UE::GC::FPropertyStack& DebugPath) override;
 	UE_API virtual void BeginDestroy() override;
 	UE_API virtual void AddReferencedObjects(FReferenceCollector& Collector) override;
+#if UE_WITH_CONSTINIT_UOBJECT && !IS_MONOLITHIC
+	UE_API virtual void LinkCompiledInPointerFields(UE::CoreUObject::Private::FCompiledInObjectLinker& Linker);
+#endif // UE_WITH_CONSTINIT_UOBJECT && !IS_MONOLITHIC
 	// End of UObject interface
 
 	/**
@@ -3746,6 +3763,9 @@ public:
 	UE_API virtual void Serialize( FArchive& Ar ) override;
 	UE_API virtual void AddReferencedObjects(FReferenceCollector& Collector) override;
 	UE_API virtual void GetPreloadDependencies(TArray<UObject*>& OutDeps) override;
+#if UE_WITH_CONSTINIT_UOBJECT && !IS_MONOLITHIC
+	UE_API virtual void LinkCompiledInPointerFields(UE::CoreUObject::Private::FCompiledInObjectLinker& Linker);
+#endif // UE_WITH_CONSTINIT_UOBJECT && !IS_MONOLITHIC
 	// End of UObject interface
 
 	// Field interface
@@ -3902,6 +3922,9 @@ public:
 	UE_API virtual void Serialize(FArchive& Ar) override;
 	UE_API virtual void AddReferencedObjects(FReferenceCollector& Collector) override;
 	UE_API virtual void GetPreloadDependencies(TArray<UObject*>& OutDeps) override;
+#if UE_WITH_CONSTINIT_UOBJECT && !IS_MONOLITHIC
+	UE_API virtual void LinkCompiledInPointerFields(UE::CoreUObject::Private::FCompiledInObjectLinker& Linker);
+#endif // UE_WITH_CONSTINIT_UOBJECT && !IS_MONOLITHIC
 	// End of UObject interface
 
 	// Field Interface
@@ -4066,6 +4089,9 @@ public:
 	UE_API virtual void Serialize(FArchive& Ar) override;
 	UE_API virtual void AddReferencedObjects(FReferenceCollector& Collector) override;
 	UE_API virtual void GetPreloadDependencies(TArray<UObject*>& OutDeps) override;
+#if UE_WITH_CONSTINIT_UOBJECT && !IS_MONOLITHIC
+	UE_API virtual void LinkCompiledInPointerFields(UE::CoreUObject::Private::FCompiledInObjectLinker& Linker);
+#endif // UE_WITH_CONSTINIT_UOBJECT && !IS_MONOLITHIC
 	// End of UObject interface
 
 	// Field interface
@@ -6315,9 +6341,13 @@ public:
 	UE_API FStructProperty(FFieldVariant InOwner, const UECodeGen_Private::FStructPropertyParams& Prop);
 
 #if UE_WITH_CONSTINIT_UOBJECT
-	explicit consteval FStructProperty(UE::CodeGen::ConstInit::FPropertyParams InBaseParams, int32 StructSize, UScriptStruct* InScriptStruct)
+	explicit consteval FStructProperty(
+		UE::CodeGen::ConstInit::FPropertyParams InBaseParams,
+		int32 StructSize,
+		UE::CodeGen::ConstInit::TObjectRefParam<UScriptStruct> InScriptStruct
+	)
 		: Super(InBaseParams.SetElementSize(StructSize))
-		, Struct(ConstEval, InScriptStruct)
+		, Struct(InScriptStruct)
 	{
 	}
 #endif
@@ -6330,6 +6360,9 @@ public:
 	UE_API virtual void Serialize( FArchive& Ar ) override;
 	UE_API virtual void AddReferencedObjects(FReferenceCollector& Collector) override;
 	UE_API virtual void GetPreloadDependencies(TArray<UObject*>& OutDeps) override;
+#if UE_WITH_CONSTINIT_UOBJECT && !IS_MONOLITHIC
+	UE_API virtual void LinkCompiledInPointerFields(UE::CoreUObject::Private::FCompiledInObjectLinker& Linker);
+#endif // UE_WITH_CONSTINIT_UOBJECT && !IS_MONOLITHIC
 	// End of UObject interface
 
 	// Field interface
@@ -6410,9 +6443,9 @@ public:
 	UE_API FDelegateProperty(FFieldVariant InOwner, const UECodeGen_Private::FDelegatePropertyParams& Prop);
 
 #if UE_WITH_CONSTINIT_UOBJECT
-	explicit consteval FDelegateProperty(UE::CodeGen::ConstInit::FPropertyParams InBaseParams, UFunction* InSignatureFunction)
+	explicit consteval FDelegateProperty(UE::CodeGen::ConstInit::FPropertyParams InBaseParams, UE::CodeGen::ConstInit::TObjectRefParam<UFunction> InSignatureFunction)
 		: Super(InBaseParams)
-		, SignatureFunction(ConstEval, InSignatureFunction)
+		, SignatureFunction(InSignatureFunction)
 	{
 	}
 #endif
@@ -6425,6 +6458,9 @@ public:
 	UE_API virtual void Serialize( FArchive& Ar ) override;
 	UE_API virtual void AddReferencedObjects(FReferenceCollector& Collector) override;
 	UE_API virtual void BeginDestroy() override;
+#if UE_WITH_CONSTINIT_UOBJECT && !IS_MONOLITHIC
+	UE_API virtual void LinkCompiledInPointerFields(UE::CoreUObject::Private::FCompiledInObjectLinker& Linker);
+#endif // UE_WITH_CONSTINIT_UOBJECT && !IS_MONOLITHIC
 	// End of UObject interface
 
 	// Field interface
@@ -6475,9 +6511,9 @@ public:
 	UE_API FMulticastDelegateProperty(FFieldVariant InOwner, const UECodeGen_Private::FMulticastDelegatePropertyParams& Prop, EPropertyFlags AdditionalPropertyFlags = CPF_None);
 
 #if UE_WITH_CONSTINIT_UOBJECT
-	explicit consteval FMulticastDelegateProperty(UE::CodeGen::ConstInit::FPropertyParams InBaseParams, UFunction* InSignatureFunction)
+	explicit consteval FMulticastDelegateProperty(UE::CodeGen::ConstInit::FPropertyParams InBaseParams, UE::CodeGen::ConstInit::TObjectRefParam<UFunction> InSignatureFunction)
 		: Super(InBaseParams)
-		, SignatureFunction(ConstEval, InSignatureFunction)
+		, SignatureFunction(InSignatureFunction)
 	{
 	}
 #endif
@@ -6490,6 +6526,9 @@ public:
 	UE_API virtual void Serialize( FArchive& Ar ) override;
 	UE_API virtual void BeginDestroy() override;
 	UE_API virtual void AddReferencedObjects(FReferenceCollector& Collector) override;
+#if UE_WITH_CONSTINIT_UOBJECT && !IS_MONOLITHIC
+	UE_API virtual void LinkCompiledInPointerFields(UE::CoreUObject::Private::FCompiledInObjectLinker& Linker);
+#endif // UE_WITH_CONSTINIT_UOBJECT && !IS_MONOLITHIC
 	// End of UObject interface
 
 	// Field interface
@@ -6564,7 +6603,7 @@ public:
 	}
 
 #if UE_WITH_CONSTINIT_UOBJECT
-	explicit consteval TProperty_MulticastDelegate(UE::CodeGen::ConstInit::FPropertyParams InBaseParams, UFunction* InSignatureFunction)
+	explicit consteval TProperty_MulticastDelegate(UE::CodeGen::ConstInit::FPropertyParams InBaseParams, UE::CodeGen::ConstInit::TObjectRefParam<UFunction> InSignatureFunction)
 		: Super(InBaseParams, InSignatureFunction)
 	{
 	}
@@ -6604,7 +6643,10 @@ public:
 	UE_API FMulticastInlineDelegateProperty(FFieldVariant InOwner, const UECodeGen_Private::FMulticastDelegatePropertyParams& Prop);
 
 #if UE_WITH_CONSTINIT_UOBJECT
-	explicit consteval FMulticastInlineDelegateProperty(UE::CodeGen::ConstInit::FPropertyParams InBaseParams, UFunction* InSignatureFunction)
+	explicit consteval FMulticastInlineDelegateProperty(
+		UE::CodeGen::ConstInit::FPropertyParams InBaseParams,
+		UE::CodeGen::ConstInit::TObjectRefParam<UFunction> InSignatureFunction
+	)
 		: Super(InBaseParams, InSignatureFunction)
 	{
 	}
@@ -6654,7 +6696,10 @@ public:
 	UE_API FMulticastSparseDelegateProperty(FFieldVariant InOwner, const UECodeGen_Private::FMulticastDelegatePropertyParams& Prop);
 
 #if UE_WITH_CONSTINIT_UOBJECT
-	explicit consteval FMulticastSparseDelegateProperty(UE::CodeGen::ConstInit::FPropertyParams InBaseParams, UFunction* InSignatureFunction)
+	explicit consteval FMulticastSparseDelegateProperty(
+		UE::CodeGen::ConstInit::FPropertyParams InBaseParams,
+		UE::CodeGen::ConstInit::TObjectRefParam<UFunction> InSignatureFunction
+	)
 		: Super(InBaseParams, InSignatureFunction)
 	{
 	}
