Silence complexity warning about 'Generate' function

#Synced-CL Unknown

diff --git a/Engine/Source/Programs/Shared/EpicGames.UHT/Exporters/CodeGen/UhtHeaderCodeGeneratorCppFile.cs b/Engine/Source/Programs/Shared/EpicGames.UHT/Exporters/CodeGen/UhtHeaderCodeGeneratorCppFile.cs
--- a/Engine/Source/Programs/Shared/EpicGames.UHT/Exporters/CodeGen/UhtHeaderCodeGeneratorCppFile.cs
+++ b/Engine/Source/Programs/Shared/EpicGames.UHT/Exporters/CodeGen/UhtHeaderCodeGeneratorCppFile.cs
@@ -36,6 +36,8 @@ namespace EpicGames.UHT.Exporters.CodeGen
 		{
 		}
 
+#pragma warning disable CA1505 // warning CA1505: 'Generate' has a maintainability index of '8'. Rewrite or refactor the code to increase its maintainability index (MI) above '9'. 
+
 		/// <summary>
 		/// For a given UE header file, generated the generated H file
 		/// </summary>
@@ -508,6 +510,7 @@ namespace EpicGames.UHT.Exporters.CodeGen
 				}
 			}
 		}
+#pragma warning restore CA1505 // warning CA1505: 'Generate' has a maintainability index of '8'. Rewrite or refactor the code to increase its maintainability index (MI) above '9'. 
 
 		private void AddIncludeForType(UhtProperty uhtProperty, bool requireIncludeForClasses, HashSet<UhtHeaderFile> addedIncludes, IList<string> includesToAdd)
 		{
