Fix bad merge from backout

diff --git a/Engine/Source/Programs/Shared/EpicGames.Core/UHTTypes.cs b/Engine/Source/Programs/Shared/EpicGames.Core/UHTTypes.cs
--- a/Engine/Source/Programs/Shared/EpicGames.Core/UHTTypes.cs
+++ b/Engine/Source/Programs/Shared/EpicGames.Core/UHTTypes.cs
@@ -160,7 +160,6 @@ namespace EpicGames.Core
 		{
 			writer.WriteString(CompiledInObjectFormat.ToString());
 			writer.WriteBoolean(GenerateEnumMaxValues);
-			writer.WriteBoolean(ValidateModuleDependencies);
 		}
 
 		/// <summary>
@@ -174,7 +173,7 @@ namespace EpicGames.Core
 		/// If this is false, the UEnum code will generate it at runtime if it can do so without creating a collision. 
 		/// </summary>
 		public bool GenerateEnumMaxValues { get; set; } = false;
-		
+
 		/// <summary>
 		/// If true, enable the input cache read
 		/// </summary>
