EpicGames.Perforce: Add support for streamlog command

diff --git a/Engine/Source/Programs/Shared/EpicGames.Perforce/PerforceConnection.cs b/Engine/Source/Programs/Shared/EpicGames.Perforce/PerforceConnection.cs
--- a/Engine/Source/Programs/Shared/EpicGames.Perforce/PerforceConnection.cs
+++ b/Engine/Source/Programs/Shared/EpicGames.Perforce/PerforceConnection.cs
@@ -3565,6 +3565,70 @@ namespace EpicGames.Perforce
 		}
 		#endregion
 
+		#region p4 streamlog
+
+		/// <summary>
+		/// List revision history of streams
+		/// </summary>
+		/// <param name="connection">Connection to the Perforce server</param>
+		/// <param name="options">Options for this command</param>
+		/// <param name="changeNumber">Optional: Display the stream submitted at the specified changelist number</param>
+		/// <param name="maxChanges">Optional: Display at most this many revisions per stream argument specified</param>
+		/// <param name="streams">Streams to list history for</param>
+		/// <param name="cancellationToken">Token used to cancel the operation</param>
+		/// <returns>Stream log records</returns>
+		public static async Task<List<StreamLogRecord>> GetStreamLogAsync(this IPerforceConnection connection, StreamLogOptions options, int? changeNumber, int? maxChanges, string[] streams, CancellationToken cancellationToken = default)
+		{
+			return (await TryGetStreamLogAsync(connection, options, changeNumber, maxChanges, streams, cancellationToken)).Data;
+		}
+
+		/// <summary>
+		/// List revision history of streams
+		/// </summary>
+		/// <param name="connection">Connection to the Perforce server</param>
+		/// <param name="options">Options for this command</param>
+		/// <param name="changeNumber">Optional: Display the stream submitted at the specified changelist number</param>
+		/// <param name="maxChanges">Optional: Display at most this many revisions per stream argument specified</param>
+		/// <param name="streams">Streams to list history for</param>
+		/// <param name="cancellationToken">Token used to cancel the operation</param>
+		/// <returns>Stream log records</returns>
+		public static Task<PerforceResponseList<StreamLogRecord>> TryGetStreamLogAsync(this IPerforceConnection connection, StreamLogOptions options, int? changeNumber, int? maxChanges, string[] streams, CancellationToken cancellationToken = default)
+		{
+			List<string> arguments = [];
+			
+			if ((options & StreamLogOptions.IncludeHistory) != 0)
+			{
+				arguments.Add("-i");
+			}
+			if ((options & StreamLogOptions.DisplayStreamContent) != 0)
+			{
+				arguments.Add("-h");
+			}
+			if ((options & StreamLogOptions.IncludeTimes) != 0)
+			{
+				arguments.Add("-t");
+			}
+			if ((options & StreamLogOptions.LongOutput) != 0)
+			{
+				arguments.Add("-l");
+			}
+			if ((options & StreamLogOptions.TruncatedLongOutput) != 0)
+			{
+				arguments.Add("-L");
+			}
+			if (changeNumber != null)
+			{
+				arguments.Add($"-c{changeNumber.Value}");
+			}
+            if (maxChanges != null)
+			{
+				arguments.Add($"-m{maxChanges.Value}");
+			}
+			arguments.AddRange(streams);
+			return CommandAsync<StreamLogRecord>(connection, "streamlog", arguments, null, cancellationToken);
+		}
+		#endregion
+
 		#region p4 streams
 
 		/// <summary>
diff --git a/Engine/Source/Programs/Shared/EpicGames.Perforce/Records/StreamLogOptions.cs b/Engine/Source/Programs/Shared/EpicGames.Perforce/Records/StreamLogOptions.cs
new file mode 100644
--- /dev/null
+++ b/Engine/Source/Programs/Shared/EpicGames.Perforce/Records/StreamLogOptions.cs
@@ -0,0 +1,43 @@
+// Copyright Epic Games, Inc. All Rights Reserved.
+
+using System;
+
+namespace EpicGames.Perforce
+{
+	/// <summary>
+	/// Options for the 'streamlog' command
+	/// </summary>
+	[Flags]
+	public enum StreamLogOptions
+	{
+		/// <summary>
+		/// No options specified
+		/// </summary>
+		None = 0,
+
+		/// <summary>
+		/// Include inherited stream history (-i)
+		/// </summary>
+		IncludeHistory = 1,
+
+		/// <summary>
+		/// Display the stream content history instead of stream name history (-h)
+		/// </summary>
+		DisplayStreamContent = 2,
+
+		/// <summary>
+		/// Display the time as well as the date (-t)
+		/// </summary>
+		IncludeTimes = 4,
+		
+		/// <summary>
+		/// List long output, with the full text of each changelist description (-l)
+		/// </summary>
+		LongOutput = 8,
+
+		/// <summary>
+		/// List long output, with the full text of each changelist description truncated at 250 characters (-L)
+		/// </summary>
+		TruncatedLongOutput = 16,
+	}
+}
diff --git a/Engine/Source/Programs/Shared/EpicGames.Perforce/Records/StreamLogRecord.cs b/Engine/Source/Programs/Shared/EpicGames.Perforce/Records/StreamLogRecord.cs
new file mode 100644
--- /dev/null
+++ b/Engine/Source/Programs/Shared/EpicGames.Perforce/Records/StreamLogRecord.cs
@@ -0,0 +1,44 @@
+// Copyright Epic Games, Inc. All Rights Reserved.
+
+using System;
+using System.Diagnostics;
+
+namespace EpicGames.Perforce
+{
+	/// <summary>
+	/// Contains summary information about a stream log record
+	/// </summary>
+	[DebuggerDisplay("{Number}: {Description}")]
+	public class StreamLogRecord
+	{
+		/// <summary>
+		/// The changelist number
+		/// </summary>
+		[PerforceTag("change")]
+		public int Number { get; set; }
+
+		/// <summary>
+		/// The date the change was last modified
+		/// </summary>
+		[PerforceTag("time")]
+		public DateTime Time { get; set; }
+
+		/// <summary>
+		/// The user that owns or submitted the change
+		/// </summary>
+		[PerforceTag("user")]
+		public string User { get; set; } = String.Empty;
+
+		/// <summary>
+		/// The client that owns the change
+		/// </summary>
+		[PerforceTag("client")]
+		public string? Client { get; set; }
+
+		/// <summary>
+		/// Description for the changelist
+		/// </summary>
+		[PerforceTag("desc")]
+		public string Description { get; set; } = String.Empty;
+	}
+}
