EpicGames.Perforce: Add sparse streams to fixture

#skipci

diff --git a/Engine/Source/Programs/Shared/EpicGames.Perforce.Fixture/PerforceFixture.cs b/Engine/Source/Programs/Shared/EpicGames.Perforce.Fixture/PerforceFixture.cs
--- a/Engine/Source/Programs/Shared/EpicGames.Perforce.Fixture/PerforceFixture.cs
+++ b/Engine/Source/Programs/Shared/EpicGames.Perforce.Fixture/PerforceFixture.cs
@@ -315,6 +315,20 @@ public class PerforceFixture
 					// data.txt is excluded in view
 				}),
 		});
+	
+	public StreamFixture StreamFooSparseDev { get; } = new("//Foo/Main-SparseDev", 12,
+		new List<ChangelistFixture>
+		{
+			new(12, "Stream //Foo/Main-SparseDev saved.", // TODO: add action
+				new List<DepotFileFixture>()
+				{
+					new("//Foo/Main/main.cpp", "main.cpp", "This is change main.cpp #3\n", 3),
+					new("//Foo/Main/main.h", "main.h", "This is change main.h #2\n", 2),
+					new("//Foo/Main/shared.h", "shared.h", "This is change common.h #1\n", 1),
+					new("//Foo/Main/Data/data.txt", "Data/data.txt", "This is change data.txt #1\n", 1),
+					new("//Foo/Main/Data/moredata.txt", "Data/moredata.txt", "This is change moredata.txt #1\n", 1),
+				})
+		});
 
 	public static string CalcMd5(string content)
 	{
diff --git a/Engine/Source/Programs/Shared/EpicGames.Perforce.Fixture/setup.sh b/Engine/Source/Programs/Shared/EpicGames.Perforce.Fixture/setup.sh
--- a/Engine/Source/Programs/Shared/EpicGames.Perforce.Fixture/setup.sh
+++ b/Engine/Source/Programs/Shared/EpicGames.Perforce.Fixture/setup.sh
@@ -167,5 +167,37 @@ echo "This is change main.cpp #2 in //Foo/Release" > $CLIENT_ROOT/main.cpp
 echo "This is change release.cpp #1 in //Foo/Release" > $CLIENT_ROOT/release.cpp && p4 add $CLIENT_ROOT/release.cpp
 p4 submit -d "Improvements to //Foo/Release"
 
+# Create a sparse dev stream
+p4 stream -i << EOF
+Stream: //Foo/Main-SparseDev
+Owner:  test.user
+Name:   Main-SparseDev
+Parent: //Foo/Main
+Type:   sparsedev
+Description:
+        Sparse dev stream for //Foo/Main
+
+Options:        allsubmit unlocked notoparent nofromparent mergedown
+ParentView:     inherit
+Paths:
+        share ...
+EOF
+
+# Create a sparse release stream
+p4 stream -i << EOF
+Stream: //Foo/Main-SparseRel
+Owner:  test.user
+Name:   Main-SparseRel
+Parent: //Foo/Release
+Type:   sparserel
+Description:
+        Sparse release stream for //Foo/Main
+
+Options:        allsubmit unlocked notoparent nofromparent mergedown
+ParentView:     inherit
+Paths:
+        share ...
+EOF
+
 # Stop server running in background
 pkill p4d
\ No newline at end of file
