[Horde][Analytics][Serever] Commit telemetry integration into the p4 cache service.

#Synced-CL 48176906

diff --git a/Engine/Source/Programs/Horde/Plugins/Build/HordeServer.Build/VersionControl/Perforce/CommitTelemetry.cs b/Engine/Source/Programs/Horde/Plugins/Build/HordeServer.Build/VersionControl/Perforce/CommitTelemetry.cs
new file mode 100644
--- /dev/null
+++ b/Engine/Source/Programs/Horde/Plugins/Build/HordeServer.Build/VersionControl/Perforce/CommitTelemetry.cs
@@ -0,0 +1,31 @@
+﻿// Copyright Epic Games, Inc. All Rights Reserved.
+
+using EpicGames.Horde.Commits;
+using EpicGames.Horde.Streams;
+
+namespace HordeServer.VersionControl.Perforce
+{
+	/// <summary>
+	/// Record used for Commit telemetry. Models a <see cref="ICommit"/>.
+	/// </summary>
+	/// <remarks>Modifications to this record's structure should be versioned appropriately through <see cref="SchemaVersion"/> via <see cref="CurrentSchemaVersion"/>, and call sites should continue to publish in a backwards compatible manner.</remarks>
+	public record CommitTelemetry(string EventName, int SchemaVersion, StreamId StreamId, int Change, int OriginalChange, DateTime SubmitDate, string? ExternalIssueKey, string? RobomergeAuthorTag, string? RobomergeSourceTag, string? RobomergeOwnerTag, string? RobomergeBotTag, string? ReviewTag)
+	{
+		/// <summary>
+		/// Default event name for the CommitTelemetry.
+		/// </summary>
+		public static readonly string DefaultEventName = "State.Change.Dev";
+
+		/// <summary>
+		/// Current schema version.
+		/// </summary>
+		public static readonly int CurrentSchemaVersion = 4;
+
+		/// <summary>
+		/// Constructs a CommitTelemetry event based on the provided commit and commit tag container.
+		/// </summary>
+		/// <param name="commit">The commit to base the telemetry event off of.</param>
+		/// <param name="commitTagContainer">The commit description tag container to base the telemetry off of.</param>
+		public CommitTelemetry(ICommit commit, CommitDescriptionTagContainer commitTagContainer) : this(DefaultEventName, CurrentSchemaVersion, commit.StreamId, commit.Id.Order, commit.OriginalCommitId.Order, commit.DateUtc, commitTagContainer.JiraTag, commitTagContainer.RobomergeAuthorTag, commitTagContainer.RobomergeSourceTag, commitTagContainer.RobomergeOwnerTag, commitTagContainer.RobomergeBotTag, commitTagContainer.ReviewTag) { }
+	}
+}
\ No newline at end of file
diff --git a/Engine/Source/Programs/Horde/Plugins/Build/HordeServer.Build/VersionControl/Perforce/PerforceService.cs b/Engine/Source/Programs/Horde/Plugins/Build/HordeServer.Build/VersionControl/Perforce/PerforceService.cs
--- a/Engine/Source/Programs/Horde/Plugins/Build/HordeServer.Build/VersionControl/Perforce/PerforceService.cs
+++ b/Engine/Source/Programs/Horde/Plugins/Build/HordeServer.Build/VersionControl/Perforce/PerforceService.cs
@@ -712,7 +712,7 @@ namespace HordeServer.VersionControl.Perforce
 
 		protected async ValueTask<ICommit> CreateCommitAsync(IPooledPerforceConnection perforce, StreamConfig streamConfig, DescribeRecord describeRecord, bool hasAllFiles, InfoRecord serverInfo, CancellationToken cancellationToken)
 		{
-			DateTime timeUtc = new DateTime(describeRecord.Time.Ticks - serverInfo.TimeZoneOffsetSecs * TimeSpan.TicksPerSecond, DateTimeKind.Utc);
+			DateTime timeUtc = describeRecord.Time.Kind == DateTimeKind.Utc ? describeRecord.Time : new DateTime(describeRecord.Time.Ticks - serverInfo.TimeZoneOffsetSecs * TimeSpan.TicksPerSecond, DateTimeKind.Utc);
 
 			Commit commit = await CreateCommitInternalAsync(streamConfig, describeRecord.Number, describeRecord.User, describeRecord.Description, describeRecord.Path, timeUtc, cancellationToken);
 
@@ -728,7 +728,7 @@ namespace HordeServer.VersionControl.Perforce
 
 		async ValueTask<Commit> CreateCommitAsync(StreamConfig streamConfig, ChangesRecord changesRecord, InfoRecord serverInfo, CancellationToken cancellationToken)
 		{
-			DateTime timeUtc = new DateTime(changesRecord.Time.Ticks - serverInfo.TimeZoneOffsetSecs * TimeSpan.TicksPerSecond, DateTimeKind.Utc);
+			DateTime timeUtc = changesRecord.Time.Kind == DateTimeKind.Utc ? changesRecord.Time : new DateTime(changesRecord.Time.Ticks - serverInfo.TimeZoneOffsetSecs * TimeSpan.TicksPerSecond, DateTimeKind.Utc);
 
 			Commit commit = await CreateCommitInternalAsync(streamConfig, changesRecord.Number, changesRecord.User, changesRecord.Description, changesRecord.Path, timeUtc, cancellationToken);
 
diff --git a/Engine/Source/Programs/Horde/Plugins/Build/HordeServer.Build/VersionControl/Perforce/PerforceServiceCache.cs b/Engine/Source/Programs/Horde/Plugins/Build/HordeServer.Build/VersionControl/Perforce/PerforceServiceCache.cs
--- a/Engine/Source/Programs/Horde/Plugins/Build/HordeServer.Build/VersionControl/Perforce/PerforceServiceCache.cs
+++ b/Engine/Source/Programs/Horde/Plugins/Build/HordeServer.Build/VersionControl/Perforce/PerforceServiceCache.cs
@@ -10,6 +10,7 @@ using EpicGames.Perforce;
 using EpicGames.Redis;
 using HordeServer.Server;
 using HordeServer.Streams;
+using HordeServer.Telemetry;
 using HordeServer.Users;
 using HordeServer.Utilities;
 using Microsoft.Extensions.Hosting;
@@ -225,6 +226,7 @@ namespace HordeServer.VersionControl.Perforce
 		readonly IDowntimeService _downtimeService;
 		readonly IMongoCollection<CachedCommitDoc> _commits;
 		readonly IOptionsMonitor<BuildConfig> _buildConfig;
+		readonly ITelemetryWriter _telemetryWriter;
 		readonly Tracer _tracer;
 		readonly ILogger _logger;
 
@@ -235,7 +237,7 @@ namespace HordeServer.VersionControl.Perforce
 		/// <summary>
 		/// Constructor
 		/// </summary>
-		public PerforceServiceCache(PerforceLoadBalancer loadBalancer, IMongoService mongoService, IRedisService redisService, IDowntimeService downtimeService, IUserCollection userCollection, IClock clock, IOptions<BuildServerConfig> settings, IOptionsMonitor<BuildConfig> buildConfig, Tracer tracer, ILogger<PerforceService> logger)
+		public PerforceServiceCache(PerforceLoadBalancer loadBalancer, IMongoService mongoService, IRedisService redisService, IDowntimeService downtimeService, IUserCollection userCollection, IClock clock, IOptions<BuildServerConfig> settings, IOptionsMonitor<BuildConfig> buildConfig, ITelemetryWriter telemetryWriter, Tracer tracer, ILogger<PerforceService> logger)
 			: base(loadBalancer, userCollection, settings, buildConfig, tracer, logger)
 		{
 			_mongoService = mongoService;
@@ -248,6 +250,7 @@ namespace HordeServer.VersionControl.Perforce
 			_commits = mongoService.GetCollection<CachedCommitDoc>("CommitsV3", indexes);
 
 			_buildConfig = buildConfig;
+			_telemetryWriter = telemetryWriter;
 			_tracer = tracer;
 			_logger = logger;
 
@@ -547,6 +550,19 @@ namespace HordeServer.VersionControl.Perforce
 							await AddCachedCommitAsync(commitDoc, cancellationToken);
 
 							await _redisService.PublishAsync(s_commitUpdateChannel, streamInfo.StreamConfig.Id, CommandFlags.FireAndForget);
+
+							// Emit telemetry
+							if (!streamInfo.StreamConfig.TelemetryStoreId.IsEmpty)
+							{
+								try
+								{
+									_telemetryWriter.WriteEvent(streamInfo.StreamConfig.TelemetryStoreId, new CommitTelemetry(commit, CommitDescriptionTagContainer.Create(commitDoc.Description)));
+								}
+								catch (Exception ex)
+								{
+									_logger.LogError(ex, "Failed to publish Commit telemetry.");
+								}
+							}
 						}
 					}
 				}
diff --git a/Engine/Source/Programs/Shared/EpicGames.Horde/Commits/CommitTag.cs b/Engine/Source/Programs/Shared/EpicGames.Horde/Commits/CommitTag.cs
--- a/Engine/Source/Programs/Shared/EpicGames.Horde/Commits/CommitTag.cs
+++ b/Engine/Source/Programs/Shared/EpicGames.Horde/Commits/CommitTag.cs
@@ -3,10 +3,124 @@
 using System;
 using System.Text.Json;
 using System.Text.Json.Serialization;
+using System.Text.RegularExpressions;
 using EpicGames.Core;
 
 namespace EpicGames.Horde.Commits
 {
+	internal static class RegexExtensions
+	{
+		public static string? ValueOrNull(this Group group)
+		{
+			return group.Success ? group.Value : null;
+		}
+	}
+
+	/// <summary>
+	/// Convinence data structure that contains all tags associated with a provided commit description.
+	/// </summary>
+	public partial struct CommitDescriptionTagContainer
+	{
+
+		#region -- Compile-time Regex --
+
+		[GeneratedRegex(@"^#ROBOMERGE-OWNER:\s*([^@\s]+)", RegexOptions.Multiline)]
+		private static partial Regex RobomergeOwnerRegex();
+
+		[GeneratedRegex(@"^#ROBOMERGE-SOURCE:\s*(.*)", RegexOptions.Multiline)]
+		private static partial Regex RobomergeSourceRegex();
+
+		[GeneratedRegex(@"^#ROBOMERGE-AUTHOR:\s*(\S+)", RegexOptions.Multiline)]
+		private static partial Regex RobomergeAuthorRegex();
+
+		[GeneratedRegex(@"^#ROBOMERGE-BOT:\s*(.*)", RegexOptions.Multiline)]
+		private static partial Regex RobomergeBotRegex();
+
+		[GeneratedRegex(@"^#jira\s*([^@\s]+)", RegexOptions.Multiline)]
+		private static partial Regex JiraRegex();
+
+		[GeneratedRegex(@"^#tests\s*(.*)", RegexOptions.Multiline)]
+		private static partial Regex TestsRegex();
+
+		[GeneratedRegex(@"^#preflight\s*([^@\s]+)", RegexOptions.Multiline)]
+		private static partial Regex PreflightRegex();
+
+		[GeneratedRegex(@"^#rb\s+(.+)$", RegexOptions.Multiline)]
+		private static partial Regex ReviewRegex();
+
+		#endregion -- Compile-time Regex --
+
+		#region -- Public Properties --
+
+		/// <summary>
+		/// The source description of which the tags are extracted from.
+		/// </summary>
+		public string Description { get; init; }
+
+		/// <summary>
+		/// The jira tag in the commit description.
+		/// </summary>
+		public string? JiraTag { get; init; }
+
+		/// <summary>
+		/// The review tag in the commit description.
+		/// </summary>
+		public string? ReviewTag { get; init; }
+
+		/// <summary>
+		/// The presubmit tag in the commit description.
+		/// </summary>
+		public string? PresubmitTag { get; init; }
+
+		/// <summary>
+		/// The test tag in the commit description.
+		/// </summary>
+		public string? TestTag { get; init; }
+
+		/// <summary>
+		/// The robomerge owner tag in the commit description.
+		/// </summary>
+		public string? RobomergeOwnerTag { get; init; }
+
+		/// <summary>
+		/// The robomerge author tag in the commit description.
+		/// </summary>
+		public string? RobomergeAuthorTag { get; init; }
+
+		/// <summary>
+		/// The robomerge source tag in the commit description.
+		/// </summary>
+		public string? RobomergeSourceTag { get; init; }
+
+		/// <summary>
+		/// The robomerge bot tag in the commit description.
+		/// </summary>
+		public string? RobomergeBotTag { get; init; }
+
+		#endregion -- Public Properties --
+
+		/// <summary>
+		/// Creates a new CommitDescriptionTagContainer from a provided description.
+		/// </summary>
+		/// <param name="description">The description to use as the source for the CommitDescriptionTagContainer.</param>
+		/// <returns>The CommitDescriptionTagContainer as constructed from the input description.</returns>
+		public static CommitDescriptionTagContainer Create(string description)
+		{
+			return new CommitDescriptionTagContainer
+			{
+				Description = description,
+				RobomergeOwnerTag = RobomergeOwnerRegex().Match(description).Groups[1].ValueOrNull()?.Trim(),
+				RobomergeSourceTag = RobomergeSourceRegex().Match(description).Groups[1].ValueOrNull()?.Trim(),
+				RobomergeAuthorTag = RobomergeAuthorRegex().Match(description).Groups[1].ValueOrNull()?.Trim(),
+				RobomergeBotTag = RobomergeBotRegex().Match(description).Groups[1].ValueOrNull()?.Trim(),
+				JiraTag = JiraRegex().Match(description).Groups[1].ValueOrNull()?.Trim(),
+				PresubmitTag = PreflightRegex().Match(description).Groups[1].ValueOrNull()?.Trim(),
+				TestTag = TestsRegex().Match(description).Groups[1].ValueOrNull()?.Trim(),
+				ReviewTag = ReviewRegex().Match(description).Groups[1].ValueOrNull()?.Trim()
+			};
+		}
+	}
+
 	/// <summary>
 	/// Constants for known commit tags
 	/// </summary>
