EpicGames.Build: Adjust aggregate exception logging

#rnx
#Synced-CL 47605596

diff --git a/Engine/Source/Programs/Shared/EpicGames.Build/Exceptions/BuildExceptionExtensions.cs b/Engine/Source/Programs/Shared/EpicGames.Build/Exceptions/BuildExceptionExtensions.cs
--- a/Engine/Source/Programs/Shared/EpicGames.Build/Exceptions/BuildExceptionExtensions.cs
+++ b/Engine/Source/Programs/Shared/EpicGames.Build/Exceptions/BuildExceptionExtensions.cs
@@ -33,7 +33,7 @@ namespace UnrealBuildBase
 			}
 			else if (ex is AggregateException aggregateException)
 			{
-				logger.LogError(ex, "Unhandled {Count} aggregate exceptions", aggregateException.InnerExceptions.Count);
+				logger.LogInformation(ex, "Unhandled {Count} aggregate exceptions", aggregateException.InnerExceptions.Count);
 
 				foreach (Exception innerException in aggregateException.InnerExceptions)
 				{
@@ -53,10 +53,26 @@ namespace UnrealBuildBase
 		/// <returns>CompilationResult</returns>
 		public static CompilationResult GetCompilationResult(this Exception ex)
 		{
-			return (ex as CompilationResultException)?.Result
-				?? (ex.InnerException as CompilationResultException)?.Result
-				?? (ex as AggregateException)?.InnerExceptions.OfType<CompilationResultException>().FirstOrDefault()?.Result
-				?? CompilationResult.OtherCompilationError;
+			if (ex is CompilationResultException exc)
+			{
+				return exc.Result;
+			}
+			else if (ex is AggregateException aggException)
+			{
+				foreach (Exception innerException in aggException.InnerExceptions)
+				{
+					CompilationResult result = GetCompilationResult(innerException);
+					if (result != CompilationResult.OtherCompilationError)
+					{
+						return result;
+					}
+				}
+			}
+			else if (ex.InnerException != null)
+			{
+				return GetCompilationResult(ex.InnerException);
+			}
+			return CompilationResult.OtherCompilationError;
 		}
 	}
 }
