EpicGames.Build: Rename some variables

#rnx

#Synced-CL 47605596

diff --git a/Engine/Source/Programs/Shared/EpicGames.Build/Exceptions/BuildExceptionExtensions.cs b/Engine/Source/Programs/Shared/EpicGames.Build/Exceptions/BuildExceptionExtensions.cs
--- a/Engine/Source/Programs/Shared/EpicGames.Build/Exceptions/BuildExceptionExtensions.cs
+++ b/Engine/Source/Programs/Shared/EpicGames.Build/Exceptions/BuildExceptionExtensions.cs
@@ -53,13 +53,13 @@ namespace UnrealBuildBase
 		/// <returns>CompilationResult</returns>
 		public static CompilationResult GetCompilationResult(this Exception ex)
 		{
-			if (ex is CompilationResultException exc)
+			if (ex is CompilationResultException compilationResultsException)
 			{
-				return exc.Result;
+				return compilationResultsException.Result;
 			}
-			else if (ex is AggregateException aggException)
+			else if (ex is AggregateException aggregateException)
 			{
-				foreach (Exception innerException in aggException.InnerExceptions)
+				foreach (Exception innerException in aggregateException.InnerExceptions)
 				{
 					CompilationResult result = GetCompilationResult(innerException);
 					if (result != CompilationResult.OtherCompilationError)
