Initial work centralizing some Horde environment variables and dashboard routes used by various C# code.

Follow up CLs will contain changes to relevant project and engine code to make use of these.

#Synced-CL 48349895

diff --git a/Engine/Source/Programs/Shared/EpicGames.Horde/Utilities/HordeEnvVars.cs b/Engine/Source/Programs/Shared/EpicGames.Horde/Utilities/HordeEnvVars.cs
new file mode 100644
--- /dev/null
+++ b/Engine/Source/Programs/Shared/EpicGames.Horde/Utilities/HordeEnvVars.cs
@@ -0,0 +1,60 @@
+﻿// Copyright Epic Games, Inc. All Rights Reserved.
+
+using System;
+
+// The class members contained here are all self-explanatory enough that public
+// comments on each of them are unnecessary.
+#pragma warning disable CS1591 // Missing XML comment for publicly visible type or member
+
+namespace EpicGames.Horde.Utilities;
+
+/// <summary>
+/// Contains the environment variable names for common Horde environment
+/// variables.
+///
+/// These are meant to be used by any and all C# code that is NOT Horde itself.
+/// If these keys were used by Horde itself then as soon as a changes to them
+/// was committed our tooling would break until Horde was actually deployed. 
+/// </summary>
+public static class HordeEnvKeys
+{
+	public const string BatchId = "UE_HORDE_BATCHID";
+	public const string CleanupScriptPath = "UE_HORDE_CLEANUP";
+	public const string GraphUpdatePath = "UE_HORDE_GRAPH_UPDATE";
+	public const string HordeToken = "UE_HORDE_TOKEN";
+	public const string HordeUrl = "UE_HORDE_URL";
+	public const string JobId = "UE_HORDE_JOBID";
+	public const string LeaseCleanupScriptPath = "UE_HORDE_LEASE_CLEANUP";
+	public const string SharedDirPath = "UE_HORDE_SHARED_DIR";
+	public const string StepId = "UE_HORDE_STEPID";
+	public const string StepName = "UE_HORDE_STEPNAME";
+	public const string StreamId = "UE_HORDE_STREAMID";
+	public const string TemplateId = "UE_HORDE_TEMPLATEID";
+	public const string TemplateName = "UE_HORDE_TEMPLATENAME";
+}
+
+/// <summary>
+/// Contains static references to common Horde environment variables.
+///
+/// Note that for each of these we use expression bodied members so that the
+/// value of the environment variable is read again each time. Ideally no one
+/// would be changing any of these environment variables and this wouldn't be
+/// required, however we have no way to guarantee that doesn't happen (and
+/// in fact have seen cases where it does) and so we do this to be safe.
+/// </summary>
+public static class HordeEnvVars
+{
+	public static string? BatchId => Environment.GetEnvironmentVariable(HordeEnvKeys.BatchId);
+	public static string? CleanupScriptPath => Environment.GetEnvironmentVariable(HordeEnvKeys.CleanupScriptPath);
+	public static string? GraphUpdatePath => Environment.GetEnvironmentVariable(HordeEnvKeys.GraphUpdatePath);
+	public static string? HordeToken => Environment.GetEnvironmentVariable(HordeEnvKeys.HordeToken);
+	public static string? HordeUrl => Environment.GetEnvironmentVariable(HordeEnvKeys.HordeUrl);
+	public static string? JobId => Environment.GetEnvironmentVariable(HordeEnvKeys.JobId);
+	public static string? LeaseCleanupScriptPath => Environment.GetEnvironmentVariable(HordeEnvKeys.LeaseCleanupScriptPath);
+	public static string? SharedDirPath => Environment.GetEnvironmentVariable(HordeEnvKeys.SharedDirPath);
+	public static string? StepId => Environment.GetEnvironmentVariable(HordeEnvKeys.StepId);
+	public static string? StepName => Environment.GetEnvironmentVariable(HordeEnvKeys.StepName);
+	public static string? StreamId => Environment.GetEnvironmentVariable(HordeEnvKeys.StreamId);
+	public static string? TemplateId => Environment.GetEnvironmentVariable(HordeEnvKeys.TemplateId);
+	public static string? TemplateName => Environment.GetEnvironmentVariable(HordeEnvKeys.TemplateName);
+}
\ No newline at end of file
diff --git a/Engine/Source/Programs/Shared/EpicGames.Horde/Utilities/HordeRoutes.cs b/Engine/Source/Programs/Shared/EpicGames.Horde/Utilities/HordeRoutes.cs
new file mode 100644
--- /dev/null
+++ b/Engine/Source/Programs/Shared/EpicGames.Horde/Utilities/HordeRoutes.cs
@@ -0,0 +1,129 @@
+﻿// Copyright Epic Games, Inc. All Rights Reserved.
+
+using System;
+
+
+namespace EpicGames.Horde.Utilities;
+
+/// <summary>
+/// Provides simple access to different Horde dashboard routes. The primary use
+/// case is when adding a link to Horde that is inserted into a log or report.
+/// </summary>
+public static class HordeRoutes
+{
+	/// <summary>
+	/// Returns an HTTP url for a given Horde job.
+	///
+	/// Example: https://horde.example.com/job/690b24e0c619fc760780db81
+	/// </summary>
+	public static string JobPath(string hordeUrl, string jobId)
+	{
+		return $"{hordeUrl.TrimEnd('/')}/job/{jobId}";
+	}
+
+	/// <summary>
+	/// Returns an HTTP url for a Horde job with details taken from environment
+	/// variables. If any required environment variables are not set an empty
+	/// string will be returned instead.
+	///
+	/// Example: https://horde.example.com/job/690b24e0c619fc760780db81
+	/// </summary>
+	public static string JobPathFromEnv()
+	{
+		string url = HordeEnvVars.HordeUrl ?? String.Empty;
+		string jobId = HordeEnvVars.JobId ?? String.Empty;
+
+		if (String.IsNullOrEmpty(url) || String.IsNullOrEmpty(jobId))
+		{
+			return String.Empty;
+		}
+		
+		return JobPath(url, jobId);
+	}
+
+	/// <summary>
+	/// Returns an HTTP url for a given Horde step.
+	///
+	/// Example: https://horde.example.com/job/690b24e0c619fc760780db81?step=ab17
+	/// </summary>
+	public static string StepPath(string hordeUrl, string jobId, string stepId)
+	{
+		return new UriBuilder($"{hordeUrl.TrimEnd('/')}/job/{jobId}")
+		{
+			Query = $"step={Uri.EscapeDataString(stepId)}"
+		}.ToString();
+	}
+
+	/// <summary>
+	/// Returns an HTTP url for a Horde step with details taken from environment
+	/// variables. If any required environment variables are not set an empty
+	/// string will be returned instead.
+	///
+	/// Example: https://horde.example.com/job/690b24e0c619fc760780db81?step=ab17
+	/// </summary>
+	public static string StepPathFromEnv()
+	{
+		string url = HordeEnvVars.HordeUrl ?? String.Empty;
+		string jobId = HordeEnvVars.JobId ?? String.Empty;
+		string stepId = HordeEnvVars.StepId ?? String.Empty;
+
+		if (
+			String.IsNullOrEmpty(url)
+			|| String.IsNullOrEmpty(jobId)
+			|| String.IsNullOrEmpty(stepId)
+		)
+		{
+			return String.Empty;
+		}
+		
+		return StepPath(url, jobId, stepId);
+	}
+
+	/// <summary>
+	/// Returns an HTTP url for a given Horde test automation.
+	///
+	/// Example: https://horde.example.com/test-automation?job-690b24e0c619fc760780db81&amp;step=ab17
+	/// </summary>
+	public static string TestAutomationPath(string hordeUrl, string jobId, string stepId)
+	{
+		return new UriBuilder($"{hordeUrl.TrimEnd('/')}/test-automation")
+		{
+			Query = $"job={Uri.EscapeDataString(jobId)}&step={Uri.EscapeDataString(stepId)}"
+		}.ToString();
+	}
+
+	/// <summary>
+	/// Returns an HTTP url for a Horde test automation with details taken from
+	/// environment variables. If any required environment variables are not set
+	/// an empty string will be returned instead.
+	///
+	/// Example: https://horde.example.com/test-automation?job-690b24e0c619fc760780db81&amp;step=ab17
+	/// </summary>
+	public static string TestAutomationPathFromEnv()
+	{
+		string url = HordeEnvVars.HordeUrl ?? String.Empty;
+		string jobId = HordeEnvVars.JobId ?? String.Empty;
+		string stepId = HordeEnvVars.StepId ?? String.Empty;
+
+		if (
+			String.IsNullOrEmpty(url)
+			|| String.IsNullOrEmpty(jobId)
+			|| String.IsNullOrEmpty(stepId)
+		)
+		{
+			return String.Empty;
+		}
+		
+		return TestAutomationPath(url, jobId, stepId);
+	}
+
+	/// <summary>
+	/// Returns an HTTP url for a given Horde log
+	///
+	/// Example: https://horde.example.com/log/690b24e0c619fc760780db81?step=ab17
+	/// </summary>
+	public static string LogPath(string hordeUrl, string logId)
+	{
+		return $"{hordeUrl.TrimEnd('/')}/log/{logId}";
+	}
+}
\ No newline at end of file
