CloudStorage CreatedAt is part of metadata not result. Currently createdAt is also not a DateTime so parsing the String.

#Synced-CL 48173295

diff --git a/Engine/Source/Programs/Shared/EpicGames.ProjectStore/CloudStorage/CloudStorage.cs b/Engine/Source/Programs/Shared/EpicGames.ProjectStore/CloudStorage/CloudStorage.cs
--- a/Engine/Source/Programs/Shared/EpicGames.ProjectStore/CloudStorage/CloudStorage.cs
+++ b/Engine/Source/Programs/Shared/EpicGames.ProjectStore/CloudStorage/CloudStorage.cs
@@ -300,7 +300,7 @@ namespace EpicGames.ProjectStore
 					StringId.TryParse(new Utf8String(bucketIdString), 256, out StringId bucketId, out string? _);
 					string buildIdString = result["buildId"].AsString();
 					CbObjectId buildId = CbObjectId.Parse(buildIdString);
-					DateTime createdAt = result["createdAt"].AsDateTime();
+					DateTime createdAt = metadata["createdAt"].IsDateTime() ? metadata["createdAt"].AsDateTime() : DateTime.Parse(metadata["createdAt"].AsString());
 					CbField commitField = metadata["commit"];
 
 					string? commit = null;
