UnrealBuildTool: Conditional checks for dotnet 10

#rnx
#Synced-CL 47605596

diff --git a/Engine/Source/Programs/Shared/EpicGames.MsBuild/CsProjBuilder.cs b/Engine/Source/Programs/Shared/EpicGames.MsBuild/CsProjBuilder.cs
--- a/Engine/Source/Programs/Shared/EpicGames.MsBuild/CsProjBuilder.cs
+++ b/Engine/Source/Programs/Shared/EpicGames.MsBuild/CsProjBuilder.cs
@@ -136,19 +136,26 @@ namespace EpicGames.MsBuild
 						return;
 					}
 
-					if (!OperatingSystem.IsWindows())
-					{
-						// check the TargetFramework of the project: we can't build Windows-only projects on 
-						// non-Windows platforms.
-						if (project.GetProperty("TargetFramework").EvaluatedValue.Contains("windows", StringComparison.Ordinal))
-						{
-							logger.LogInformation("Skipping windows-only project {ProjectPath}", projectPath);
-							return;
-						}
-					}
-
+					// check the TargetFramework for unsupported versions. see https://learn.microsoft.com/en-us/dotnet/standard/frameworks
 					string targetFramework = project.GetProperty("TargetFramework").EvaluatedValue;
-					bool skipFramework = false;
+					bool skipFramework = targetFramework.Contains("netstandard", StringComparison.Ordinal)
+						| targetFramework.Contains("netcoreapp", StringComparison.Ordinal)
+						| targetFramework.Contains("uap", StringComparison.Ordinal)
+						| targetFramework.Contains("netnano", StringComparison.Ordinal)
+						| !targetFramework.Contains('.', StringComparison.Ordinal) // only .NET Core versions are supported and will always contain a period
+#if NET8_0_OR_GREATER
+						| targetFramework.Contains("net5.0", StringComparison.Ordinal)
+						| targetFramework.Contains("net6.0", StringComparison.Ordinal)
+						| targetFramework.Contains("net7.0", StringComparison.Ordinal)
+#endif
+#if NET9_0_OR_GREATER
+						| targetFramework.Contains("net8.0", StringComparison.Ordinal);
+#endif
+#if NET10_0_OR_GREATER
+						| targetFramework.Contains("net9.0", StringComparison.Ordinal);
+#endif
+					;
+
 					// check the TargetFramework of the project: we can't build Windows-only projects on non-Windows platforms etc.
 					if (targetFramework.Contains('-', StringComparison.Ordinal))
 					{
@@ -158,8 +165,6 @@ namespace EpicGames.MsBuild
 						platformSpecificProjects.Add(foundProject);
 					}
 
-					skipFramework |= targetFramework.Contains("net6.0", StringComparison.Ordinal);
-
 					if (skipFramework)
 					{
 						logger.LogInformation("Skipping {Framework} project {ProjectPath}", targetFramework, projectPath);
diff --git a/Engine/Source/Programs/UnrealBuildTool/ProjectFiles/ProjectFileGenerator.cs b/Engine/Source/Programs/UnrealBuildTool/ProjectFiles/ProjectFileGenerator.cs
--- a/Engine/Source/Programs/UnrealBuildTool/ProjectFiles/ProjectFileGenerator.cs
+++ b/Engine/Source/Programs/UnrealBuildTool/ProjectFiles/ProjectFileGenerator.cs
@@ -23,7 +23,6 @@ namespace UnrealBuildTool
 			PluginFiles = InPluginFiles;
 			ReferencingProjects = new List<FileReference>();
 			ModuleToSourceFiles = new Dictionary<FileReference, List<FileReference>>();
-
 		}
 
 		public List<FileReference> PluginFiles;
@@ -748,6 +747,13 @@ namespace UnrealBuildTool
 				FileReference.Combine(Unreal.EngineDirectory, "Source", "Programs", "UnrealBuildTool", "UnrealBuildTool.csproj")
 			};
 
+			string TargetFramework =
+#if NET10_0_OR_GREATER
+				"net10.0";
+#elif NET8_0_OR_GREATER
+				"net8.0";
+#endif
+
 			RulesAssembly? CurrentParent = RulesAssembly.Parent;
 			while (CurrentParent != null)
 			{
@@ -775,10 +781,10 @@ namespace UnrealBuildTool
 					"\" />");
 
 				Writer.WriteLine("  <PropertyGroup>");
-				Writer.WriteLine("    <TargetFramework>net8.0</TargetFramework>");
+				Writer.WriteLine($"    <TargetFramework>{TargetFramework}</TargetFramework>");
 				Writer.WriteLine("    <AppendTargetFrameworkToOutputPath>false</AppendTargetFrameworkToOutputPath>"); // Shorten intermediate filepath slightly
-				Writer.WriteLine("    <Configurations>Debug;Release;Development</Configurations>"); // VCSharpProject requires at least Debug & Development configurations
-				Writer.WriteLine("    <DefineConstants>$(DefineConstants);" + String.Join(';', RulesAssembly.PreprocessorDefines!) + "</DefineConstants>");
+				Writer.WriteLine("    <Configurations>Debug;Release;Development;Analyze</Configurations>"); // VCSharpProject requires at least Debug & Development configurations
+				Writer.WriteLine($"    <DefineConstants>$(DefineConstants);{String.Join(';', RulesAssembly.PreprocessorDefines!)}</DefineConstants>");
 				Writer.WriteLine("  </PropertyGroup>");
 
 				Writer.WriteLine("  <ItemGroup>");
