Fix a few remaining .editorconfig enabled issues in UBT referenced code

#rnx
#Synced-CL 47605596

diff --git a/Engine/Source/Programs/Shared/EpicGames.Build/UEBuildPlatformSDK.cs b/Engine/Source/Programs/Shared/EpicGames.Build/UEBuildPlatformSDK.cs
--- a/Engine/Source/Programs/Shared/EpicGames.Build/UEBuildPlatformSDK.cs
+++ b/Engine/Source/Programs/Shared/EpicGames.Build/UEBuildPlatformSDK.cs
@@ -468,7 +468,6 @@ namespace EpicGames.Core
 			// walk over the valid software ranges, potentially setting the current verison if one was passed in
 			// if there is only one known SDK named "Software" we allow any DeviceType
 			IEnumerable<SDKDescriptor> SoftwareSDKs = GetValidSoftwareVersions().Sdks;
-			bool bAllowAnyDeviceType = DeviceType == null || (SoftwareSDKs.Count() == 1 && SoftwareSDKs.First().Name == "Software");
 
 			foreach (SDKDescriptor Desc in SoftwareSDKs)
 			{
diff --git a/Engine/Source/Programs/Shared/UnrealBuildTool.Tests/TargetRulesTests.cs b/Engine/Source/Programs/Shared/UnrealBuildTool.Tests/TargetRulesTests.cs
--- a/Engine/Source/Programs/Shared/UnrealBuildTool.Tests/TargetRulesTests.cs
+++ b/Engine/Source/Programs/Shared/UnrealBuildTool.Tests/TargetRulesTests.cs
@@ -1,6 +1,5 @@
 // Copyright Epic Games, Inc. All Rights Reserved.
 
-using System;
 using System.Collections.Generic;
 using System.IO;
 using EpicGames.Core;
diff --git a/Engine/Source/Programs/Shared/UnrealBuildTool.Tests/TestUtilities/StaticTargetRulesInitializer.cs b/Engine/Source/Programs/Shared/UnrealBuildTool.Tests/TestUtilities/StaticTargetRulesInitializer.cs
--- a/Engine/Source/Programs/Shared/UnrealBuildTool.Tests/TestUtilities/StaticTargetRulesInitializer.cs
+++ b/Engine/Source/Programs/Shared/UnrealBuildTool.Tests/TestUtilities/StaticTargetRulesInitializer.cs
@@ -11,6 +11,7 @@ namespace UnrealBuildTool.Tests.TestUtilities
 	public static class StaticTargetRulesInitializer
 	{
 		[AssemblyInitialize]
+		[System.Diagnostics.CodeAnalysis.SuppressMessage("Style", "IDE0060:Remove unused parameter", Justification = "Required for UnitTesting")]
 		public static void InitializeBuildEnvironment(TestContext context)
 		{
 			// Clean any previous runs linger temp folder
diff --git a/Engine/Source/Programs/UnrealBuildTool/Configuration/Rules/ModuleRules.cs b/Engine/Source/Programs/UnrealBuildTool/Configuration/Rules/ModuleRules.cs
--- a/Engine/Source/Programs/UnrealBuildTool/Configuration/Rules/ModuleRules.cs
+++ b/Engine/Source/Programs/UnrealBuildTool/Configuration/Rules/ModuleRules.cs
@@ -1201,10 +1201,7 @@ namespace UnrealBuildTool
 
                 return false;
             }
-			set
-			{
-				bDisableAutoRTFMInstrumentationPrivate = value;
-			}
+			set => bDisableAutoRTFMInstrumentationPrivate = value;
 		}
 		private bool? bDisableAutoRTFMInstrumentationPrivate;
 
