Fix link order instability

#rnx

#Synced-CL 48364771

diff --git a/Engine/Source/Programs/Shared/EpicGames.UHT/Utils/UhtSession.cs b/Engine/Source/Programs/Shared/EpicGames.UHT/Utils/UhtSession.cs
--- a/Engine/Source/Programs/Shared/EpicGames.UHT/Utils/UhtSession.cs
+++ b/Engine/Source/Programs/Shared/EpicGames.UHT/Utils/UhtSession.cs
@@ -1,6 +1,7 @@
 // Copyright Epic Games, Inc. All Rights Reserved.
 
 using System;
+using System.Collections;
 using System.Collections.Concurrent;
 using System.Collections.Generic;
 using System.Data;
@@ -2282,7 +2283,7 @@ namespace EpicGames.UHT.Utils
 				Dictionary<UhtObject, int> linkableObjects = module.Packages.SelectMany(p => p.Children.OfType<UhtObject>().Where(o => o.AlternateObject is null))
 					// Objects whose outer is not a package but which need to be registered
 					.Concat(module.Headers.SelectMany(h => h.References.Singletons).Where(o => o.Package != o.Outer))
-					.OrderBy(o => o.LinkOrder)
+					.OrderBy(o => (o.LinkOrder, o.SourceName, o.Outer!.SourceName))
 					.Select((obj, index) => (obj, index))
 					.ToDictionary();
 
