Cloud Storage - Added Boost Worker Memory option which is enabled on all farm machines, this will increase the amount of memory used thus increasing memory requirments but reduces disk pressure.

#Synced-CL 47492933

diff --git a/Engine/Source/Programs/AutomationTool/BuildGraph/Tasks/RetrieveCloudArtifactTask.cs b/Engine/Source/Programs/AutomationTool/BuildGraph/Tasks/RetrieveCloudArtifactTask.cs
--- a/Engine/Source/Programs/AutomationTool/BuildGraph/Tasks/RetrieveCloudArtifactTask.cs
+++ b/Engine/Source/Programs/AutomationTool/BuildGraph/Tasks/RetrieveCloudArtifactTask.cs
@@ -397,6 +397,8 @@ namespace AutomationTool.Tasks
 					AllowMultipart = _parameters.AllowMultipart, 
 					AssumeHttp2 = String.Equals(httpVersion, "http2-only", StringComparison.OrdinalIgnoreCase),
 					BoostWorkers = CommandUtils.IsBuildMachine || _parameters.BoostWorkers,
+					// boost worker memory if boosting workers as well
+					BoostWorkerMemory = CommandUtils.IsBuildMachine || _parameters.BoostWorkers,
 					TraceFile = traceFile
 				}, optionsOverride: optionsOverride);
 		}
diff --git a/Engine/Source/Programs/Shared/EpicGames.ProjectStore/CloudStorage/CloudStorage.cs b/Engine/Source/Programs/Shared/EpicGames.ProjectStore/CloudStorage/CloudStorage.cs
--- a/Engine/Source/Programs/Shared/EpicGames.ProjectStore/CloudStorage/CloudStorage.cs
+++ b/Engine/Source/Programs/Shared/EpicGames.ProjectStore/CloudStorage/CloudStorage.cs
@@ -191,6 +191,7 @@ namespace EpicGames.ProjectStore
 			}
 			string http2Options = downloadOptions?.AssumeHttp2 ?? false ? "--assume-http2" : "";
 			string boostWorkersOption = downloadOptions?.BoostWorkers ?? false ? "--boost-workers" : "";
+			string boostWorkerMemoryOption = downloadOptions?.BoostWorkerMemory ?? false ? "--boost-worker-memory" : "";
 			string allowMultipartOption = downloadOptions?.AllowMultipart != null ? $"--allow-multipart={downloadOptions.AllowMultipart}" : "";
 			// trace file option needs to be passed before other arguments including the verbs
 			string traceOption = downloadOptions?.TraceFile != null ? $"--tracefile={downloadOptions.TraceFile.FullName} " : "";
@@ -199,7 +200,7 @@ namespace EpicGames.ProjectStore
 			string wildcardOption = includeFiles == null ? "" : $"--wildcard \"{String.Join(';', includeFiles.Select(p => ToZenWildcard(outputDir, p)))}\"";
 			string excludeOption =  excludeFiles == null ? "" : $"--exclude-wildcard \"{String.Join(';', excludeFiles.Select(p => ToZenWildcard(outputDir, p)))}\"";
 			// pass the access token via environment variable to avoid showing it on the cli
-			string cmdline = $"{traceOption}builds download {_options.Value.GetHostCliArgs(optionsOverride)} --namespace {ns} --bucket {bucketId} {wildcardOption} {excludeOption} {allowMultipartOption} {http2Options} {boostWorkersOption} {zenFolderPath} \"{outputDir}\" {buildId}";
+			string cmdline = $"{traceOption}builds download {_options.Value.GetHostCliArgs(optionsOverride)} --namespace {ns} --bucket {bucketId} {wildcardOption} {excludeOption} {allowMultipartOption} {http2Options} {boostWorkersOption} {boostWorkerMemoryOption} {zenFolderPath} \"{outputDir}\" {buildId}";
 
 			int exitCode = await RunZenCliAsync(cmdline, progress: progress, optionsOverride: optionsOverride, cancellationToken: cancellationToken);
 			EnsureZenExitCodeSuccess(exitCode);
diff --git a/Engine/Source/Programs/Shared/EpicGames.ProjectStore/CloudStorage/ICloudStorage.cs b/Engine/Source/Programs/Shared/EpicGames.ProjectStore/CloudStorage/ICloudStorage.cs
--- a/Engine/Source/Programs/Shared/EpicGames.ProjectStore/CloudStorage/ICloudStorage.cs
+++ b/Engine/Source/Programs/Shared/EpicGames.ProjectStore/CloudStorage/ICloudStorage.cs
@@ -199,6 +199,11 @@ namespace EpicGames.ProjectStore
 		/// </summary>
 		public bool BoostWorkers { get; set; }
 
+		/// <summary>
+		/// Set to increase memory used by workers when BoostWorkers is enabled to reduce disk pressure by buffering more in memory
+		/// </summary>
+		public bool BoostWorkerMemory { get; set; }
+
 		/// <summary>
 		/// Optional path to an utrace file, if specified Zen will output its utrace file here
 		/// </summary>
