Removed an unused method.

#rnx

#Synced-CL 48364771

diff --git a/Engine/Source/Programs/Shared/EpicGames.UHT/Exporters/CodeGen/UhtPackageCodeGenerator.cs b/Engine/Source/Programs/Shared/EpicGames.UHT/Exporters/CodeGen/UhtPackageCodeGenerator.cs
--- a/Engine/Source/Programs/Shared/EpicGames.UHT/Exporters/CodeGen/UhtPackageCodeGenerator.cs
+++ b/Engine/Source/Programs/Shared/EpicGames.UHT/Exporters/CodeGen/UhtPackageCodeGenerator.cs
@@ -245,42 +245,6 @@ namespace EpicGames.UHT.Exporters.CodeGen
 			return builder.AppendSingletonRef(provider, obj, UhtSingletonType.ConstInit);
 		}
 
-		/// <summary>
-		/// Append nullptr or an address-of expression for the constinit singleton of a built in type
-		/// </summary>
-		/// <param name="builder">Stringbuillder to append to</param>
-		/// <param name="session">Session object to retrieve built-in types from</param>
-		/// <param name="provider">Interface to access singleton names</param>
-		/// <param name="type">Object to append the singleton for</param>
-		/// <returns></returns>
-		internal static StringBuilder AppendConstInitSingletonRef(this StringBuilder builder, UhtSession session, IUhtSingletonNameProvider provider, UhtEngineType type)
-		{
-			switch (type)
-			{
-				case UhtEngineType.Package:
-					return builder.AppendSingletonRef(provider, session.UPackage, UhtSingletonType.ConstInit);
-				case UhtEngineType.Class:
-					return builder.AppendSingletonRef(provider, session.UClass, UhtSingletonType.ConstInit);
-				case UhtEngineType.Delegate:
-					return builder.AppendSingletonRef(provider, session.UDelegateFunction, UhtSingletonType.ConstInit);
-				case UhtEngineType.SparseDelegate:
-					return builder.AppendSingletonRef(provider, session.USparseDelegateFunction, UhtSingletonType.ConstInit);
-				case UhtEngineType.Enum:
-					return builder.AppendSingletonRef(provider, session.UEnum, UhtSingletonType.ConstInit);
-				case UhtEngineType.Function:
-					return builder.AppendSingletonRef(provider, session.UFunction, UhtSingletonType.ConstInit);
-				case UhtEngineType.Interface:
-					return builder.AppendSingletonRef(provider, session.UInterface, UhtSingletonType.ConstInit);
-				case UhtEngineType.ScriptStruct:
-					return builder.AppendSingletonRef(provider, session.UScriptStruct, UhtSingletonType.ConstInit);
-				case UhtEngineType.NativeInterface:
-					throw new NotImplementedException("ConstInit singleton for NativeInterface not implemented");
-				case UhtEngineType.Property:
-					throw new ArgumentException("Cannot insert singleton ref for property type because properties are not UObjects");
-				default:
-					throw new ArgumentException($"Invalid value ${type} for UhtEngineType");
-			}
-		}
 
 		internal static StringBuilder AppendLinkingMacros(this StringBuilder builder, UhtSession session, IEnumerable<UhtModule> modules)
 		{
