Cleaning up ununsed usings

#rnx

#Synced-CL 48364771

diff --git a/Engine/Source/Programs/Shared/EpicGames.UHT/Exporters/CodeGen/EditorDataStorageUhtExtension.cs b/Engine/Source/Programs/Shared/EpicGames.UHT/Exporters/CodeGen/EditorDataStorageUhtExtension.cs
--- a/Engine/Source/Programs/Shared/EpicGames.UHT/Exporters/CodeGen/EditorDataStorageUhtExtension.cs
+++ b/Engine/Source/Programs/Shared/EpicGames.UHT/Exporters/CodeGen/EditorDataStorageUhtExtension.cs
@@ -1,10 +1,8 @@
 ﻿// Copyright Epic Games, Inc. All Rights Reserved.
 
-using System;
 using System.Text;
 using EpicGames.UHT.Tables;
 using EpicGames.UHT.Utils;
-using EpicGames.UHT.Exporters.CodeGen;
 using EpicGames.UHT.Types;
 
 namespace EditorDataStorageUhtExtension
