Fixes for static analysis warnings in HordeEnvVars and HordeRoutes

#Synced-CL 48353684

diff --git a/Engine/Source/Programs/Shared/EpicGames.Horde/Utilities/HordeEnvVars.cs b/Engine/Source/Programs/Shared/EpicGames.Horde/Utilities/HordeEnvVars.cs
--- a/Engine/Source/Programs/Shared/EpicGames.Horde/Utilities/HordeEnvVars.cs
+++ b/Engine/Source/Programs/Shared/EpicGames.Horde/Utilities/HordeEnvVars.cs
@@ -48,7 +48,19 @@ public static class HordeEnvVars
 	public static string? CleanupScriptPath => Environment.GetEnvironmentVariable(HordeEnvKeys.CleanupScriptPath);
 	public static string? GraphUpdatePath => Environment.GetEnvironmentVariable(HordeEnvKeys.GraphUpdatePath);
 	public static string? HordeToken => Environment.GetEnvironmentVariable(HordeEnvKeys.HordeToken);
-	public static string? HordeUrl => Environment.GetEnvironmentVariable(HordeEnvKeys.HordeUrl);
+	public static Uri? HordeUrl
+	{
+		get
+		{
+			string? value = Environment.GetEnvironmentVariable(HordeEnvKeys.HordeUrl);
+			if (String.IsNullOrWhiteSpace(value))
+			{
+				return null;
+			}
+			
+			return Uri.TryCreate(value, UriKind.Absolute, out Uri? uri) ? uri : null;
+		}
+	}
 	public static string? JobId => Environment.GetEnvironmentVariable(HordeEnvKeys.JobId);
 	public static string? LeaseCleanupScriptPath => Environment.GetEnvironmentVariable(HordeEnvKeys.LeaseCleanupScriptPath);
 	public static string? SharedDirPath => Environment.GetEnvironmentVariable(HordeEnvKeys.SharedDirPath);
diff --git a/Engine/Source/Programs/Shared/EpicGames.Horde/Utilities/HordeRoutes.cs b/Engine/Source/Programs/Shared/EpicGames.Horde/Utilities/HordeRoutes.cs
--- a/Engine/Source/Programs/Shared/EpicGames.Horde/Utilities/HordeRoutes.cs
+++ b/Engine/Source/Programs/Shared/EpicGames.Horde/Utilities/HordeRoutes.cs
@@ -1,7 +1,7 @@
 ﻿// Copyright Epic Games, Inc. All Rights Reserved.
 
 using System;
-
+using EpicGames.Core;
 
 namespace EpicGames.Horde.Utilities;
 
@@ -12,118 +12,150 @@ namespace EpicGames.Horde.Utilities;
 public static class HordeRoutes
 {
 	/// <summary>
-	/// Returns an HTTP url for a given Horde job.
+	/// Returns a <see cref="System.Uri" /> for a given Horde job.
 	///
 	/// Example: https://horde.example.com/job/690b24e0c619fc760780db81
 	/// </summary>
-	public static string JobPath(string hordeUrl, string jobId)
+	/// <exception cref="ArgumentException">
+	/// Thrown when jobId is empty or just whitespace.
+	/// </exception>
+	public static Uri JobPath(Uri hordeUrl, string jobId)
 	{
-		return $"{hordeUrl.TrimEnd('/')}/job/{jobId}";
+		ArgumentException.ThrowIfNullOrWhiteSpace(jobId);
+		
+		return new Uri(hordeUrl, $"job/{jobId}");
 	}
 
 	/// <summary>
-	/// Returns an HTTP url for a Horde job with details taken from environment
-	/// variables. If any required environment variables are not set an empty
-	/// string will be returned instead.
+	/// Returns a <see cref="System.Uri"/> for a given Horde job with details
+	/// taken from environment variables.
+	///
+	/// Uses the following environment variables:
+	/// - <see cref="HordeEnvVars.HordeUrl" />
+	/// - <see cref="HordeEnvVars.JobId" />
 	///
 	/// Example: https://horde.example.com/job/690b24e0c619fc760780db81
 	/// </summary>
-	public static string JobPathFromEnv()
+	/// <exception cref="ArgumentException">
+	/// Thrown when required environment variables are not set.
+	/// </exception>
+	public static Uri JobPathFromEnv()
 	{
-		string url = HordeEnvVars.HordeUrl ?? String.Empty;
-		string jobId = HordeEnvVars.JobId ?? String.Empty;
-
-		if (String.IsNullOrEmpty(url) || String.IsNullOrEmpty(jobId))
-		{
-			return String.Empty;
-		}
+		Uri? url = HordeEnvVars.HordeUrl;
+		string? jobId = HordeEnvVars.JobId;
 		
+		ArgumentNullException.ThrowIfNull(url);
+		ArgumentException.ThrowIfNullOrWhiteSpace(jobId);
+
 		return JobPath(url, jobId);
 	}
 
 	/// <summary>
-	/// Returns an HTTP url for a given Horde step.
+	/// Returns a <see cref="System.Uri"/> for a given Horde step.
 	///
 	/// Example: https://horde.example.com/job/690b24e0c619fc760780db81?step=ab17
 	/// </summary>
-	public static string StepPath(string hordeUrl, string jobId, string stepId)
+	/// <exception cref="ArgumentException">
+	/// Thrown when jobId or stepId are empty or just whitespace.
+	/// </exception>
+	public static Uri StepPath(Uri hordeUrl, string jobId, string stepId)
 	{
-		return new UriBuilder($"{hordeUrl.TrimEnd('/')}/job/{jobId}")
-		{
-			Query = $"step={Uri.EscapeDataString(stepId)}"
-		}.ToString();
-	}
+		ArgumentException.ThrowIfNullOrWhiteSpace(jobId);
+		ArgumentException.ThrowIfNullOrWhiteSpace(stepId);
+		
+		QueryStringBuilder query = new QueryStringBuilder();
+		query.Add("step", stepId);
 
+		return new Uri(hordeUrl, $"job/{jobId}?{query}");
+	}
+	
 	/// <summary>
-	/// Returns an HTTP url for a Horde step with details taken from environment
-	/// variables. If any required environment variables are not set an empty
-	/// string will be returned instead.
+	/// Returns a <see cref="System.Uri"/> for a given Horde step with details
+	/// taken from environment variables.
+	///
+	/// Uses the following environment variables:
+	/// - <see cref="HordeEnvVars.HordeUrl" />
+	/// - <see cref="HordeEnvVars.JobId" />
+	/// - <see cref="HordeEnvVars.StepId" />
 	///
 	/// Example: https://horde.example.com/job/690b24e0c619fc760780db81?step=ab17
 	/// </summary>
-	public static string StepPathFromEnv()
+	/// <exception cref="ArgumentException">
+	/// Thrown when required environment variables are not set.
+	/// </exception>
+	public static Uri StepPathFromEnv()
 	{
-		string url = HordeEnvVars.HordeUrl ?? String.Empty;
-		string jobId = HordeEnvVars.JobId ?? String.Empty;
-		string stepId = HordeEnvVars.StepId ?? String.Empty;
-
-		if (
-			String.IsNullOrEmpty(url)
-			|| String.IsNullOrEmpty(jobId)
-			|| String.IsNullOrEmpty(stepId)
-		)
-		{
-			return String.Empty;
-		}
+		Uri? url = HordeEnvVars.HordeUrl;
+		string? jobId = HordeEnvVars.JobId;
+		string? stepId = HordeEnvVars.StepId;
+		
+		ArgumentNullException.ThrowIfNull(url);
+		ArgumentException.ThrowIfNullOrWhiteSpace(jobId);
+		ArgumentException.ThrowIfNullOrWhiteSpace(stepId);
 		
 		return StepPath(url, jobId, stepId);
 	}
 
 	/// <summary>
-	/// Returns an HTTP url for a given Horde test automation.
+	/// Returns a <see cref="System.Uri"/> for a given Horde test automation.
 	///
 	/// Example: https://horde.example.com/test-automation?job-690b24e0c619fc760780db81&amp;step=ab17
 	/// </summary>
-	public static string TestAutomationPath(string hordeUrl, string jobId, string stepId)
+	/// <exception cref="ArgumentException">
+	/// Thrown when jobId or stepId are empty or just whitespace.
+	/// </exception>
+	public static Uri TestAutomationPath(Uri hordeUrl, string jobId, string stepId)
 	{
-		return new UriBuilder($"{hordeUrl.TrimEnd('/')}/test-automation")
-		{
-			Query = $"job={Uri.EscapeDataString(jobId)}&step={Uri.EscapeDataString(stepId)}"
-		}.ToString();
+		ArgumentException.ThrowIfNullOrWhiteSpace(jobId);
+		ArgumentException.ThrowIfNullOrWhiteSpace(stepId);
+		
+		QueryStringBuilder query = new QueryStringBuilder();
+		query.Add("job", jobId);
+		query.Add("step", stepId);
+
+		return new Uri(hordeUrl, $"test-automation?{query}");
 	}
 
 	/// <summary>
-	/// Returns an HTTP url for a Horde test automation with details taken from
-	/// environment variables. If any required environment variables are not set
-	/// an empty string will be returned instead.
+	/// Returns a <see cref="System.Uri"/> for a given Horde test automation
+	/// with details taken from environment variables.
+	///
+	/// Uses the following environment variables:
+	/// - <see cref="HordeEnvVars.HordeUrl" />
+	/// - <see cref="HordeEnvVars.JobId" />
+	/// - <see cref="HordeEnvVars.StepId" />
+	///
 	///
 	/// Example: https://horde.example.com/test-automation?job-690b24e0c619fc760780db81&amp;step=ab17
 	/// </summary>
-	public static string TestAutomationPathFromEnv()
+	/// <exception cref="ArgumentException">
+	/// Thrown when required environment variables are not set.
+	/// </exception>
+	public static Uri TestAutomationPathFromEnv()
 	{
-		string url = HordeEnvVars.HordeUrl ?? String.Empty;
-		string jobId = HordeEnvVars.JobId ?? String.Empty;
-		string stepId = HordeEnvVars.StepId ?? String.Empty;
-
-		if (
-			String.IsNullOrEmpty(url)
-			|| String.IsNullOrEmpty(jobId)
-			|| String.IsNullOrEmpty(stepId)
-		)
-		{
-			return String.Empty;
-		}
+		Uri? url = HordeEnvVars.HordeUrl;
+		string? jobId = HordeEnvVars.JobId;
+		string? stepId = HordeEnvVars.StepId;
+		
+		ArgumentNullException.ThrowIfNull(url);
+		ArgumentException.ThrowIfNullOrWhiteSpace(jobId);
+		ArgumentException.ThrowIfNullOrWhiteSpace(stepId);
 		
 		return TestAutomationPath(url, jobId, stepId);
 	}
 
 	/// <summary>
-	/// Returns an HTTP url for a given Horde log
+	/// Returns a <see cref="System.Uri"/> for a given Horde log.
 	///
 	/// Example: https://horde.example.com/log/690b24e0c619fc760780db81?step=ab17
 	/// </summary>
-	public static string LogPath(string hordeUrl, string logId)
+	/// <exception cref="ArgumentException">
+	/// Thrown when logId is empty or just whitespace.
+	/// </exception>
+	public static Uri LogPath(Uri hordeUrl, string logId)
 	{
-		return $"{hordeUrl.TrimEnd('/')}/log/{logId}";
+		ArgumentException.ThrowIfNullOrWhiteSpace(logId);
+		
+		return new Uri(hordeUrl, $"log/{logId}");
 	}
 }
\ No newline at end of file
