Fix static code analysis errors in EpicGames.ProjectStore

#Synced-CL 47492933

diff --git a/Engine/Source/Programs/Shared/EpicGames.ProjectStore/CloudStorage/CloudStorage.cs b/Engine/Source/Programs/Shared/EpicGames.ProjectStore/CloudStorage/CloudStorage.cs
--- a/Engine/Source/Programs/Shared/EpicGames.ProjectStore/CloudStorage/CloudStorage.cs
+++ b/Engine/Source/Programs/Shared/EpicGames.ProjectStore/CloudStorage/CloudStorage.cs
@@ -7,10 +7,7 @@ using System.ComponentModel.DataAnnotations;
 using System.Diagnostics;
 using System.IO;
 using System.Linq;
-using System.Net;
 using System.Runtime.CompilerServices;
-using System.Text.Json;
-using System.Text.Json.Serialization;
 using System.Text.RegularExpressions;
 using System.Threading;
 using System.Threading.Tasks;
diff --git a/Engine/Source/Programs/Shared/EpicGames.ProjectStore/ProjectStoreTypes.cs b/Engine/Source/Programs/Shared/EpicGames.ProjectStore/ProjectStoreTypes.cs
--- a/Engine/Source/Programs/Shared/EpicGames.ProjectStore/ProjectStoreTypes.cs
+++ b/Engine/Source/Programs/Shared/EpicGames.ProjectStore/ProjectStoreTypes.cs
@@ -135,7 +135,7 @@ namespace EpicGames.ProjectStore
 		/// Delegate used to run command-line applications
 		/// </summary>
 		#pragma warning disable CA1051	
-		protected readonly RunAppDelegate _runApp;
+		protected readonly RunAppDelegate RunApp;
 		#pragma warning restore CA1051	
 
 		/// <summary>
@@ -144,7 +144,7 @@ namespace EpicGames.ProjectStore
 		/// <param name="runApp">Delegate used to run command-line applications.</param>
 		protected Build(RunAppDelegate runApp)
 		{
-			_runApp = runApp;
+			RunApp = runApp;
 		}
 
 		/// <summary>
@@ -241,7 +241,7 @@ namespace EpicGames.ProjectStore
 			StringBuilder projectCreateCommandLine = new StringBuilder();
 			projectCreateCommandLine.AppendFormat("project-create --force-update {0} {1} {2} {3} {4}", projectId, rootDirectory.FullName, engineDirectory.FullName, projectFile.Directory.FullName, projectFile.FullName);
 
-			_runApp(GetZenExecutable(engineDirectory).FullName, projectCreateCommandLine.ToString());
+			RunApp(GetZenExecutable(engineDirectory).FullName, projectCreateCommandLine.ToString());
 		}
 
 		/// <summary>
@@ -490,8 +490,8 @@ namespace EpicGames.ProjectStore
 				}
 
 				FileReference zenExe = GetZenExecutable(engineDirectory);
-				_runApp(zenExe.FullName, oplogCreateCommandLine.ToString());
-				_runApp(zenExe.FullName, oplogImportCommandLine.ToString());
+				RunApp(zenExe.FullName, oplogCreateCommandLine.ToString());
+				RunApp(zenExe.FullName, oplogImportCommandLine.ToString());
 			}
 		}
 
@@ -572,8 +572,8 @@ namespace EpicGames.ProjectStore
 			oplogImportCommandLine.AppendFormat(" --builds-id {0}", _buildData.BuildId);
 
 			FileReference zenExe = GetZenExecutable(engineDirectory);
-			_runApp(zenExe.FullName, oplogCreateCommandLine.ToString());
-			_runApp(zenExe.FullName, oplogImportCommandLine.ToString());
+			RunApp(zenExe.FullName, oplogCreateCommandLine.ToString());
+			RunApp(zenExe.FullName, oplogImportCommandLine.ToString());
 		}
 
 		/// <inheritdoc/>
@@ -594,14 +594,14 @@ namespace EpicGames.ProjectStore
 		/// <param name="runApp">A delegate used to run command-line apps.</param>
 		protected BuildIndex(RunAppDelegate runApp)
 		{
-			_runApp = runApp;
+			RunApp = runApp;
 		}
 
 		/// <summary>
 		/// A delegate used to run command-line apps. See documentation of RunAppDelegate for details.
 		/// </summary>
 		#pragma warning disable CA1051	
-		protected readonly RunAppDelegate _runApp;
+		protected readonly RunAppDelegate RunApp;
 		#pragma warning restore CA1051	
 
 		/// <summary>
@@ -728,7 +728,7 @@ namespace EpicGames.ProjectStore
 			}
 
 			DirectoryReference buildDir = DirectoryReference.Combine(buildsDir, changelist.ToString());
-			return new FileshareBuild(DirectoryReference.EnumerateFiles(buildDir).First(), foundChangelist, _runApp);
+			return new FileshareBuild(DirectoryReference.EnumerateFiles(buildDir).First(), foundChangelist, RunApp);
 		}
 	}
 
@@ -896,7 +896,7 @@ namespace EpicGames.ProjectStore
 			StringBuilder commandLine = new StringBuilder();
 
 			commandLine.AppendFormat(" --Service={0} --OutFile={1}", cloudAuthServiceName, outputFile.FullName);
-			_runApp(oidcExe.FullName, commandLine.ToString());
+			RunApp(oidcExe.FullName, commandLine.ToString());
 
 			string oidcOutput = FileReference.ReadAllText(outputFile);
 			AccessTokenOutput? token = JsonSerializer.Deserialize<AccessTokenOutput>(oidcOutput);
@@ -1074,7 +1074,7 @@ namespace EpicGames.ProjectStore
 				throw new BuildIndexException("Exact match for target changelist requested, but no snapshots matching that changelist were found");
 			}
 
-			return new UnrealCloudDDCBuild(builds.Last(), _host, _namespace, GetBucket(project, branch, platform, runtime), _accessToken, _runApp);
+			return new UnrealCloudDDCBuild(builds.Last(), _host, _namespace, GetBucket(project, branch, platform, runtime), _accessToken, RunApp);
 		}
 	}
 }
