EpicGames.LogMatchers: Fix issues and remove unneeded .editorconfig

#rnx
#Synced-CL 47605596

diff --git a/Engine/Source/Programs/Shared/EpicGames.LogMatchers/.editorconfig b/Engine/Source/Programs/Shared/EpicGames.LogMatchers/.editorconfig
deleted file mode 100644
diff --git a/Engine/Source/Programs/Shared/EpicGames.LogMatchers/SanitizerEventMatcher.cs b/Engine/Source/Programs/Shared/EpicGames.LogMatchers/SanitizerEventMatcher.cs
--- a/Engine/Source/Programs/Shared/EpicGames.LogMatchers/SanitizerEventMatcher.cs
+++ b/Engine/Source/Programs/Shared/EpicGames.LogMatchers/SanitizerEventMatcher.cs
@@ -129,36 +129,36 @@ namespace EpicGames.LogMatchers
 		/// <summary>
 		/// Add Value to Properties using (Regex) Group match and update Message with the corresponding key instead of value from Group.
 		/// </summary>
-		/// <param name="Value"></param>
-		/// <param name="Group"></param>
-		/// <param name="Properties"></param>
-		/// <param name="Message"></param>
-		/// <param name="Offset">Optional argument to offset Group index from Message if Message changed since Group was initialy matched</param>
-		private static void AddPropertyAndReplace(LogValue Value, Group Group, Dictionary<string, object> Properties, ref string Message, int Offset = 0)
+		/// <param name="value"></param>
+		/// <param name="group"></param>
+		/// <param name="properties"></param>
+		/// <param name="message"></param>
+		/// <param name="offset">Optional argument to offset Group index from Message if Message changed since Group was initialy matched</param>
+		private static void AddPropertyAndReplace(LogValue value, Group group, Dictionary<string, object> properties, ref string message, int offset = 0)
 		{
-			string Key = Value.Type.ToString();
-			Value.Text = Group.Value;
-			Properties.Add(Key, Value);
-			int Index = Group.Index + Offset;
-			Message = $"{Message.Substring(0, Index)}{{{Key}}}{Message.Substring(Index + Group.Length)}";
+			string key = value.Type.ToString();
+			value.Text = group.Value;
+			properties.Add(key, value);
+			int index = group.Index + offset;
+			message = $"{message.Substring(0, index)}{{{key}}}{message.Substring(index + group.Length)}";
 		}
 
 #pragma warning disable CA1045 // Do not pass types by reference
 		/// <summary>
 		/// Add Sanitizer Summary information to input Key/Value Pair Properties and update Message with Key markers
 		/// </summary>
-		/// <param name="Message"></param>
-		/// <param name="Properties"></param>
+		/// <param name="message"></param>
+		/// <param name="properties"></param>
 		/// <returns></returns>
-		public static bool AddSanitizerSummaryProperties(ref string Message, Dictionary<string, object> Properties)
+		public static bool AddSanitizerSummaryProperties(ref string message, Dictionary<string, object> properties)
 		{
-			Match Match = SummaryPattern.Match(Message);
-			if (Match.Success)
+			Match match = SummaryPattern.Match(message);
+			if (match.Success)
 			{
-				int InitialLength = Message.Length;
-				AddPropertyAndReplace(SanitizerIssueHandler.SanitizerName, Match.Groups["SanitizerName"], Properties, ref Message);
-				AddPropertyAndReplace(SanitizerIssueHandler.SummaryReason, Match.Groups["SummaryReason"], Properties, ref Message, Message.Length - InitialLength);
-				AddPropertyAndReplace(SanitizerIssueHandler.SummarySourceFile, Match.Groups["SourceFile"], Properties, ref Message, Message.Length - InitialLength);
+				int initialLength = message.Length;
+				AddPropertyAndReplace(SanitizerIssueHandler.SanitizerName, match.Groups["SanitizerName"], properties, ref message);
+				AddPropertyAndReplace(SanitizerIssueHandler.SummaryReason, match.Groups["SummaryReason"], properties, ref message, message.Length - initialLength);
+				AddPropertyAndReplace(SanitizerIssueHandler.SummarySourceFile, match.Groups["SourceFile"], properties, ref message, message.Length - initialLength);
 
 				return true;
 			}
@@ -225,11 +225,11 @@ namespace EpicGames.LogMatchers
 							bIsRaceDetectorTrace = sanitizerId == KnownLogEvents.Sanitizer_RaceDetector;
 
 							// These annotations are used by IssueHandler
-							Group SummaryReason = match.Groups["SummaryReason"];
-							if (SummaryReason.Success)
+							Group summaryReason = match.Groups["SummaryReason"];
+							if (summaryReason.Success)
 							{
 								// SummaryReason is only set on ReportLevel lines for non-LLVM sanitizers
-								builder.Annotate(SummaryReason, SanitizerIssueHandler.SummaryReason);
+								builder.Annotate(summaryReason, SanitizerIssueHandler.SummaryReason);
 							}
 							builder.Annotate(match.Groups["SanitizerName"], SanitizerIssueHandler.SanitizerName);
 						}
diff --git a/Engine/Source/Programs/Shared/EpicGames.LogMatchers/WorldLeakEventMatcher.cs b/Engine/Source/Programs/Shared/EpicGames.LogMatchers/WorldLeakEventMatcher.cs
--- a/Engine/Source/Programs/Shared/EpicGames.LogMatchers/WorldLeakEventMatcher.cs
+++ b/Engine/Source/Programs/Shared/EpicGames.LogMatchers/WorldLeakEventMatcher.cs
@@ -3,6 +3,7 @@
 using EpicGames.Core;
 using EpicGames.Horde.Issues.Handlers;
 using Microsoft.Extensions.Logging;
+using System;
 using System.Text.RegularExpressions;
 
 #pragma warning disable MA0016
@@ -91,7 +92,7 @@ namespace EpicGames.LogMatchers
 						// If we found the underline pattern, the current line should be our lingering reference
 						if (builder.Current.TryMatch(ReportLingeringReference, out match))
 						{
-							if (match.Groups.TryGetValue("ObjectText", out Group? group) && !string.IsNullOrEmpty(group.Value))
+							if (match.Groups.TryGetValue("ObjectText", out Group? group) && !String.IsNullOrEmpty(group.Value))
 							{
 								builder.Annotate(group, WorldLeakIssueHandler.LingeringReference);
 								bFoundLingeringReference = true;
@@ -105,11 +106,11 @@ namespace EpicGames.LogMatchers
 					else if (!bFoundFirstObject && builder.Current.TryMatch(ReportBodyPattern, out match))
 					{
 						// optional
-						if (match.Groups.TryGetValue("ObjectFlags", out Group? objectFlagGroup) && !string.IsNullOrEmpty(objectFlagGroup.Value))
+						if (match.Groups.TryGetValue("ObjectFlags", out Group? objectFlagGroup) && !String.IsNullOrEmpty(objectFlagGroup.Value))
 						{
 							builder.TryAnnotate(objectFlagGroup, WorldLeakIssueHandler.ObjectFlags);
 						}
-						if (match.Groups.TryGetValue("ObjectText", out Group? objectTextGroup) && !string.IsNullOrEmpty(objectTextGroup.Value))
+						if (match.Groups.TryGetValue("ObjectText", out Group? objectTextGroup) && !String.IsNullOrEmpty(objectTextGroup.Value))
 						{
 							builder.TryAnnotate(objectTextGroup, WorldLeakIssueHandler.ObjectText);
 							bFoundFirstObject = true;
diff --git a/Engine/Source/Programs/Shared/EpicGames.LogMatchers/ZenCliEventMatcher.cs b/Engine/Source/Programs/Shared/EpicGames.LogMatchers/ZenCliEventMatcher.cs
--- a/Engine/Source/Programs/Shared/EpicGames.LogMatchers/ZenCliEventMatcher.cs
+++ b/Engine/Source/Programs/Shared/EpicGames.LogMatchers/ZenCliEventMatcher.cs
@@ -11,13 +11,13 @@ namespace EpicGames.LogMatchers
 	/// </summary>
 	class ZenCliEventMatcher : ILogEventMatcher
 	{
-		static readonly Regex s_HttpOperationErrorPattern = new Regex(
+		static readonly Regex s_httpOperationErrorPattern = new Regex(
 			@"Error: Operation failed due to a http error: (?<message>.*)\((?<http_error>\d+)\)");
 
-		static readonly Regex s_HttpRequestErrorPattern = new Regex(
+		static readonly Regex s_httpRequestErrorPattern = new Regex(
 			@"HttpClient request failed.*?""error_code"": (?<http_error>\d+)");
 
-		static readonly Regex s_HttpRequestJsonErrorPattern = new Regex(
+		static readonly Regex s_httpRequestJsonErrorPattern = new Regex(
 			@"HttpClient request failed.*? Status: (?<http_error>\d+), Error: \'");
 
 		public LogEventMatch? Match(ILogCursor input)
@@ -25,7 +25,7 @@ namespace EpicGames.LogMatchers
 			// Error: Operation failed due to a http error: Failed putting build part: JupiterSession::PutBuildBlob: HTTP error 504 Gateway Timeout ({"error_code": 504}) (504)
 			// HttpClient request failed (session: guid): Url: url, Status: 504, Error: '' (0), Bytes: 0/19 (Up/Down), Elapsed: 1m00s, Response: '{\"error_code\": 504}', Reason: 'Gateway Time-out'
 			Match? httpMatch;
-			if (input.TryMatch(s_HttpOperationErrorPattern, out httpMatch) || input.TryMatch(s_HttpRequestErrorPattern, out httpMatch) || input.TryMatch(s_HttpRequestJsonErrorPattern, out httpMatch))
+			if (input.TryMatch(s_httpOperationErrorPattern, out httpMatch) || input.TryMatch(s_httpRequestErrorPattern, out httpMatch) || input.TryMatch(s_httpRequestJsonErrorPattern, out httpMatch))
 			{
 				LogEventBuilder builder = new LogEventBuilder(input);
 				string httpError = httpMatch.Groups["http_error"].Value;
