EpicGames.Perforce: Add support for snapshot-based testing

diff --git a/Engine/Source/Programs/Shared/EpicGames.Perforce.Fixture/PerforceSnapshot.cs b/Engine/Source/Programs/Shared/EpicGames.Perforce.Fixture/PerforceSnapshot.cs
new file mode 100644
--- /dev/null
+++ b/Engine/Source/Programs/Shared/EpicGames.Perforce.Fixture/PerforceSnapshot.cs
@@ -0,0 +1,249 @@
+﻿// Copyright Epic Games, Inc. All Rights Reserved.
+
+using System;
+using System.Collections.Generic;
+using System.IO;
+using System.Linq;
+using System.Reflection;
+using System.Text.Json;
+using System.Threading;
+using System.Threading.Tasks;
+using EpicGames.Core;
+using Microsoft.Extensions.Logging.Abstractions;
+using Microsoft.VisualStudio.TestTools.UnitTesting;
+
+#pragma warning disable CA1307 // Specify StringComparison for clarity
+
+namespace EpicGames.Perforce.Fixture;
+
+public record ServerSnapshot(Dictionary<string, StreamSnapshot> StreamSnapshots)
+{
+	public StreamSnapshot GetStream(string stream)
+	{
+		return StreamSnapshots.GetValueOrDefault(stream) ?? throw new Exception($"Unable to find stream {stream}");
+	}
+}
+
+public record StreamSnapshot(string Name, List<Changelist> Changelists, List<WorkspaceSnapshot> WorkspaceSnapshots)
+{
+	public WorkspaceSnapshot GetWorkspace(int cl, int shelvedCl = -1)
+	{
+		WorkspaceSnapshot? ws = WorkspaceSnapshots.Find(x => x.Change == cl && x.ShelvedChange == shelvedCl);
+		return ws ?? throw new Exception($"Unable to find CL {cl} with shelve {shelvedCl} in {Name}");
+	}
+}
+
+public record Changelist(int Change, string Description); // TODO: Files in changelist
+
+public record WorkspaceFile(
+	string Path,
+	long Size,
+	string Content,
+	string Digest,
+	int Revision
+)
+{
+	public static WorkspaceFile Create(string path, string content)
+	{
+		// Since content of files in fixture are only single lines,
+		// the trick below is to workaround differences in line endings after sync (client vs server).
+		content = content.Replace("\r\n", "\n", StringComparison.Ordinal);
+		long size = content.Length;
+		
+		return new WorkspaceFile(path, size, content, PerforceFixture.CalcMd5(content).ToUpperInvariant(), -1);
+	}
+}
+
+public record WorkspaceSnapshot(int Change, int ShelvedChange, HashSet<WorkspaceFile> Files)
+{
+	/// <summary>
+	/// Assert local directory contains exactly the files described by stream
+	/// </summary>
+	/// <param name="clientRoot">Client/workspace root directory in the local file system</param>
+	public void AssertDepotFiles(string clientRoot)
+	{
+		//FileSet snapshotFiles = new (Files.Select(x => (x.Path, x.Size, x.Digest)));
+		HashSet<WorkspaceFile> localFiles = PerforceSnapshot.GetLocalFileSet(clientRoot);
+		AssertFileSets(localFiles, Files);
+	}
+	
+	/// <summary>
+	/// Assert directory contains exactly the specified file set
+	/// </summary>
+	/// <param name="actual">Actual file set on disk</param>
+	/// <param name="expected">Expected file set to exist on disk</param>
+	public static void AssertFileSets(HashSet<WorkspaceFile> actual, HashSet<WorkspaceFile> expected)
+	{
+		if (!expected.SetEquals(actual))
+		{
+			List<WorkspaceFile> expectedList = new(expected);
+			List<WorkspaceFile> actualList = new(actual);
+
+			expectedList.Sort((a, b) => String.Compare(a.Path, b.Path, StringComparison.Ordinal));
+			actualList.Sort((a, b) => String.Compare(a.Path, b.Path, StringComparison.Ordinal));
+
+			string ReplaceLineEndings(string s) => s.Replace("\n", "\\n").Replace("\r", "\\r");
+			
+			Console.WriteLine("Expected " + new string('-', 90));
+			foreach (WorkspaceFile wf in expectedList)
+			{
+				Console.WriteLine($"{wf.Path,-20} | {wf.Size,5} | {wf.Digest} | {ReplaceLineEndings(wf.Content)}");
+			}
+
+			Console.WriteLine("");
+			Console.WriteLine("Actual " + new string('-', 92));
+			foreach (WorkspaceFile wf in actualList)
+			{
+				Console.WriteLine($"{wf.Path,-20} | {wf.Size,5} | {wf.Digest} | {ReplaceLineEndings(wf.Content)}");
+			}
+
+			Assert.Fail("Files on disk does not match files in stream at given CL");
+		}
+	}
+}
+
+public record WorkspaceState(string Stream, int Cl, int ShelvedCl = 0, bool RemoveUntracked = true);
+
+public class PerforceSnapshot(PerforceConnection perforce)
+{
+	public const string SnapshotJsonFile = "snapshot.json"; // Also referenced in .csproj for inclusion
+
+	public const string FooMain = "//Foo/Main";
+	public const string FooRelease = "//Foo/Release";
+	public const string FooMainSparseDev = "//Foo/Main-SparseDev";
+	public const string FooMainSparseRel = "//Foo/Main-SparseRel";
+	
+	public static readonly IReadOnlyList<WorkspaceState> FooMainStates = [
+		// Normal changelists
+		Main(2), Main(3), Main(4), Main(5), Main(6), Main(7), Main(9),
+			
+		// Shelved changelists
+		Main(2, 8), Main(3, 8), Main(7, 8), Main(9, 8)
+	];
+		
+	public static readonly IReadOnlyList<WorkspaceState> FooReleaseStates = [Release(11), Release(12)];
+	public static readonly IReadOnlyList<WorkspaceState> FooMainSparseDevStates = [SparseDev(9), SparseDev(13)];
+	
+	private static readonly JsonSerializerOptions s_jsonSerializerOptions = new () { WriteIndented = true };
+	private static readonly Lazy<ServerSnapshot> s_serverSnapshot = new(LoadSnapshot);
+
+	public static ServerSnapshot GetSnapshot() => s_serverSnapshot.Value;
+	public static WorkspaceState Main(int cl, int shelvedCl = -1) => new (FooMain, cl, shelvedCl);
+	public static WorkspaceState Release(int cl, int shelvedCl = -1) => new (FooRelease, cl, shelvedCl);
+	public static WorkspaceState SparseDev(int cl, int shelvedCl = -1) => new (FooMainSparseDev, cl, shelvedCl);
+	
+	private static ServerSnapshot LoadSnapshot()
+	{
+		Assembly assembly = typeof(PerforceSnapshot).Assembly;
+		const string ResourceName = "EpicGames.Perforce.Fixture." + SnapshotJsonFile;
+		using Stream? stream = assembly.GetManifestResourceStream(ResourceName);
+		if (stream == null)
+		{
+			throw new FileNotFoundException($"Embedded resource not found: {ResourceName}");
+		}
+			
+		ServerSnapshot? snapshot = JsonSerializer.Deserialize<ServerSnapshot>(stream, s_jsonSerializerOptions);
+		return snapshot ?? throw new Exception("Unable to JSON deserialize snapshot");
+	}
+	
+	public async Task TakeSnapshotsAsync(string snapshotJsonFile)
+	{
+		Dictionary<string, StreamSnapshot> streamSnapshots = new();
+		async Task AddStream(IReadOnlyList<WorkspaceState> states)
+		{
+			string streamName = states[0].Stream;
+			streamSnapshots[streamName] = await TakeStreamSnapshotAsync(streamName, states);
+		}
+
+		await AddStream(FooMainStates);
+		await AddStream(FooReleaseStates);
+		await AddStream(FooMainSparseDevStates);
+		string data = JsonSerializer.Serialize(new ServerSnapshot(streamSnapshots), s_jsonSerializerOptions);
+		await File.WriteAllTextAsync(snapshotJsonFile, data);
+	}
+
+	private async Task<StreamSnapshot> TakeStreamSnapshotAsync(string stream, IReadOnlyList<WorkspaceState> streamStates)
+	{
+		List<WorkspaceSnapshot> workspaceSnapshots = [];
+		foreach (WorkspaceState ss in streamStates)
+		{
+			workspaceSnapshots.Add(await TakeWorkspaceSnapshotAsync(ss));
+		}
+		return new StreamSnapshot(stream, [], workspaceSnapshots);
+	}
+
+	private async Task<WorkspaceSnapshot> TakeWorkspaceSnapshotAsync(WorkspaceState ss)
+	{
+		Console.WriteLine($"Syncing {ss.Stream} @ {ss.Cl} with shelve {ss.ShelvedCl} ...\n");
+		string guid = Guid.NewGuid().ToString()[..8];
+		string pathFriendlyName = $"{ss.Stream.Replace("/", "", StringComparison.Ordinal)}-{ss.Cl}-{ss.ShelvedCl}";
+		DirectoryReference dir = new (Path.Join(Path.GetTempPath(), $"p4-snapshot-{guid}-{pathFriendlyName}"));
+		DirectoryReference.CreateDirectory(dir);
+
+		HashSet<WorkspaceFile> workspaceFiles = [];
+		(IPerforceConnection perforce, string clientRoot) = await CreateClientAsync(ss.Stream, CancellationToken.None);
+		await perforce.SyncAsync($"//...@{ss.Cl}").ToListAsync();
+		
+		if (ss.ShelvedCl > 0)
+		{
+			await perforce.UnshelveAsync(
+				ss.ShelvedCl, -1, null, null, null,
+				UnshelveOptions.ForceOverwrite, Array.Empty<string>(), CancellationToken.None);
+		}
+
+		List<string> localFiles = Directory.EnumerateFiles(clientRoot, "*", SearchOption.AllDirectories).ToList();
+		foreach (string absPath in localFiles)
+		{
+			string localFile = Path.GetRelativePath(clientRoot, absPath).Replace('\\', '/');
+			string content = await File.ReadAllTextAsync(absPath);
+			workspaceFiles.Add(WorkspaceFile.Create(localFile, content));
+		}
+		
+		return new WorkspaceSnapshot(ss.Cl, ss.ShelvedCl, workspaceFiles);
+	}
+	
+	public static HashSet<WorkspaceFile> GetLocalFileSet(string clientRoot)
+	{
+		EnumerationOptions options = new() { RecurseSubdirectories = true };
+
+		HashSet<WorkspaceFile> localFiles = new(Directory.EnumerateFiles(clientRoot, "*", options)
+			.Select(x => Path.GetRelativePath(clientRoot, x))
+			.Select(x => x.Replace("\\", "/", StringComparison.Ordinal))
+			.Select(clientFile =>
+			{
+				string absPath = Path.Join(clientRoot, clientFile);
+				string content = File.ReadAllText(absPath);
+
+				// Since content of files in fixture are only single lines, the trick below works to workaround
+				// differences in line endings after sync (client vs server).
+				content = content.Replace("\r\n", "\n", StringComparison.Ordinal);
+				long size = content.Length;
+				// return WorkspaceFile.C(clientFile, size, PerforceFixture.CalcMd5(content).ToUpperInvariant());
+				return WorkspaceFile.Create(clientFile, content);
+			}));
+
+		return localFiles;
+	}
+
+	private async Task<(IPerforceConnection perforce, string clientRoot)> CreateClientAsync(string stream, CancellationToken ct)
+	{
+		string guid = Guid.NewGuid().ToString()[..8];
+		string clientName = $"p4-snapshot-{guid}";
+		string clientRoot = Path.Join(Path.GetTempPath(), clientName);
+		Directory.CreateDirectory(clientRoot);
+		
+		ClientRecord newClient = new (clientName, perforce.Settings.UserName, clientRoot)
+		{
+			Host = Environment.MachineName,
+			Type = "partitioned",
+			Stream = stream,
+			Options = ClientOptions.Clobber
+		};
+		Directory.CreateDirectory(newClient.Root);
+		await perforce.TryDeleteClientAsync(DeleteClientOptions.None, clientName, ct);
+		await perforce.CreateClientAsync(newClient, ct);
+		
+		PerforceSettings perforceSettings = new (perforce.Settings) { ClientName = clientName, PreferNativeClient = true };
+		return (await PerforceConnection.CreateAsync(perforceSettings, NullLogger.Instance), clientRoot);
+	}
+}
\ No newline at end of file
diff --git a/Engine/Source/Programs/Shared/EpicGames.Perforce.Fixture/Program.cs b/Engine/Source/Programs/Shared/EpicGames.Perforce.Fixture/Program.cs
new file mode 100644
--- /dev/null
+++ b/Engine/Source/Programs/Shared/EpicGames.Perforce.Fixture/Program.cs
@@ -0,0 +1,29 @@
+﻿// Copyright Epic Games, Inc. All Rights Reserved.
+
+using System;
+using System.Threading.Tasks;
+using EpicGames.Core;
+
+namespace EpicGames.Perforce.Fixture;
+
+public static class Program
+{
+	public static async Task<int> Main()
+	{
+		Console.WriteLine("Generate snapshot data from Perforce server for use in tests");
+		
+		string? serverUrl = Environment.GetEnvironmentVariable(BasePerforceFixtureTest.EnvVar);
+		if (serverUrl == null)
+		{
+			Console.WriteLine($"Set env var {BasePerforceFixtureTest.EnvVar} to configure Perforce connection settings");
+			return 1;
+		}
+		
+		string guid = Guid.NewGuid().ToString()[..8];
+		string clientName = "p4-snapshot-generator-" + guid;
+		using PerforceConnection perforce = BasePerforceFixtureTest.GetPerforceConnection(serverUrl, clientName, new DefaultConsoleLogger());
+		PerforceSnapshot snapshot = new (perforce);
+		await snapshot.TakeSnapshotsAsync(PerforceSnapshot.SnapshotJsonFile);
+		return 0;
+	}
+}
\ No newline at end of file
diff --git a/Engine/Source/Programs/Shared/EpicGames.Perforce.Fixture/snapshot.json b/Engine/Source/Programs/Shared/EpicGames.Perforce.Fixture/snapshot.json
new file mode 100644
--- /dev/null
+++ b/Engine/Source/Programs/Shared/EpicGames.Perforce.Fixture/snapshot.json
@@ -0,0 +1,609 @@
+{
+  "StreamSnapshots": {
+    "//Foo/Main": {
+      "Name": "//Foo/Main",
+      "Changelists": [],
+      "WorkspaceSnapshots": [
+        {
+          "Change": 2,
+          "ShelvedChange": -1,
+          "Files": [
+            {
+              "Path": "common.h",
+              "Size": 27,
+              "Content": "This is change common.h #1\n",
+              "Digest": "66521760C323E061BC3E9C9A2C67CD13",
+              "Revision": -1
+            },
+            {
+              "Path": "main.cpp",
+              "Size": 27,
+              "Content": "This is change main.cpp #1\n",
+              "Digest": "4FE1DE51507BF699F8BCF087AFB27DBB",
+              "Revision": -1
+            },
+            {
+              "Path": "main.h",
+              "Size": 25,
+              "Content": "This is change main.h #1\n",
+              "Digest": "16A5187ED86AD1425C7108F7869477BA",
+              "Revision": -1
+            },
+            {
+              "Path": "unused.cpp",
+              "Size": 29,
+              "Content": "This is change unused.cpp #1\n",
+              "Digest": "B84DF012D6393DD890000B144F876F40",
+              "Revision": -1
+            },
+            {
+              "Path": "Data/data.txt",
+              "Size": 27,
+              "Content": "This is change data.txt #1\n",
+              "Digest": "6EB6D4624CA0D84C8A1AB46A42756D35",
+              "Revision": -1
+            }
+          ]
+        },
+        {
+          "Change": 3,
+          "ShelvedChange": -1,
+          "Files": [
+            {
+              "Path": "common.h",
+              "Size": 27,
+              "Content": "This is change common.h #1\n",
+              "Digest": "66521760C323E061BC3E9C9A2C67CD13",
+              "Revision": -1
+            },
+            {
+              "Path": "main.cpp",
+              "Size": 27,
+              "Content": "This is change main.cpp #2\n",
+              "Digest": "9364BBCA53127F5049DB72A33EDD5CB2",
+              "Revision": -1
+            },
+            {
+              "Path": "main.h",
+              "Size": 25,
+              "Content": "This is change main.h #1\n",
+              "Digest": "16A5187ED86AD1425C7108F7869477BA",
+              "Revision": -1
+            },
+            {
+              "Path": "unused.cpp",
+              "Size": 29,
+              "Content": "This is change unused.cpp #1\n",
+              "Digest": "B84DF012D6393DD890000B144F876F40",
+              "Revision": -1
+            },
+            {
+              "Path": "Data/data.txt",
+              "Size": 27,
+              "Content": "This is change data.txt #1\n",
+              "Digest": "6EB6D4624CA0D84C8A1AB46A42756D35",
+              "Revision": -1
+            }
+          ]
+        },
+        {
+          "Change": 4,
+          "ShelvedChange": -1,
+          "Files": [
+            {
+              "Path": "common.h",
+              "Size": 27,
+              "Content": "This is change common.h #1\n",
+              "Digest": "66521760C323E061BC3E9C9A2C67CD13",
+              "Revision": -1
+            },
+            {
+              "Path": "main.cpp",
+              "Size": 27,
+              "Content": "This is change main.cpp #2\n",
+              "Digest": "9364BBCA53127F5049DB72A33EDD5CB2",
+              "Revision": -1
+            },
+            {
+              "Path": "main.h",
+              "Size": 25,
+              "Content": "This is change main.h #1\n",
+              "Digest": "16A5187ED86AD1425C7108F7869477BA",
+              "Revision": -1
+            },
+            {
+              "Path": "Data/data.txt",
+              "Size": 27,
+              "Content": "This is change data.txt #1\n",
+              "Digest": "6EB6D4624CA0D84C8A1AB46A42756D35",
+              "Revision": -1
+            }
+          ]
+        },
+        {
+          "Change": 5,
+          "ShelvedChange": -1,
+          "Files": [
+            {
+              "Path": "main.cpp",
+              "Size": 27,
+              "Content": "This is change main.cpp #2\n",
+              "Digest": "9364BBCA53127F5049DB72A33EDD5CB2",
+              "Revision": -1
+            },
+            {
+              "Path": "main.h",
+              "Size": 25,
+              "Content": "This is change main.h #1\n",
+              "Digest": "16A5187ED86AD1425C7108F7869477BA",
+              "Revision": -1
+            },
+            {
+              "Path": "shared.h",
+              "Size": 27,
+              "Content": "This is change common.h #1\n",
+              "Digest": "66521760C323E061BC3E9C9A2C67CD13",
+              "Revision": -1
+            },
+            {
+              "Path": "Data/data.txt",
+              "Size": 27,
+              "Content": "This is change data.txt #1\n",
+              "Digest": "6EB6D4624CA0D84C8A1AB46A42756D35",
+              "Revision": -1
+            }
+          ]
+        },
+        {
+          "Change": 6,
+          "ShelvedChange": -1,
+          "Files": [
+            {
+              "Path": "main.cpp",
+              "Size": 27,
+              "Content": "This is change main.cpp #3\n",
+              "Digest": "FCB0EE3699D298EBFDBDD5646E076E59",
+              "Revision": -1
+            },
+            {
+              "Path": "main.h",
+              "Size": 25,
+              "Content": "This is change main.h #2\n",
+              "Digest": "6BBB5178B88ED28C895300AF0F362209",
+              "Revision": -1
+            },
+            {
+              "Path": "shared.h",
+              "Size": 27,
+              "Content": "This is change common.h #1\n",
+              "Digest": "66521760C323E061BC3E9C9A2C67CD13",
+              "Revision": -1
+            },
+            {
+              "Path": "Data/data.txt",
+              "Size": 27,
+              "Content": "This is change data.txt #1\n",
+              "Digest": "6EB6D4624CA0D84C8A1AB46A42756D35",
+              "Revision": -1
+            }
+          ]
+        },
+        {
+          "Change": 7,
+          "ShelvedChange": -1,
+          "Files": [
+            {
+              "Path": "main.cpp",
+              "Size": 27,
+              "Content": "This is change main.cpp #3\n",
+              "Digest": "FCB0EE3699D298EBFDBDD5646E076E59",
+              "Revision": -1
+            },
+            {
+              "Path": "main.h",
+              "Size": 25,
+              "Content": "This is change main.h #2\n",
+              "Digest": "6BBB5178B88ED28C895300AF0F362209",
+              "Revision": -1
+            },
+            {
+              "Path": "shared.h",
+              "Size": 27,
+              "Content": "This is change common.h #1\n",
+              "Digest": "66521760C323E061BC3E9C9A2C67CD13",
+              "Revision": -1
+            },
+            {
+              "Path": "Data/data.txt",
+              "Size": 27,
+              "Content": "This is change data.txt #1\n",
+              "Digest": "6EB6D4624CA0D84C8A1AB46A42756D35",
+              "Revision": -1
+            },
+            {
+              "Path": "Data/moredata.txt",
+              "Size": 31,
+              "Content": "This is change moredata.txt #1\n",
+              "Digest": "C8C9EE282B2B5BCB3F1F9DE0C4DF024A",
+              "Revision": -1
+            }
+          ]
+        },
+        {
+          "Change": 9,
+          "ShelvedChange": -1,
+          "Files": [
+            {
+              "Path": "main.cpp",
+              "Size": 27,
+              "Content": "This is change main.cpp #4\n",
+              "Digest": "6D8B09596E7B45C923FCE5A45662E20E",
+              "Revision": -1
+            },
+            {
+              "Path": "main.h",
+              "Size": 25,
+              "Content": "This is change main.h #4\n",
+              "Digest": "ADDF62E18A7E93D7F643FE4ACCCE2990",
+              "Revision": -1
+            },
+            {
+              "Path": "shared.h",
+              "Size": 27,
+              "Content": "This is change common.h #1\n",
+              "Digest": "66521760C323E061BC3E9C9A2C67CD13",
+              "Revision": -1
+            },
+            {
+              "Path": "Data/data.txt",
+              "Size": 27,
+              "Content": "This is change data.txt #1\n",
+              "Digest": "6EB6D4624CA0D84C8A1AB46A42756D35",
+              "Revision": -1
+            },
+            {
+              "Path": "Data/moredata.txt",
+              "Size": 31,
+              "Content": "This is change moredata.txt #1\n",
+              "Digest": "C8C9EE282B2B5BCB3F1F9DE0C4DF024A",
+              "Revision": -1
+            }
+          ]
+        },
+        {
+          "Change": 2,
+          "ShelvedChange": 8,
+          "Files": [
+            {
+              "Path": "common.h",
+              "Size": 27,
+              "Content": "This is change common.h #1\n",
+              "Digest": "66521760C323E061BC3E9C9A2C67CD13",
+              "Revision": -1
+            },
+            {
+              "Path": "main.h",
+              "Size": 25,
+              "Content": "This is change main.h #3\n",
+              "Digest": "B5423FE256B087FF11045710904FD235",
+              "Revision": -1
+            },
+            {
+              "Path": "shelved.cpp",
+              "Size": 30,
+              "Content": "This is change shelved.cpp #1\n",
+              "Digest": "0B0894F8B928D830559151063B72DCD7",
+              "Revision": -1
+            },
+            {
+              "Path": "unused.cpp",
+              "Size": 29,
+              "Content": "This is change unused.cpp #2\n",
+              "Digest": "9C381D85D336FB6B45842E808E43DCA4",
+              "Revision": -1
+            },
+            {
+              "Path": "Data/data.txt",
+              "Size": 27,
+              "Content": "This is change data.txt #1\n",
+              "Digest": "6EB6D4624CA0D84C8A1AB46A42756D35",
+              "Revision": -1
+            }
+          ]
+        },
+        {
+          "Change": 3,
+          "ShelvedChange": 8,
+          "Files": [
+            {
+              "Path": "common.h",
+              "Size": 27,
+              "Content": "This is change common.h #1\n",
+              "Digest": "66521760C323E061BC3E9C9A2C67CD13",
+              "Revision": -1
+            },
+            {
+              "Path": "main.h",
+              "Size": 25,
+              "Content": "This is change main.h #3\n",
+              "Digest": "B5423FE256B087FF11045710904FD235",
+              "Revision": -1
+            },
+            {
+              "Path": "shelved.cpp",
+              "Size": 30,
+              "Content": "This is change shelved.cpp #1\n",
+              "Digest": "0B0894F8B928D830559151063B72DCD7",
+              "Revision": -1
+            },
+            {
+              "Path": "unused.cpp",
+              "Size": 29,
+              "Content": "This is change unused.cpp #2\n",
+              "Digest": "9C381D85D336FB6B45842E808E43DCA4",
+              "Revision": -1
+            },
+            {
+              "Path": "Data/data.txt",
+              "Size": 27,
+              "Content": "This is change data.txt #1\n",
+              "Digest": "6EB6D4624CA0D84C8A1AB46A42756D35",
+              "Revision": -1
+            }
+          ]
+        },
+        {
+          "Change": 7,
+          "ShelvedChange": 8,
+          "Files": [
+            {
+              "Path": "main.h",
+              "Size": 25,
+              "Content": "This is change main.h #3\n",
+              "Digest": "B5423FE256B087FF11045710904FD235",
+              "Revision": -1
+            },
+            {
+              "Path": "shared.h",
+              "Size": 27,
+              "Content": "This is change common.h #1\n",
+              "Digest": "66521760C323E061BC3E9C9A2C67CD13",
+              "Revision": -1
+            },
+            {
+              "Path": "shelved.cpp",
+              "Size": 30,
+              "Content": "This is change shelved.cpp #1\n",
+              "Digest": "0B0894F8B928D830559151063B72DCD7",
+              "Revision": -1
+            },
+            {
+              "Path": "unused.cpp",
+              "Size": 29,
+              "Content": "This is change unused.cpp #2\n",
+              "Digest": "9C381D85D336FB6B45842E808E43DCA4",
+              "Revision": -1
+            },
+            {
+              "Path": "Data/data.txt",
+              "Size": 27,
+              "Content": "This is change data.txt #1\n",
+              "Digest": "6EB6D4624CA0D84C8A1AB46A42756D35",
+              "Revision": -1
+            },
+            {
+              "Path": "Data/moredata.txt",
+              "Size": 31,
+              "Content": "This is change moredata.txt #1\n",
+              "Digest": "C8C9EE282B2B5BCB3F1F9DE0C4DF024A",
+              "Revision": -1
+            }
+          ]
+        },
+        {
+          "Change": 9,
+          "ShelvedChange": 8,
+          "Files": [
+            {
+              "Path": "main.h",
+              "Size": 25,
+              "Content": "This is change main.h #3\n",
+              "Digest": "B5423FE256B087FF11045710904FD235",
+              "Revision": -1
+            },
+            {
+              "Path": "shared.h",
+              "Size": 27,
+              "Content": "This is change common.h #1\n",
+              "Digest": "66521760C323E061BC3E9C9A2C67CD13",
+              "Revision": -1
+            },
+            {
+              "Path": "shelved.cpp",
+              "Size": 30,
+              "Content": "This is change shelved.cpp #1\n",
+              "Digest": "0B0894F8B928D830559151063B72DCD7",
+              "Revision": -1
+            },
+            {
+              "Path": "unused.cpp",
+              "Size": 29,
+              "Content": "This is change unused.cpp #2\n",
+              "Digest": "9C381D85D336FB6B45842E808E43DCA4",
+              "Revision": -1
+            },
+            {
+              "Path": "Data/data.txt",
+              "Size": 27,
+              "Content": "This is change data.txt #1\n",
+              "Digest": "6EB6D4624CA0D84C8A1AB46A42756D35",
+              "Revision": -1
+            },
+            {
+              "Path": "Data/moredata.txt",
+              "Size": 31,
+              "Content": "This is change moredata.txt #1\n",
+              "Digest": "C8C9EE282B2B5BCB3F1F9DE0C4DF024A",
+              "Revision": -1
+            }
+          ]
+        }
+      ]
+    },
+    "//Foo/Release": {
+      "Name": "//Foo/Release",
+      "Changelists": [],
+      "WorkspaceSnapshots": [
+        {
+          "Change": 11,
+          "ShelvedChange": -1,
+          "Files": [
+            {
+              "Path": "main.cpp",
+              "Size": 27,
+              "Content": "This is change main.cpp #2\n",
+              "Digest": "9364BBCA53127F5049DB72A33EDD5CB2",
+              "Revision": -1
+            },
+            {
+              "Path": "main.h",
+              "Size": 25,
+              "Content": "This is change main.h #1\n",
+              "Digest": "16A5187ED86AD1425C7108F7869477BA",
+              "Revision": -1
+            },
+            {
+              "Path": "shared.h",
+              "Size": 27,
+              "Content": "This is change common.h #1\n",
+              "Digest": "66521760C323E061BC3E9C9A2C67CD13",
+              "Revision": -1
+            }
+          ]
+        },
+        {
+          "Change": 12,
+          "ShelvedChange": -1,
+          "Files": [
+            {
+              "Path": "main.cpp",
+              "Size": 44,
+              "Content": "This is change main.cpp #2 in //Foo/Release\n",
+              "Digest": "16009C33B215A393F42B04E3A11EA0D2",
+              "Revision": -1
+            },
+            {
+              "Path": "main.h",
+              "Size": 25,
+              "Content": "This is change main.h #1\n",
+              "Digest": "16A5187ED86AD1425C7108F7869477BA",
+              "Revision": -1
+            },
+            {
+              "Path": "release.cpp",
+              "Size": 47,
+              "Content": "This is change release.cpp #1 in //Foo/Release\n",
+              "Digest": "B0AF06334A16CEEF8CCFF655BEA862FF",
+              "Revision": -1
+            },
+            {
+              "Path": "shared.h",
+              "Size": 27,
+              "Content": "This is change common.h #1\n",
+              "Digest": "66521760C323E061BC3E9C9A2C67CD13",
+              "Revision": -1
+            }
+          ]
+        }
+      ]
+    },
+    "//Foo/Main-SparseDev": {
+      "Name": "//Foo/Main-SparseDev",
+      "Changelists": [],
+      "WorkspaceSnapshots": [
+        {
+          "Change": 9,
+          "ShelvedChange": -1,
+          "Files": [
+            {
+              "Path": "main.cpp",
+              "Size": 27,
+              "Content": "This is change main.cpp #4\n",
+              "Digest": "6D8B09596E7B45C923FCE5A45662E20E",
+              "Revision": -1
+            },
+            {
+              "Path": "main.h",
+              "Size": 25,
+              "Content": "This is change main.h #4\n",
+              "Digest": "ADDF62E18A7E93D7F643FE4ACCCE2990",
+              "Revision": -1
+            },
+            {
+              "Path": "shared.h",
+              "Size": 27,
+              "Content": "This is change common.h #1\n",
+              "Digest": "66521760C323E061BC3E9C9A2C67CD13",
+              "Revision": -1
+            },
+            {
+              "Path": "Data/data.txt",
+              "Size": 27,
+              "Content": "This is change data.txt #1\n",
+              "Digest": "6EB6D4624CA0D84C8A1AB46A42756D35",
+              "Revision": -1
+            },
+            {
+              "Path": "Data/moredata.txt",
+              "Size": 31,
+              "Content": "This is change moredata.txt #1\n",
+              "Digest": "C8C9EE282B2B5BCB3F1F9DE0C4DF024A",
+              "Revision": -1
+            }
+          ]
+        },
+        {
+          "Change": 13,
+          "ShelvedChange": -1,
+          "Files": [
+            {
+              "Path": "main.cpp",
+              "Size": 27,
+              "Content": "This is change main.cpp #4\n",
+              "Digest": "6D8B09596E7B45C923FCE5A45662E20E",
+              "Revision": -1
+            },
+            {
+              "Path": "main.h",
+              "Size": 25,
+              "Content": "This is change main.h #4\n",
+              "Digest": "ADDF62E18A7E93D7F643FE4ACCCE2990",
+              "Revision": -1
+            },
+            {
+              "Path": "shared.h",
+              "Size": 27,
+              "Content": "This is change common.h #1\n",
+              "Digest": "66521760C323E061BC3E9C9A2C67CD13",
+              "Revision": -1
+            },
+            {
+              "Path": "Data/data.txt",
+              "Size": 27,
+              "Content": "This is change data.txt #1\n",
+              "Digest": "6EB6D4624CA0D84C8A1AB46A42756D35",
+              "Revision": -1
+            },
+            {
+              "Path": "Data/moredata.txt",
+              "Size": 31,
+              "Content": "This is change moredata.txt #1\n",
+              "Digest": "C8C9EE282B2B5BCB3F1F9DE0C4DF024A",
+              "Revision": -1
+            }
+          ]
+        }
+      ]
+    }
+  }
+}
\ No newline at end of file
