EpicGames.Core: Add SupportedOSPlatformGuard to some helper functions

#rnx
#Synced-CL 48614040

diff --git a/Engine/Source/Programs/Shared/EpicGames.Core/RuntimePlatform.cs b/Engine/Source/Programs/Shared/EpicGames.Core/RuntimePlatform.cs
--- a/Engine/Source/Programs/Shared/EpicGames.Core/RuntimePlatform.cs
+++ b/Engine/Source/Programs/Shared/EpicGames.Core/RuntimePlatform.cs
@@ -2,6 +2,7 @@
 
 using System;
 using System.Runtime.InteropServices;
+using System.Runtime.Versioning;
 
 namespace EpicGames.Core
 {
@@ -14,18 +15,21 @@ namespace EpicGames.Core
 		/// Whether we are currently running on Linux.
 		/// </summary>
 		[Obsolete("Replace with OperatingSystem.IsLinux()")]
+		[SupportedOSPlatformGuard("linux")]
 		public static bool IsLinux => OperatingSystem.IsLinux();
 
 		/// <summary>
 		/// Whether we are currently running on a MacOS platform.
 		/// </summary>
 		[Obsolete("Replace with OperatingSystem.IsMacOS()")]
+		[SupportedOSPlatformGuard("macos")]
 		public static bool IsMac => OperatingSystem.IsMacOS();
 
 		/// <summary>
 		/// Whether we are currently running a Windows platform.
 		/// </summary>
 		[Obsolete("Replace with OperatingSystem.IsWindows()")]
+		[SupportedOSPlatformGuard("windows")]
 		public static bool IsWindows => OperatingSystem.IsWindows();
 
 		/// <summary>
@@ -36,13 +40,13 @@ namespace EpicGames.Core
 			/// <summary>
 			/// Windows
 			/// </summary>
-			Windows, 
-			
+			Windows,
+
 			/// <summary>
 			/// Linux
 			/// </summary>
-			Linux, 
-			
+			Linux,
+
 			/// <summary>
 			/// Mac
 			/// </summary>
@@ -68,7 +72,7 @@ namespace EpicGames.Core
 		/// <summary>
 		/// The current runtime platform is Windows but we're running on Wine
 		/// </summary>
-        public static bool IsRunningOnWine()
+		public static bool IsRunningOnWine()
 		{
 			if (OperatingSystem.IsWindows())
 			{
