Resolve some analysis warnings in Rules

#rnx
#Synced-CL 48614040

diff --git a/Engine/Source/Programs/Horde/Samples/HordeTest/HordeTest.Build.cs b/Engine/Source/Programs/Horde/Samples/HordeTest/HordeTest.Build.cs
--- a/Engine/Source/Programs/Horde/Samples/HordeTest/HordeTest.Build.cs
+++ b/Engine/Source/Programs/Horde/Samples/HordeTest/HordeTest.Build.cs
@@ -1,24 +1,17 @@
 // Copyright Epic Games, Inc. All Rights Reserved.
 
+namespace UnrealBuildTool.Rules;
 
-// Copyright Epic Games, Inc. All Rights Reserved.
-
-using UnrealBuildTool;
-
-namespace UnrealBuildTool.Rules
+[SupportedPlatforms(UnrealPlatformClass.Desktop)]
+public class HordeTest : ModuleRules
 {
-	public class HordeTest : ModuleRules
+	public HordeTest(ReadOnlyTargetRules target) : base(target)
 	{
-		public HordeTest(ReadOnlyTargetRules Target) : base(Target)
-		{
-			PrivateDependencyModuleNames.AddRange(
-				new string[] {
-				"Core",
-				"Horde",
-				"DesktopPlatform",
-				"HTTP"
-				}
-			);
-		}
+		PrivateDependencyModuleNames.AddRange([
+			"Core",
+			"Horde",
+			"DesktopPlatform",
+			"HTTP"
+		]);
 	}
 }
\ No newline at end of file
diff --git a/Engine/Source/Programs/Shared/EpicGames.Perforce.Native/DotNetPerforceLib.Build.cs b/Engine/Source/Programs/Shared/EpicGames.Perforce.Native/DotNetPerforceLib.Build.cs
--- a/Engine/Source/Programs/Shared/EpicGames.Perforce.Native/DotNetPerforceLib.Build.cs
+++ b/Engine/Source/Programs/Shared/EpicGames.Perforce.Native/DotNetPerforceLib.Build.cs
@@ -1,57 +1,54 @@
 // Copyright Epic Games, Inc. All Rights Reserved.
 
+namespace UnrealBuildTool.Rules;
 
-// Copyright Epic Games, Inc. All Rights Reserved.
-
-namespace UnrealBuildTool.Rules
+[SupportedPlatforms(UnrealPlatformClass.Desktop)]
+public class DotNetPerforceLib : ModuleRules
 {
-	public class DotNetPerforceLib : ModuleRules
+	public DotNetPerforceLib(ReadOnlyTargetRules target) : base(target)
 	{
-		public DotNetPerforceLib(ReadOnlyTargetRules Target) : base(Target)
+		if (target.Platform == UnrealTargetPlatform.Win64)
 		{
-			if (Target.Platform == UnrealTargetPlatform.Win64)
-			{
-				PrivateDefinitions.Add("OS_NT");
-				PrivateDefinitions.Add("_CRT_NONSTDC_NO_DEPRECATE");
-			}
-
-			PublicIncludePathModuleNames.Add("Core");
+			PrivateDefinitions.Add("OS_NT");
+			PrivateDefinitions.Add("_CRT_NONSTDC_NO_DEPRECATE");
+		}
 
-			PrivateDependencyModuleNames.Add("OpenSSL");
+		PublicIncludePathModuleNames.Add("Core");
 
-			string ApiPath = Target.UEThirdPartySourceDirectory + "Perforce/p4api-2023.2/";
-			if (Target.Platform == UnrealTargetPlatform.Linux)
-			{
-				PublicSystemIncludePaths.Add(ApiPath + "Linux/include");
-				PublicAdditionalLibraries.Add(ApiPath + "Linux/lib/libclient.a");
-				PublicAdditionalLibraries.Add(ApiPath + "Linux/lib/librpc.a");
-				PublicAdditionalLibraries.Add(ApiPath + "Linux/lib/libsupp.a");
-				PublicAdditionalLibraries.Add(ApiPath + "Linux/lib/libp4script_cstub.a");
+		PrivateDependencyModuleNames.Add("OpenSSL");
 
-				PublicSystemLibraries.Add("dl");
-			}
-			else if (Target.Platform == UnrealTargetPlatform.Mac)
-			{
-				PublicIncludePaths.Add(ApiPath + "Mac/include");
-				PublicAdditionalLibraries.Add(ApiPath + "Mac/lib/libclient.a");
-				PublicAdditionalLibraries.Add(ApiPath + "Mac/lib/librpc.a");
-				PublicAdditionalLibraries.Add(ApiPath + "Mac/lib/libsupp.a");
-				PublicAdditionalLibraries.Add(ApiPath + "Mac/lib/libp4script_cstub.a");
+		string apiPath = Target.UEThirdPartySourceDirectory + "Perforce/p4api-2023.2/";
+		if (target.Platform == UnrealTargetPlatform.Linux)
+		{
+			PublicSystemIncludePaths.Add(apiPath + "Linux/include");
+			PublicAdditionalLibraries.Add(apiPath + "Linux/lib/libclient.a");
+			PublicAdditionalLibraries.Add(apiPath + "Linux/lib/librpc.a");
+			PublicAdditionalLibraries.Add(apiPath + "Linux/lib/libsupp.a");
+			PublicAdditionalLibraries.Add(apiPath + "Linux/lib/libp4script_cstub.a");
 
-				PublicFrameworks.Add("Foundation");
-				PublicFrameworks.Add("CoreFoundation");
-				PublicFrameworks.Add("CoreGraphics");
-				PublicFrameworks.Add("CoreServices");
-				PublicFrameworks.Add("Security");
-			}
-			else if (Target.Platform == UnrealTargetPlatform.Win64)
-			{
-				PublicIncludePaths.Add(ApiPath + "Win64/include");
-				PublicAdditionalLibraries.Add(ApiPath + "Win64/lib/libclient.lib");
-				PublicAdditionalLibraries.Add(ApiPath + "Win64/lib/librpc.lib");
-				PublicAdditionalLibraries.Add(ApiPath + "Win64/lib/libsupp.lib");
-				PublicAdditionalLibraries.Add(ApiPath + "Win64/lib/libp4script_cstub.lib");
-			}
+			PublicSystemLibraries.Add("dl");
+		}
+		else if (target.Platform == UnrealTargetPlatform.Mac)
+		{
+			PublicIncludePaths.Add(apiPath + "Mac/include");
+			PublicAdditionalLibraries.Add(apiPath + "Mac/lib/libclient.a");
+			PublicAdditionalLibraries.Add(apiPath + "Mac/lib/librpc.a");
+			PublicAdditionalLibraries.Add(apiPath + "Mac/lib/libsupp.a");
+			PublicAdditionalLibraries.Add(apiPath + "Mac/lib/libp4script_cstub.a");
+
+			PublicFrameworks.Add("Foundation");
+			PublicFrameworks.Add("CoreFoundation");
+			PublicFrameworks.Add("CoreGraphics");
+			PublicFrameworks.Add("CoreServices");
+			PublicFrameworks.Add("Security");
+		}
+		else if (target.Platform == UnrealTargetPlatform.Win64)
+		{
+			PublicIncludePaths.Add(apiPath + "Win64/include");
+			PublicAdditionalLibraries.Add(apiPath + "Win64/lib/libclient.lib");
+			PublicAdditionalLibraries.Add(apiPath + "Win64/lib/librpc.lib");
+			PublicAdditionalLibraries.Add(apiPath + "Win64/lib/libsupp.lib");
+			PublicAdditionalLibraries.Add(apiPath + "Win64/lib/libp4script_cstub.lib");
 		}
 	}
 }
\ No newline at end of file
diff --git a/Engine/Source/Programs/Shared/EpicGames.Perforce.Native/DotNetPerforceLib.Target.cs b/Engine/Source/Programs/Shared/EpicGames.Perforce.Native/DotNetPerforceLib.Target.cs
--- a/Engine/Source/Programs/Shared/EpicGames.Perforce.Native/DotNetPerforceLib.Target.cs
+++ b/Engine/Source/Programs/Shared/EpicGames.Perforce.Native/DotNetPerforceLib.Target.cs
@@ -2,12 +2,14 @@
 
 using UnrealBuildTool;
 
+[SupportedPlatforms(UnrealPlatformClass.Desktop)]
 [System.Diagnostics.CodeAnalysis.SuppressMessage("Design", "CA1050:Declare types in namespaces", Justification = "<Pending>")]
 public class DotNetPerforceLibTarget : TargetRules
 {
 	public DotNetPerforceLibTarget(TargetInfo target) : base(target)
 	{
 		Type = TargetType.Program;
+		IncludeOrderVersion = EngineIncludeOrderVersion.Latest;
 		LinkType = TargetLinkType.Monolithic;
 		LaunchModuleName = "DotNetPerforceLib";
 
