Move AppDomainExtensions from UnrealBuildTool to EpicGames.Build

#rnx
#Synced-CL 48635954

diff --git a/Engine/Source/Programs/Shared/EpicGames.Build/System/AppDomainExtensions.cs b/Engine/Source/Programs/Shared/EpicGames.Build/System/AppDomainExtensions.cs
new file mode 100644
--- /dev/null
+++ b/Engine/Source/Programs/Shared/EpicGames.Build/System/AppDomainExtensions.cs
@@ -0,0 +1,56 @@
+// Copyright Epic Games, Inc. All Rights Reserved.
+
+using System;
+using System.Collections.Generic;
+using System.Linq;
+using System.Reflection;
+
+namespace UnrealBuildBase
+{
+
+	/// <summary>
+	/// Extensions to AppDomain
+	/// </summary>
+	public static class AppDomainExtensions
+	{
+		/// <summary>
+		/// Using reflection to find all subclasses for a class.
+		/// Will check which assembly that class is in, then traverse that assmbly plus all assemblies depending on that assembly
+		/// to find the base classes. This is much more efficient than parsing all loaded assemblies to find subclasses
+		/// </summary>
+		/// <param name="domain">AppDomain to use</param>
+		/// <param name="baseClass">Base class to use</param>
+		public static IEnumerable<Type> EnumerateSubclasses(this AppDomain domain, Type baseClass)
+		{
+			Assembly baseAssembly = baseClass.Assembly;
+			AssemblyName baseAssemblyName = baseAssembly.GetName();
+
+			// Only search in assemblies that are either baseClass's assembly or assemblies depending on it
+			foreach (Assembly assembly in domain.GetAssemblies().Where(x => !x.IsDynamic))
+			{
+				if (!ReferenceEquals(assembly, baseAssembly))
+				{
+					try
+					{
+						if (!assembly
+							.GetReferencedAssemblies()
+							.Any(x => AssemblyName.ReferenceMatchesDefinition(x, baseAssemblyName)))
+						{
+							continue;
+						}
+					}
+					catch
+					{
+						continue;
+					}
+				}
+
+				Type[] types = assembly.SafeGetLoadedTypes();
+				foreach (Type t in types.Where(t => t.IsClass && !t.IsAbstract && t.IsSubclassOf(baseClass)))
+				{
+					yield return t;
+				}
+			}
+		}
+	}
+}
\ No newline at end of file
diff --git a/Engine/Source/Programs/Shared/EpicGames.Build/System/AssemblyExtensions.cs b/Engine/Source/Programs/Shared/EpicGames.Build/System/AssemblyExtensions.cs
--- a/Engine/Source/Programs/Shared/EpicGames.Build/System/AssemblyExtensions.cs
+++ b/Engine/Source/Programs/Shared/EpicGames.Build/System/AssemblyExtensions.cs
@@ -6,20 +6,26 @@ using System.Reflection;
 
 namespace UnrealBuildBase
 {
+	/// <summary>
+	/// Extensions to Assembly
+	/// </summary>
 	public static class AssemblyExtensions
 	{
-		public static Type[] SafeGetLoadedTypes(this Assembly Dll)
+		/// <summary>
+		/// Safely get all loaded types from an assembly.
+		/// </summary>
+		/// <param name="assembly">Assembly to use</param>
+		/// <returns></returns>
+		public static Type[] SafeGetLoadedTypes(this Assembly assembly)
 		{
-			Type[] AllTypes;
 			try
 			{
-				AllTypes = Dll.GetTypes();
+				return assembly.GetTypes();
 			}
 			catch (ReflectionTypeLoadException e)
 			{
-				AllTypes = e.Types.Where(x => x != null).Cast<Type>().ToArray();
+				return [.. e.Types.OfType<Type>()];
 			}
-			return AllTypes;
 		}
 	}
 }
diff --git a/Engine/Source/Programs/UnrealBuildTool/System/AppDomainExtensions.cs b/Engine/Source/Programs/UnrealBuildTool/System/AppDomainExtensions.cs
deleted file mode 100644
