Made UHT parsing remove leading whitespace and asterisks for multiline cstyle comments
#Synced-CL 48534016

diff --git a/Engine/Source/Programs/Shared/EpicGames.UHT/Parsers/UhtParsingScope.cs b/Engine/Source/Programs/Shared/EpicGames.UHT/Parsers/UhtParsingScope.cs
--- a/Engine/Source/Programs/Shared/EpicGames.UHT/Parsers/UhtParsingScope.cs
+++ b/Engine/Source/Programs/Shared/EpicGames.UHT/Parsers/UhtParsingScope.cs
@@ -349,7 +349,7 @@ namespace EpicGames.UHT.Parsers
 				line = line.TrimEnd();
 
 				// Remove leading "*" and "* " in javadoc comments.
-				if (javaDocStyle)
+				if (javaDocStyle || cStyle)
 				{
 					// Find first non-whitespace character
 					int pos = 0;
