[UBT] Have the compile event matcher recognise extensionless files containing underscores

UBT fails to recognise a diagnostic in files like "type_traits"

#Synced-CL 48716336

diff --git a/Engine/Source/Programs/Shared/EpicGames.LogMatchers/CompileEventMatcher.cs b/Engine/Source/Programs/Shared/EpicGames.LogMatchers/CompileEventMatcher.cs
--- a/Engine/Source/Programs/Shared/EpicGames.LogMatchers/CompileEventMatcher.cs
+++ b/Engine/Source/Programs/Shared/EpicGames.LogMatchers/CompileEventMatcher.cs
@@ -23,7 +23,7 @@ namespace EpicGames.LogMatchers
 				// any path character
 				@"[^:<>*?""]+" +
 				// valid source file extension (or extensionless file)
-				@"(?:\.(?i)(?:h|hpp|hxx|c|cc|cpp|cxx|inc|inl|cs|targets|verse)|[/\\][a-zA-Z0-9]+)" +
+				@"(?:\.(?i)(?:h|hpp|hxx|c|cc|cpp|cxx|inc|inl|cs|targets|verse)|[/\\][a-zA-Z0-9_]+)" +
 			// or the string "<scratch space>"
 			@")|<scratch space>)";
 
