Support -x for opened

#Synced-CL 48746020

diff --git a/Engine/Source/Programs/Shared/EpicGames.Perforce/PerforceConnection.cs b/Engine/Source/Programs/Shared/EpicGames.Perforce/PerforceConnection.cs
--- a/Engine/Source/Programs/Shared/EpicGames.Perforce/PerforceConnection.cs
+++ b/Engine/Source/Programs/Shared/EpicGames.Perforce/PerforceConnection.cs
@@ -2774,6 +2774,10 @@ namespace EpicGames.Perforce
 			{
 				arguments.Add($"-a");
 			}
+			if ((options & OpenedOptions.Exclusive) != 0)
+			{
+				arguments.Add($"-x");
+			}
 			if (changeNumber == PerforceReflection.DefaultChange)
 			{
 				arguments.Add($"-cdefault");
diff --git a/Engine/Source/Programs/Shared/EpicGames.Perforce/Records/OpenedOptions.cs b/Engine/Source/Programs/Shared/EpicGames.Perforce/Records/OpenedOptions.cs
--- a/Engine/Source/Programs/Shared/EpicGames.Perforce/Records/OpenedOptions.cs
+++ b/Engine/Source/Programs/Shared/EpicGames.Perforce/Records/OpenedOptions.cs
@@ -24,5 +24,10 @@ namespace EpicGames.Perforce
 		/// Short output; do not output the revision number or file type. This option is more efficient, particularly when using the -a (all-workspaces) option at large sites.
 		/// </summary>
 		ShortOutput = 2,
+
+		/// <summary>
+		/// List files opened exclusively
+		/// </summary>
+		Exclusive = 4,
 	}
 }
diff --git a/Engine/Source/Programs/Shared/EpicGames.Perforce/Records/OpenedRecord.cs b/Engine/Source/Programs/Shared/EpicGames.Perforce/Records/OpenedRecord.cs
--- a/Engine/Source/Programs/Shared/EpicGames.Perforce/Records/OpenedRecord.cs
+++ b/Engine/Source/Programs/Shared/EpicGames.Perforce/Records/OpenedRecord.cs
@@ -50,13 +50,13 @@ namespace EpicGames.Perforce
 		/// <summary>
 		/// Open action, if opened in your workspace
 		/// </summary>
-		[PerforceTag("action")]
+		[PerforceTag("action", Optional = true)]
 		public FileAction Action { get; set; }
 
 		/// <summary>
 		/// Change containing the open file
 		/// </summary>
-		[PerforceTag("change")]
+		[PerforceTag("change", Optional = true)]
 		public int Change { get; set; }
 
 		/// <summary>
