Add support for Hashing arbitary spans of data to IoHash instead of only byte spans
Move BloomFilter.cs from UnrealBuildTool to EpicGames.Build

#rnx

#Synced-CL 48679523

diff --git a/Engine/Source/Programs/Shared/EpicGames.Build/System/BloomFilter.cs b/Engine/Source/Programs/Shared/EpicGames.Build/System/BloomFilter.cs
new file mode 100644
--- /dev/null
+++ b/Engine/Source/Programs/Shared/EpicGames.Build/System/BloomFilter.cs
@@ -0,0 +1,75 @@
+// Copyright Epic Games, Inc. All Rights Reserved.
+
+using System;
+using System.Collections.Generic;
+using EpicGames.Core;
+
+namespace UnrealBuildBase
+{
+	/// <summary>
+	/// Simple bloom filter implementation which can be used to optimize searches by early out on 
+	/// </summary>
+	public sealed class BloomFilter
+	{
+		const uint ByteCount = 128;
+		private readonly byte[] _bits = new byte[ByteCount];
+
+		/// <summary>
+		/// Add string to filter
+		/// </summary>
+		/// <param name="s"></param>
+		public void Add(ReadOnlySpan<char> s)
+		{
+			uint hash = (uint)IoHash.Compute(s).GetHashCode();
+			int bitIndex = (int)(hash % (ByteCount * 8));
+			int byteIndex = bitIndex / 8;
+			int bitInByte = bitIndex - byteIndex * 8;
+			_bits[byteIndex] |= (byte)(1 << bitInByte);
+		}
+
+		/// <summary>
+		/// Add multiple strings to filter
+		/// </summary>
+		/// <param name="strings"></param>
+		public void AddRange(IEnumerable<string> strings)
+		{
+			foreach (string s in strings)
+			{
+				Add(s);
+			}
+		}
+
+		/// <summary>
+		/// Will return false if string is guaranteed to not exist, and true if it could exist
+		/// </summary>
+		/// <param name="s"></param>
+		public bool MightContain(ReadOnlySpan<char> s)
+		{
+			uint hash = (uint)IoHash.Compute(s).GetHashCode();
+			int bitIndex = (int)(hash % (ByteCount * 8));
+			int byteIndex = bitIndex / 8;
+			int bitInByte = bitIndex - byteIndex * 8;
+			return (_bits[byteIndex] & (1 << bitInByte)) != 0;
+		}
+
+		/// <summary>
+		/// Will traverse a string to find substrings separated with separator and then test them
+		/// Useful if you for example want to check if a a string might exist somewhere in a path
+		/// </summary>
+		/// <param name="span"></param>
+		/// <param name="separator"></param>
+		public bool MightContainSubString(ReadOnlySpan<char> span, char separator)
+		{
+			int separatorIndex;
+			while ((separatorIndex = span.IndexOf(separator)) != -1)
+			{
+				if (MightContain(span.Slice(0, separatorIndex)))
+				{
+					return true;
+				}
+				span = span.Slice(separatorIndex + 1);
+			}
+			return MightContain(span);
+		}
+	}
+}
\ No newline at end of file
diff --git a/Engine/Source/Programs/Shared/EpicGames.IoHash/IoHash.cs b/Engine/Source/Programs/Shared/EpicGames.IoHash/IoHash.cs
--- a/Engine/Source/Programs/Shared/EpicGames.IoHash/IoHash.cs
+++ b/Engine/Source/Programs/Shared/EpicGames.IoHash/IoHash.cs
@@ -6,6 +6,7 @@ using System.Buffers.Binary;
 using System.ComponentModel;
 using System.Globalization;
 using System.IO;
+using System.Runtime.InteropServices;
 using System.Text;
 using System.Text.Json;
 using System.Text.Json.Serialization;
@@ -43,7 +44,7 @@ namespace EpicGames.Core
 		/// <summary>
 		/// Hash consisting of zeroes
 		/// </summary>
-		public static IoHash Zero { get; } = new IoHash(0, 0, 0);
+		public static IoHash Zero { get; } = new(0, 0, 0);
 
 		/// <summary>
 		/// Constructor
@@ -62,7 +63,7 @@ namespace EpicGames.Core
 		{
 			Span<byte> output = stackalloc byte[32];
 			hasher.Finalize(output);
-			return new IoHash(output);
+			return new(output);
 		}
 
 		/// <summary>
@@ -70,51 +71,62 @@ namespace EpicGames.Core
 		/// </summary>
 		/// <param name="data">Data to compute the hash for</param>
 		/// <returns>New hash instance containing the hash of the data</returns>
-		public static IoHash Compute(ReadOnlySpan<byte> data)
+		public static IoHash Compute(ReadOnlySpan<byte> data) => Compute<byte>(data);
+
+		/// <summary>
+		/// Creates the IoHash for a block of data.
+		/// </summary>
+		/// <typeparam name="T">Type of the data</typeparam>
+		/// <param name="data">Data to compute the hash for</param>
+		/// <returns>New hash instance containing the hash of the data</returns>
+		public static IoHash Compute<T>(ReadOnlySpan<T> data) where T : unmanaged
 		{
-			Span<byte> output = stackalloc byte[32];
-			using (Blake3.Hasher hasher = Blake3.Hasher.New())
+			using Blake3.Hasher hasher = Blake3.Hasher.New();
+			if (data.Length * Marshal.SizeOf<T>() < MultiThreadedSize)
 			{
-				if (data.Length < MultiThreadedSize)
-				{
-					hasher.Update(data);
-				}
-				else
-				{
-					hasher.UpdateWithJoin(data);
-				}
-				hasher.Finalize(output);
+				hasher.Update(data);
 			}
-			return new IoHash(output);
+			else
+			{
+				hasher.UpdateWithJoin(data);
+			}
+			return FromBlake3(hasher);
 		}
 
 		/// <summary>
-		/// Creates the IoHash for a block of data.
+		/// Creates the IoHash for a sequence of data.
+		/// </summary>
+		/// <param name="sequence">Data to compute the hash for</param>
+		/// <returns>New hash instance containing the hash of the data</returns>
+		public static IoHash Compute(ReadOnlySequence<byte> sequence) => Compute<byte>(sequence);
+
+		/// <summary>
+		/// Creates the IoHash for a sequence of data.
 		/// </summary>
+		/// <typeparam name="T">Type of the data</typeparam>
 		/// <param name="sequence">Data to compute the hash for</param>
 		/// <returns>New hash instance containing the hash of the data</returns>
-		public static IoHash Compute(ReadOnlySequence<byte> sequence)
+		public static IoHash Compute<T>(ReadOnlySequence<T> sequence) where T : unmanaged
 		{
 			if (sequence.IsSingleSegment)
 			{
 				return Compute(sequence.FirstSpan);
 			}
 
-			using (Blake3.Hasher hasher = Blake3.Hasher.New())
+			using Blake3.Hasher hasher = Blake3.Hasher.New();
+			int size = Marshal.SizeOf<T>();
+			foreach (ReadOnlyMemory<T> segment in sequence)
 			{
-				foreach (ReadOnlyMemory<byte> segment in sequence)
+				if (segment.Length * size < MultiThreadedSize)
+				{
+					hasher.Update(segment.Span);
+				}
+				else
 				{
-					if (segment.Length < MultiThreadedSize)
-					{
-						hasher.Update(segment.Span);
-					}
-					else
-					{
-						hasher.UpdateWithJoin(segment.Span);
-					}
+					hasher.UpdateWithJoin(segment.Span);
 				}
-				return FromBlake3(hasher);
 			}
+			return FromBlake3(hasher);
 		}
 
 		/// <summary>
@@ -139,18 +151,16 @@ namespace EpicGames.Core
 		/// <returns>New content hash instance containing the hash of the data</returns>
 		public static IoHash Compute(Stream stream)
 		{
-			using (Blake3.Hasher hasher = Blake3.Hasher.New())
-			{
-				Span<byte> buffer = stackalloc byte[16384];
+			using Blake3.Hasher hasher = Blake3.Hasher.New();
+			Span<byte> buffer = stackalloc byte[16384];
 
-				int length;
-				while ((length = stream.Read(buffer)) > 0)
-				{
-					hasher.Update(buffer.Slice(0, length));
-				}
-
-				return FromBlake3(hasher);
+			int length;
+			while ((length = stream.Read(buffer)) > 0)
+			{
+				hasher.Update(buffer.Slice(0, length));
 			}
+
+			return FromBlake3(hasher);
 		}
 
 		/// <summary>
@@ -172,17 +182,15 @@ namespace EpicGames.Core
 		{
 			const int MaxBufferSize = 1 * 1024 * 1024;
 			const int MinBufferSize = 16 * 1024;
-			using (Blake3.Hasher hasher = Blake3.Hasher.New())
+			using Blake3.Hasher hasher = Blake3.Hasher.New();
+			Task Callback(ReadOnlyMemory<byte> data)
 			{
-				Task Callback(ReadOnlyMemory<byte> data)
-				{
-					hasher.UpdateWithJoin(data.Span);
-					return Task.CompletedTask;
-				}
-
-				await stream.ReadAllBytesAsync(fileSizeHint, MinBufferSize, MaxBufferSize, Callback, cancellationToken);
-				return FromBlake3(hasher);
+				hasher.UpdateWithJoin(data.Span);
+				return Task.CompletedTask;
 			}
+
+			await stream.ReadAllBytesAsync(fileSizeHint, MinBufferSize, MaxBufferSize, Callback, cancellationToken);
+			return FromBlake3(hasher);
 		}
 
 		/// <summary>
@@ -190,30 +198,21 @@ namespace EpicGames.Core
 		/// </summary>
 		/// <param name="text"></param>
 		/// <returns></returns>
-		public static IoHash Parse(string text)
-		{
-			return new IoHash(StringUtils.ParseHexString(text));
-		}
+		public static IoHash Parse(string text) => new(StringUtils.ParseHexString(text));
 
 		/// <summary>
 		/// Parses a digest from the given hex string
 		/// </summary>
 		/// <param name="text"></param>
 		/// <returns></returns>
-		public static IoHash Parse(ReadOnlySpan<char> text)
-		{
-			return new IoHash(StringUtils.ParseHexString(text));
-		}
+		public static IoHash Parse(ReadOnlySpan<char> text) => new(StringUtils.ParseHexString(text));
 
 		/// <summary>
 		/// Parses a digest from the given hex string
 		/// </summary>
 		/// <param name="text"></param>
 		/// <returns></returns>
-		public static IoHash Parse(ReadOnlySpan<byte> text)
-		{
-			return new IoHash(StringUtils.ParseHexString(text));
-		}
+		public static IoHash Parse(ReadOnlySpan<byte> text) => new(StringUtils.ParseHexString(text));
 
 		/// <summary>
 		/// Parses a digest from the given hex string
@@ -224,9 +223,9 @@ namespace EpicGames.Core
 		public static bool TryParse(ReadOnlySpan<char> text, out IoHash hash)
 		{
 			byte[]? bytes;
-			if (StringUtils.TryParseHexString(text, out bytes) && bytes.Length == IoHash.NumBytes)
+			if (StringUtils.TryParseHexString(text, out bytes) && bytes.Length == NumBytes)
 			{
-				hash = new IoHash(bytes);
+				hash = new(bytes);
 				return true;
 			}
 			else
@@ -245,9 +244,9 @@ namespace EpicGames.Core
 		public static bool TryParse(ReadOnlySpan<byte> text, out IoHash hash)
 		{
 			byte[]? bytes;
-			if (StringUtils.TryParseHexString(text, out bytes) && bytes.Length == IoHash.NumBytes)
+			if (StringUtils.TryParseHexString(text, out bytes) && bytes.Length == NumBytes)
 			{
-				hash = new IoHash(bytes);
+				hash = new(bytes);
 				return true;
 			}
 			else
@@ -288,7 +287,7 @@ namespace EpicGames.Core
 		/// </summary>
 		public Utf8String ToUtf8String()
 		{
-			Span<byte> buffer = stackalloc byte[IoHash.NumBytes];
+			Span<byte> buffer = stackalloc byte[NumBytes];
 			CopyTo(buffer);
 			return StringUtils.FormatUtf8HexString(buffer);
 		}
@@ -299,7 +298,7 @@ namespace EpicGames.Core
 		/// <param name="chars">Output buffer for the converted string</param>
 		public void ToUtf8String(Span<byte> chars)
 		{
-			Span<byte> buffer = stackalloc byte[IoHash.NumBytes];
+			Span<byte> buffer = stackalloc byte[NumBytes];
 			CopyTo(buffer);
 			StringUtils.FormatUtf8HexString(buffer, chars);
 		}
@@ -307,7 +306,7 @@ namespace EpicGames.Core
 		/// <inheritdoc/>
 		public override string ToString()
 		{
-			Span<byte> buffer = stackalloc byte[IoHash.NumBytes];
+			Span<byte> buffer = stackalloc byte[NumBytes];
 			CopyTo(buffer);
 			return StringUtils.FormatHexString(buffer);
 		}
@@ -368,10 +367,7 @@ namespace EpicGames.Core
 		/// Convert a Blake3Hash to an IoHash
 		/// </summary>
 		/// <param name="hash"></param>
-		public static implicit operator IoHash(Blake3Hash hash)
-		{
-			return new IoHash(hash.Span.Slice(0, NumBytes));
-		}
+		public static implicit operator IoHash(Blake3Hash hash) => new(hash.Span.Slice(0, NumBytes));
 	}
 
 	/// <summary>
@@ -387,7 +383,7 @@ namespace EpicGames.Core
 		public static IoHash? ReadIoHash(this BinaryArchiveReader reader)
 		{
 			byte[]? data = reader.ReadByteArray();
-			return data == null ? null : new IoHash(data);
+			return data == null ? null : new(data);
 		}
 
 		/// <summary>
@@ -395,20 +391,14 @@ namespace EpicGames.Core
 		/// </summary>
 		/// <param name="reader">Reader to serialize data from</param>
 		/// <returns>New IoHash instance</returns>
-		public static IoHash ReadIoHash(this IMemoryReader reader)
-		{
-			return new IoHash(reader.ReadFixedLengthBytes(IoHash.NumBytes).Span);
-		}
+		public static IoHash ReadIoHash(this IMemoryReader reader) => new(reader.ReadFixedLengthBytes(IoHash.NumBytes).Span);
 
 		/// <summary>
 		/// Write an <see cref="IoHash"/> to a binary archive
 		/// </summary>
 		/// <param name="writer">The writer to output data to</param>
 		/// <param name="hash">The IoHash to write</param>
-		public static void WriteIoHash(this BinaryArchiveWriter writer, IoHash? hash)
-		{
-			writer.WriteByteArray(hash?.ToByteArray());
-		}
+		public static void WriteIoHash(this BinaryArchiveWriter writer, IoHash? hash) => writer.WriteByteArray(hash?.ToByteArray());
 
 		/// <summary>
 		/// Write an <see cref="IoHash"/> to a memory writer
@@ -440,15 +430,9 @@ namespace EpicGames.Core
 	sealed class IoHashTypeConverter : TypeConverter
 	{
 		/// <inheritdoc/>
-		public override bool CanConvertFrom(ITypeDescriptorContext? context, Type sourceType)
-		{
-			return sourceType == typeof(string);
-		}
+		public override bool CanConvertFrom(ITypeDescriptorContext? context, Type sourceType) => sourceType == typeof(string);
 
 		/// <inheritdoc/>
-		public override object ConvertFrom(ITypeDescriptorContext? context, CultureInfo? culture, object value)
-		{
-			return IoHash.Parse((string)value);
-		}
+		public override object ConvertFrom(ITypeDescriptorContext? context, CultureInfo? culture, object value) => IoHash.Parse((string)value);
 	}
 }
diff --git a/Engine/Source/Programs/UnrealBuildTool/System/BloomFilter.cs b/Engine/Source/Programs/UnrealBuildTool/System/BloomFilter.cs
deleted file mode 100644
