EpicGames.IoHash: Rent 16k buffer rather than stackalloc
Add Combine helper function to hash two hashes together

#rnx
#Synced-CL 48679523

diff --git a/Engine/Source/Programs/Shared/EpicGames.IoHash/IoHash.cs b/Engine/Source/Programs/Shared/EpicGames.IoHash/IoHash.cs
--- a/Engine/Source/Programs/Shared/EpicGames.IoHash/IoHash.cs
+++ b/Engine/Source/Programs/Shared/EpicGames.IoHash/IoHash.cs
@@ -61,7 +61,7 @@ namespace EpicGames.Core
 		/// <param name="hasher">The hasher to construct from</param>
 		public static IoHash FromBlake3(Blake3.Hasher hasher)
 		{
-			Span<byte> output = stackalloc byte[32];
+			Span<byte> output = stackalloc byte[Blake3Hash.NumBytes];
 			hasher.Finalize(output);
 			return new(output);
 		}
@@ -151,16 +151,29 @@ namespace EpicGames.Core
 		/// <returns>New content hash instance containing the hash of the data</returns>
 		public static IoHash Compute(Stream stream)
 		{
-			using Blake3.Hasher hasher = Blake3.Hasher.New();
-			Span<byte> buffer = stackalloc byte[16384];
+			byte[]? buffer = null;
+			try
+			{
+				using Blake3.Hasher hasher = Blake3.Hasher.New();
+				buffer = ArrayPool<byte>.Shared.Rent(16 * 1024);
+				Span<byte> span = buffer.AsSpan();
+
+				int length;
+				while ((length = stream.Read(span)) > 0)
+				{
+					hasher.Update(span.Slice(0, length));
+				}
 
-			int length;
-			while ((length = stream.Read(buffer)) > 0)
+				ArrayPool<byte>.Shared.Return(buffer);
+				return FromBlake3(hasher);
+			}
+			finally
 			{
-				hasher.Update(buffer.Slice(0, length));
+				if (buffer != null)
+				{
+					ArrayPool<byte>.Shared.Return(buffer);
+				}
 			}
-
-			return FromBlake3(hasher);
 		}
 
 		/// <summary>
@@ -325,7 +338,7 @@ namespace EpicGames.Core
 		/// <summary>
 		/// Copies this hash into a span
 		/// </summary>
-		/// <param name="span"></param>
+		/// <param name="span">Span to copy to</param>
 		public void CopyTo(Span<byte> span)
 		{
 			BinaryPrimitives.WriteUInt64BigEndian(span, _a);
@@ -333,6 +346,24 @@ namespace EpicGames.Core
 			BinaryPrimitives.WriteUInt32BigEndian(span[16..], _c);
 		}
 
+		/// <summary>
+		/// Combine two hashes into a new hash
+		/// </summary>
+		/// <param name="first">The first hash</param>
+		/// <param name="second">The second hash</param>
+		/// <remarks>This function is not commutative, nor will it result the same hash as if the original data was hashed.</remarks>
+		/// <returns>A combined hash</returns>
+		public static IoHash Combine(IoHash first, IoHash second)
+		{
+			Span<byte> buffer = stackalloc byte[NumBytes];
+			using Blake3.Hasher hasher = Blake3.Hasher.New();
+			first.CopyTo(buffer);
+			hasher.Update(buffer);
+			second.CopyTo(buffer);
+			hasher.Update(buffer);
+			return FromBlake3(hasher);
+		}
+
 		/// <summary>
 		/// Test two hash values for equality
 		/// </summary>
