EpicGames.IoHash: Add explictit conversion from IoHash to uint which returns the first 32 bytes of the hashed value

#rnx
#Synced-CL 48679523

diff --git a/Engine/Source/Programs/Shared/EpicGames.IoHash/IoHash.cs b/Engine/Source/Programs/Shared/EpicGames.IoHash/IoHash.cs
--- a/Engine/Source/Programs/Shared/EpicGames.IoHash/IoHash.cs
+++ b/Engine/Source/Programs/Shared/EpicGames.IoHash/IoHash.cs
@@ -399,6 +399,12 @@ namespace EpicGames.Core
 		/// </summary>
 		/// <param name="hash"></param>
 		public static implicit operator IoHash(Blake3Hash hash) => new(hash.Span.Slice(0, NumBytes));
+
+		/// <summary>
+		/// Convert an IoHash to a uint by returning the first 32 bits of the hashed value
+		/// </summary>
+		/// <param name="hash"></param>
+		public static explicit operator uint(IoHash hash) => (uint)(hash._a >> 32);
 	}
 
 	/// <summary>
