PR #12382: [UHT] Fix format string when usage in member declaration detected without an alternative, in UhtIssueBehavior.AllowAndLog case

#Synced-CL 48862428

diff --git a/Engine/Source/Programs/Shared/EpicGames.UHT/Utils/UhtPropertyResolveArgs.cs b/Engine/Source/Programs/Shared/EpicGames.UHT/Utils/UhtPropertyResolveArgs.cs
--- a/Engine/Source/Programs/Shared/EpicGames.UHT/Utils/UhtPropertyResolveArgs.cs
+++ b/Engine/Source/Programs/Shared/EpicGames.UHT/Utils/UhtPropertyResolveArgs.cs
@@ -407,7 +407,7 @@ namespace EpicGames.UHT.Utils
 					}
 					else
 					{
-						TokenReader.LogInfo("{PointerTypeDesc} usage in member declaration detected [[[{Type}]]].");
+						TokenReader.LogInfo($"{pointerTypeDesc} usage in member declaration detected [[[{type}]]].");
 					}
 					break;
 
