Add a .NET project info implementation that uses MSBuild to parse out project properties, as adding enough parsing support to CsProjectInfo to allow it to parse and understand Directory.Build.props files was infeasible. Don't use this new implementation yet.

#Synced-CL 48900888

diff --git a/Engine/Source/Programs/AutomationTool/AutomationUtils/HashCollection.cs b/Engine/Source/Programs/AutomationTool/AutomationUtils/HashCollection.cs
--- a/Engine/Source/Programs/AutomationTool/AutomationUtils/HashCollection.cs
+++ b/Engine/Source/Programs/AutomationTool/AutomationUtils/HashCollection.cs
@@ -3,13 +3,13 @@
 using System;
 using System.Collections.Generic;
 using System.IO;
-using System.Security.Cryptography;
 using System.Text;
 using System.Xml;
 using System.Xml.Serialization;
 using System.Linq;
 using EpicGames.Core;
 using Microsoft.Extensions.Logging;
+using UnrealBuildBase;
 
 namespace AutomationTool
 {
@@ -264,7 +264,7 @@ namespace AutomationTool
 		/// <param name="Project"></param>
 		/// <param name="HashType"></param>
 		/// <returns></returns>
-		public static bool AddCsProjectInfo(this HashCollection Hasher, CsProjectInfo Project, HashCollection.HashType HashType)
+		public static bool AddCsProjectInfo(this HashCollection Hasher, IDotnetProjectInfo Project, HashCollection.HashType HashType)
 		{
 			// Get the output assembly and pdb file
 			FileReference OutputFile;
diff --git a/Engine/Source/Programs/AutomationTool/BuildGraph/Tasks/CsCompileTask.cs b/Engine/Source/Programs/AutomationTool/BuildGraph/Tasks/CsCompileTask.cs
--- a/Engine/Source/Programs/AutomationTool/BuildGraph/Tasks/CsCompileTask.cs
+++ b/Engine/Source/Programs/AutomationTool/BuildGraph/Tasks/CsCompileTask.cs
@@ -261,16 +261,16 @@ namespace AutomationTool.Tasks
 			outReferences = new HashSet<FileReference>();
 
 			// Read all the project information into a dictionary
-			Dictionary<FileReference, CsProjectInfo> fileToProjectInfo = new Dictionary<FileReference, CsProjectInfo>();
+			Dictionary<FileReference, IDotnetProjectInfo> fileToProjectInfo = [];
 			foreach (FileReference projectFile in projectFiles)
 			{
 				// Read all the projects
 				ReadProjectsRecursively(projectFile, initialProperties, fileToProjectInfo);
 
 				// Find all the outputs for each project
-				foreach (KeyValuePair<FileReference, CsProjectInfo> pair in fileToProjectInfo)
+				foreach (KeyValuePair<FileReference, IDotnetProjectInfo> pair in fileToProjectInfo)
 				{
-					CsProjectInfo projectInfo = pair.Value;
+					IDotnetProjectInfo projectInfo = pair.Value;
 
 					// Add all the build projects from this project
 					DirectoryReference outputDir = projectInfo.GetOutputDir(pair.Key.Directory);
@@ -279,7 +279,7 @@ namespace AutomationTool.Tasks
 					// Add any files which are only referenced
 					foreach (KeyValuePair<FileReference, bool> reference in projectInfo.References)
 					{
-						CsProjectInfo.AddReferencedAssemblyAndSupportFiles(reference.Key, outReferences);
+						projectInfo.AddReferencedAssemblyAndSupportFiles(reference.Key, outReferences);
 					}
 				}
 			}
@@ -295,7 +295,7 @@ namespace AutomationTool.Tasks
 		/// <param name="initialProperties">Mapping of property name to value for the initial project</param>
 		/// <param name="fileToProjectInfo"></param>
 		/// <returns>True if the projects were read correctly, false (and prints an error to the log) if not</returns>
-		static void ReadProjectsRecursively(FileReference file, Dictionary<string, string> initialProperties, Dictionary<FileReference, CsProjectInfo> fileToProjectInfo)
+		static void ReadProjectsRecursively(FileReference file, Dictionary<string, string> initialProperties, Dictionary<FileReference, IDotnetProjectInfo> fileToProjectInfo)
 		{
 			// Early out if we've already read this project
 			if (!fileToProjectInfo.ContainsKey(file))
diff --git a/Engine/Source/Programs/AutomationTool/CookedEditor/CookedEditor.Automation.cs b/Engine/Source/Programs/AutomationTool/CookedEditor/CookedEditor.Automation.cs
--- a/Engine/Source/Programs/AutomationTool/CookedEditor/CookedEditor.Automation.cs
+++ b/Engine/Source/Programs/AutomationTool/CookedEditor/CookedEditor.Automation.cs
@@ -580,31 +580,6 @@ public class MakeCookedEditor : BuildCommand
 		}
 	}
 
-	static void ReadProjectsRecursively(FileReference File, Dictionary<string, string> InitialProperties, Dictionary<FileReference, CsProjectInfo> FileToProjectInfo)
-	{
-		// Early out if we've already read this project
-		if (!FileToProjectInfo.ContainsKey(File))
-		{
-			// Try to read this project
-			CsProjectInfo ProjectInfo;
-			if (!CsProjectInfo.TryRead(File, InitialProperties, out ProjectInfo))
-			{
-				throw new AutomationException("Couldn't read project '{0}'", File.FullName);
-			}
-
-			// Add it to the project lookup, and try to read all the projects it references
-			FileToProjectInfo.Add(File, ProjectInfo);
-			foreach (FileReference ProjectReference in ProjectInfo.ProjectReferences.Keys)
-			{
-				if (!FileReference.Exists(ProjectReference))
-				{
-					throw new AutomationException("Unable to find project '{0}' referenced by '{1}'", ProjectReference, File);
-				}
-				ReadProjectsRecursively(ProjectReference, InitialProperties, FileToProjectInfo);
-			}
-		}
-	}
-
 	protected virtual void StageUAT(ProjectParams Params, DeploymentContext SC, ModifyStageContext Context)
 	{
 		// now look in the .json file that UAT made that lists its script files
diff --git a/Engine/Source/Programs/Shared/EpicGames.Build/CsProjectInfo.cs b/Engine/Source/Programs/Shared/EpicGames.Build/CsProjectInfo.cs
--- a/Engine/Source/Programs/Shared/EpicGames.Build/CsProjectInfo.cs
+++ b/Engine/Source/Programs/Shared/EpicGames.Build/CsProjectInfo.cs
@@ -26,32 +26,37 @@ namespace EpicGames.Core
 	/// <summary>
 	/// Basic information from a preprocessed C# project file. Supports reading a project file, expanding simple conditions in it, parsing property values, assembly references and references to other projects.
 	/// </summary>
-	public class CsProjectInfo
+	public class CsProjectInfo : IDotnetProjectInfo
 	{
 		/// <summary>
 		/// Evaluated properties from the project file
 		/// </summary>
 		public Dictionary<string, string> Properties { get; }
+		IReadOnlyDictionary<string, string> IDotnetProjectInfo.Properties => Properties;
 
 		/// <summary>
 		/// Mapping of referenced assemblies to their 'CopyLocal' (aka 'Private') setting.
 		/// </summary>
 		public Dictionary<FileReference, bool> References { get; } = [];
+		IReadOnlyDictionary<FileReference, bool> IDotnetProjectInfo.References => References;
 
 		/// <summary>
 		/// Mapping of referenced projects to their 'CopyLocal' (aka 'Private') setting.
 		/// </summary>
 		public Dictionary<FileReference, bool> ProjectReferences { get; } = [];
+		IReadOnlyDictionary<FileReference, bool> IDotnetProjectInfo.ProjectReferences => ProjectReferences;
 
 		/// <summary>
 		/// List of compile references in the project.
 		/// </summary>
 		public List<FileReference> CompileReferences { get; } = [];
+		IReadOnlyList<FileReference> IDotnetProjectInfo.CompileReferences => CompileReferences;
 
 		/// <summary>
 		/// Mapping of content IF they are flagged Always or Newer
 		/// </summary>
 		public Dictionary<FileReference, bool> ContentReferences { get; } = [];
+		IReadOnlyDictionary<FileReference, bool> IDotnetProjectInfo.ContentReferences => ContentReferences;
 
 		/// <summary>
 		/// Path to the CSProject file
@@ -71,282 +76,6 @@ namespace EpicGames.Core
 			Properties.TryAdd("MSBuildThisFileDirectory", inProjectPath.Directory.FullName);
 		}
 
-		/// <summary>
-		/// Get the ouptut file for this project
-		/// </summary>
-		/// <param name="file">If successful, receives the assembly path</param>
-		/// <returns>True if the output file was found</returns>
-		public bool TryGetOutputFile([NotNullWhen(true)] out FileReference? file)
-		{
-			DirectoryReference? outputDir;
-			if(!TryGetOutputDir(out outputDir))
-			{
-				file = null;
-				return false;
-			}
-
-			string? assemblyName;
-			if(!TryGetAssemblyName(out assemblyName))
-			{
-				file = null;
-				return false;
-			}
-
-			file = FileReference.Combine(outputDir, assemblyName + ".dll");
-			return true;
-		}
-
-		/// <summary>
-		/// Returns whether or not this project is a modern dotnet project
-		/// </summary>
-		/// <returns>True if this is a netcore project</returns>
-		private bool IsNetCoreProject()
-		{
-			string? framework;
-			return Properties.TryGetValue("TargetFramework", out framework)
-				&& (framework.StartsWith("netcoreapp", StringComparison.Ordinal) || Regex.Match(framework, @"net\d+\.0").Success);
-		}
-
-		/// <summary>
-		/// Resolve the project's output directory
-		/// </summary>
-		/// <param name="baseDirectory">Base directory to resolve relative paths to</param>
-		/// <returns>The configured output directory</returns>
-		public DirectoryReference GetOutputDir(DirectoryReference baseDirectory)
-		{
-			string? outputPath;
-			if (Properties.TryGetValue("OutputPath", out outputPath))
-			{
-				return DirectoryReference.Combine(baseDirectory, outputPath);
-			}
-			else if (IsNetCoreProject())
-			{
-				string configuration = Properties.TryGetValue("Configuration", out string? value) ? value : "Development";
-				return !Properties.TryGetValue("AppendTargetFrameworkToOutputPath", out string? appendFrameworkStr) || appendFrameworkStr.Equals("true", StringComparison.OrdinalIgnoreCase)
-					? DirectoryReference.Combine(baseDirectory, "bin", configuration, Properties["TargetFramework"])
-					: DirectoryReference.Combine(baseDirectory, "bin", configuration);
-			}
-			else
-			{
-				return baseDirectory;
-			}
-		}
-
-		/// <summary>
-		/// Resolve the project's output directory
-		/// </summary>
-		/// <param name="outputDir">If successful, receives the output directory</param>
-		/// <returns>True if the output directory was found</returns>
-		public bool TryGetOutputDir([NotNullWhen(true)] out DirectoryReference? outputDir)
-		{
-			string? outputPath;
-			if (Properties.TryGetValue("OutputPath", out outputPath))
-			{
-				outputDir = DirectoryReference.Combine(ProjectPath.Directory, outputPath);
-				return true;
-			}
-			else if (IsNetCoreProject())
-			{
-				string configuration = Properties.ContainsKey("Configuration") ? Properties["Configuration"] : "Development";
-				outputDir = DirectoryReference.Combine(ProjectPath.Directory, "bin", configuration, Properties["TargetFramework"]);
-				return true;
-			}
-			else
-			{
-				outputDir = null;
-				return false;
-			}
-		}
-
-		/// <summary>
-		/// Returns the assembly name used by this project
-		/// </summary>
-		/// <returns></returns>
-		public bool TryGetAssemblyName([NotNullWhen(true)] out string? assemblyName)
-		{
-			if (Properties.TryGetValue("AssemblyName", out assemblyName))
-			{
-				return true;
-			}
-			else if (IsNetCoreProject())
-			{
-				assemblyName = ProjectPath.GetFileNameWithoutExtension();
-				return true;
-			}
-			return false;
-		}
-
-		/// <summary>
-		/// Finds all build products from this project. This includes content and other assemblies marked to be copied local.
-		/// </summary>
-		/// <param name="outputDir">The output directory</param>
-		/// <param name="buildProducts">Receives the set of build products</param>
-		/// <param name="projectFileToInfo">Map of project file to information, to resolve build products from referenced projects copied locally</param>
-		public void FindBuildProducts(DirectoryReference outputDir, HashSet<FileReference> buildProducts, Dictionary<FileReference, CsProjectInfo> projectFileToInfo)
-		{
-			// Add the standard build products
-			FindCompiledBuildProducts(outputDir, buildProducts);
-
-			// Add the referenced assemblies which are marked to be copied into the output directory. This only happens for the main project, and does not happen for referenced projects.
-			foreach(KeyValuePair<FileReference, bool> reference in References)
-			{
-				if (reference.Value)
-				{
-					FileReference outputFile = FileReference.Combine(outputDir, reference.Key.GetFileName());
-					AddReferencedAssemblyAndSupportFiles(outputFile, buildProducts);
-				}
-			}
-
-			// Copy the build products for any referenced projects. Note that this does NOT operate recursively.
-			foreach(KeyValuePair<FileReference, bool> projectReference in ProjectReferences)
-			{
-				CsProjectInfo? otherProjectInfo;
-				if(projectFileToInfo.TryGetValue(projectReference.Key, out otherProjectInfo))
-				{
-					otherProjectInfo.FindCompiledBuildProducts(outputDir, buildProducts);
-				}
-			}
-
-			// Add any copied content. This DOES operate recursively.
-			FindCopiedContent(outputDir, buildProducts, projectFileToInfo);
-		}
-
-		/// <summary>
-		/// Determines all the compiled build products (executable, etc...) directly built from this project.
-		/// </summary>
-		/// <param name="outputDir">The output directory</param>
-		/// <param name="buildProducts">Receives the set of build products</param>
-		public void FindCompiledBuildProducts(DirectoryReference outputDir, HashSet<FileReference> buildProducts)
-		{
-			string? outputType, assemblyName;
-			if (Properties.TryGetValue("OutputType", out outputType) && TryGetAssemblyName(out assemblyName))
-			{
-				switch (outputType)
-				{
-					case "Exe":
-					case "WinExe":
-						string executableExtension = RuntimeInformation.IsOSPlatform(OSPlatform.Windows) ? ".exe" : "";
-						buildProducts.Add(FileReference.Combine(outputDir, assemblyName + executableExtension));
-						// dotnet outputs a apphost executable and a dll with the actual assembly
-						AddOptionalBuildProduct(FileReference.Combine(outputDir, assemblyName + ".dll"), buildProducts);
-						AddOptionalBuildProduct(FileReference.Combine(outputDir, assemblyName + ".pdb"), buildProducts);
-						AddOptionalBuildProduct(FileReference.Combine(outputDir, assemblyName + ".exe.config"), buildProducts);
-						AddOptionalBuildProduct(FileReference.Combine(outputDir, assemblyName + ".exe.mdb"), buildProducts);
-						break;
-					case "Library":
-						buildProducts.Add(FileReference.Combine(outputDir, assemblyName + ".dll"));
-						AddOptionalBuildProduct(FileReference.Combine(outputDir, assemblyName + ".pdb"), buildProducts);
-						AddOptionalBuildProduct(FileReference.Combine(outputDir, assemblyName + ".dll.config"), buildProducts);
-						AddOptionalBuildProduct(FileReference.Combine(outputDir, assemblyName + ".dll.mdb"), buildProducts);
-						break;
-				}
-			}
-		}
-
-		/// <summary>
-		/// Finds all content which will be copied into the output directory for this project. This includes content from any project references as "copy local" recursively (though MSBuild only traverses a single reference for actual binaries, in such cases)
-		/// </summary>
-		/// <param name="outputDir">The output directory</param>
-		/// <param name="outputFiles">Receives the set of build products</param>
-		/// <param name="projectFileToInfo">Map of project file to information, to resolve build products from referenced projects copied locally</param>
-		private void FindCopiedContent(DirectoryReference outputDir, HashSet<FileReference> outputFiles, Dictionary<FileReference, CsProjectInfo> projectFileToInfo)
-		{
-			// Copy any referenced projects too.
-			foreach(KeyValuePair<FileReference, bool> projectReference in ProjectReferences)
-			{
-				CsProjectInfo? otherProjectInfo;
-				if(projectFileToInfo.TryGetValue(projectReference.Key, out otherProjectInfo))
-				{
-					otherProjectInfo.FindCopiedContent(outputDir, outputFiles, projectFileToInfo);
-				}
-			}
-
-			// Add the content which is copied to the output directory
-			foreach (KeyValuePair<FileReference, bool> contentReference in ContentReferences)
-			{
-				FileReference contentFile = contentReference.Key;
-				if (contentReference.Value)
-				{
-					outputFiles.Add(FileReference.Combine(outputDir, contentFile.GetFileName()));
-				}
-			}
-		}
-
-		/// <summary>
-		/// Adds the given file and any additional build products to the output set
-		/// </summary>
-		/// <param name="outputFile">The assembly to add</param>
-		/// <param name="outputFiles">Set to receive the file and support files</param>
-		public static void AddReferencedAssemblyAndSupportFiles(FileReference outputFile, HashSet<FileReference> outputFiles)
-		{
-			outputFiles.Add(outputFile);
-
-			FileReference symbolFile = outputFile.ChangeExtension(".pdb");
-			if (FileReference.Exists(symbolFile))
-			{
-				outputFiles.Add(symbolFile);
-			}
-
-			FileReference documentationFile = outputFile.ChangeExtension(".xml");
-			if (FileReference.Exists(documentationFile))
-			{
-				outputFiles.Add(documentationFile);
-			}
-		}
-
-		/// <summary>
-		/// Determines if a TargetFramework is a .NET core framework
-		/// </summary>
-		/// <returns>True if the TargetFramework is a .NET core framework</returns>
-		static bool IsDotNETCoreFramework(string targetFramework)
-		{
-			if (targetFramework.ToLower().Contains("netstandard", StringComparison.Ordinal) || targetFramework.ToLower().Contains("netcoreapp", StringComparison.Ordinal))
-			{
-				return true;
-			}
-			else if (targetFramework.StartsWith("net", StringComparison.OrdinalIgnoreCase))
-			{
-				string[] versionSplit = targetFramework.Substring(3).Split('.');
-				if (versionSplit.Length >= 1 && Int32.TryParse(versionSplit[0], out int majorVersion))
-				{
-					return majorVersion >= 5;
-				}
-			}
-			return false;
-		}
-
-		/// <summary>
-		/// Determines if this project is a .NET core project
-		/// </summary>
-		/// <returns>True if the project is a .NET core project</returns>
-		public bool IsDotNETCoreProject()
-		{
-			if (Properties.TryGetValue("TargetFramework", out string? targetFramework))
-			{
-				return IsDotNETCoreFramework(targetFramework);
-			}
-
-			if (Properties.TryGetValue("TargetFrameworks", out string? targetFrameworks))
-			{
-				return targetFramework?.Split(';', StringSplitOptions.RemoveEmptyEntries).Any(x => IsDotNETCoreFramework(x.Trim())) ?? false;
-			}
-
-			return false;
-		}
-
-		/// <summary>
-		/// Adds a build product to the output list if it exists
-		/// </summary>
-		/// <param name="buildProduct">The build product to add</param>
-		/// <param name="buildProducts">List of output build products</param>
-		public static void AddOptionalBuildProduct(FileReference buildProduct, HashSet<FileReference> buildProducts)
-		{
-			if (FileReference.Exists(buildProduct))
-			{
-				buildProducts.Add(buildProduct);
-			}
-		}
-
 		/// <summary>
 		/// Parses supported elements from the children of the provided note. May recurse
 		/// for conditional elements.
diff --git a/Engine/Source/Programs/Shared/EpicGames.Build/DotnetProjectInfo.cs b/Engine/Source/Programs/Shared/EpicGames.Build/DotnetProjectInfo.cs
new file mode 100644
--- /dev/null
+++ b/Engine/Source/Programs/Shared/EpicGames.Build/DotnetProjectInfo.cs
@@ -0,0 +1,460 @@
+// Copyright Epic Games, Inc. All Rights Reserved.
+
+using System;
+using System.Collections.Generic;
+using System.Diagnostics.CodeAnalysis;
+using System.IO;
+using System.Linq;
+using System.Runtime.CompilerServices;
+using System.Runtime.InteropServices;
+using EpicGames.Core;
+using EpicGames.MsBuild;
+using Microsoft.Build.Definition;
+using Microsoft.Build.Evaluation;
+using OpenTracing;
+
+namespace UnrealBuildBase;
+
+public interface IDotnetProjectInfo
+{
+	/// <summary>
+	/// Evaluated properties from the project file
+	/// </summary>
+	public IReadOnlyDictionary<string, string> Properties { get; }
+
+	/// <summary>
+	/// Mapping of referenced assemblies to their 'CopyLocal' (aka 'Private') setting.
+	/// </summary>
+	public IReadOnlyDictionary<FileReference, bool> References { get; }
+
+	/// <summary>
+	/// Mapping of referenced projects to their 'CopyLocal' (aka 'Private') setting.
+	/// </summary>
+	public IReadOnlyDictionary<FileReference, bool> ProjectReferences { get; }
+
+	/// <summary>
+	/// List of compile references in the project.
+	/// </summary>
+	public IReadOnlyList<FileReference> CompileReferences { get; }
+
+	/// <summary>
+	/// Mapping of content to if they are flagged Always or Newer
+	/// </summary>
+	public IReadOnlyDictionary<FileReference, bool> ContentReferences { get; }
+
+	/// <summary>
+	/// Path to the CSProject file
+	/// </summary>
+	public FileReference ProjectPath { get; }
+}
+
+/// <summary>
+/// Information from a C# project file parsed by MSBuild.
+/// </summary>
+public sealed record class DotnetProjectInfo : IDotnetProjectInfo
+{
+	private DotnetProjectInfo()
+	{
+	}
+
+	/// <inheritdoc />
+	public required IReadOnlyDictionary<string, string> Properties { get; init; }
+
+	/// <inheritdoc />
+	public required IReadOnlyDictionary<FileReference, bool> References { get; init; }
+
+	/// <inheritdoc />
+	public required IReadOnlyDictionary<FileReference, bool> ProjectReferences { get; init; }
+
+	/// <inheritdoc />
+	public required IReadOnlyList<FileReference> CompileReferences { get; init; }
+
+	/// <inheritdoc />
+	public required IReadOnlyDictionary<FileReference, bool> ContentReferences { get; init; }
+
+	/// <inheritdoc />
+	public required FileReference ProjectPath { get; init; }
+
+	/// <summary>
+	/// Reads project information for the given file.
+	/// </summary>
+	/// <param name="file">The project file to read</param>
+	/// <param name="properties">Initial set of property values</param>
+	/// <returns>The parsed project info</returns>
+	public static DotnetProjectInfo Read(FileReference file, IReadOnlyDictionary<string, string> properties)
+	{
+		DotnetProjectInfo? project;
+		if (!TryRead(file, properties, out project))
+		{
+			throw new Exception(String.Format("Unable to read '{0}'", file));
+		}
+		return project;
+	}
+
+	/// <summary>
+	/// Attempts to read project information for the given file.
+	/// </summary>
+	/// <param name="file">The project file to read</param>
+	/// <param name="globalProperties">Initial set of property values</param>
+	/// <param name="outProjectInfo">If successful, the parsed project info</param>
+	/// <returns>True if the project was read successfully, false otherwise</returns>
+	public static bool TryRead(
+		FileReference file,
+		IReadOnlyDictionary<string, string> globalProperties,
+		[NotNullWhen(true)] out DotnetProjectInfo? outProjectInfo)
+	{
+		// First register MSBuild so we can use it to parse our project file. By not having the internal routine
+		// inline, we avoid having the issue of the Microsoft.Build libraries being resolved prior to the build path
+		// being set.
+		MsBuildSupport.RegisterMsBuildPath(MsBuildRegistration.Instance);
+
+		return TryReadImpl(file, globalProperties, out outProjectInfo);
+	}
+
+	[MethodImpl(MethodImplOptions.NoInlining)]
+	private static bool TryReadImpl(
+		FileReference file,
+		IReadOnlyDictionary<string, string> globalProperties,
+		[NotNullWhen(true)] out DotnetProjectInfo? outProjectInfo)
+	{
+		try
+		{
+			using ProjectCollection projectCollection = new();
+			ProjectOptions projectOptions = new()
+			{
+				GlobalProperties = globalProperties.ToDictionary(),
+				LoadSettings = ProjectLoadSettings.IgnoreMissingImports | ProjectLoadSettings.IgnoreInvalidImports,
+				ProjectCollection = projectCollection,
+			};
+			Project project = Project.FromFile(file.FullName, projectOptions);
+
+			Dictionary<string, string> properties = project.Properties.ToDictionary(p => p.Name, p => p.EvaluatedValue);
+
+			DirectoryReference projectDirectory = new(project.DirectoryPath);
+			FileReference GetFullPath(string path) =>
+				Path.IsPathFullyQualified(path) ? new FileReference(path) : FileReference.Combine(projectDirectory, path);
+
+			bool IsPrivate(ProjectItem item) => !item.HasMetadata("Private") || Boolean.Parse(item.GetMetadataValue("Private"));
+
+			IEnumerable<ProjectItem> validReferences = project.GetItems("Reference")
+				.Where(r => r.HasMetadata("HintPath"));
+			Dictionary<FileReference, bool> references = validReferences
+				.ToDictionary(r => GetFullPath(r.GetMetadataValue("HintPath")), IsPrivate);
+
+			Dictionary<FileReference, bool> projectReferences = project.GetItems("ProjectReference")
+				.ToDictionary(pr => GetFullPath(pr.EvaluatedInclude), IsPrivate);
+
+			List<FileReference> compileReferences = [.. project.GetItems("Compile").Select(c => GetFullPath(c.EvaluatedInclude))];
+
+			bool ShouldCopy(ProjectItem item)
+			{
+				string copyTo = item.GetMetadataValue("CopyToOutputDirectory");
+				return copyTo.Equals("Always", StringComparison.OrdinalIgnoreCase) || copyTo.Equals("PreserveNewest", StringComparison.OrdinalIgnoreCase);
+			}
+
+			IEnumerable<ProjectItem> contentAndNoneRefs = project.GetItems("Content").Concat(project.GetItems("None"));
+			Dictionary<FileReference, bool> contentReferences = contentAndNoneRefs
+				.ToDictionary(cr => GetFullPath(cr.EvaluatedInclude), ShouldCopy);
+
+			outProjectInfo = new()
+			{
+				Properties = properties,
+				References = references,
+				ProjectReferences = projectReferences,
+				CompileReferences = compileReferences,
+				ContentReferences = contentReferences,
+				ProjectPath = file,
+			};
+			return true;
+		}
+		catch (Exception)
+		{
+			// Do nothing.
+		}
+
+		outProjectInfo = null;
+		return false;
+	}
+
+	private class MsBuildRegistration : IMsBuildRegistration
+	{
+		internal static readonly MsBuildRegistration Instance = new MsBuildRegistration();
+
+		public DirectoryReference DotnetDirectory => Unreal.DotnetDirectory;
+
+		public FileReference DotnetPath => Unreal.DotnetPath;
+	}
+}
+
+public static class IDotnetProjectInfoExtensions
+{
+	/// <summary>
+	/// Determines if this project is a .NET Core or .NET 5+ project
+	/// </summary>
+	public static bool IsDotnetCoreOrDotnet5Plus(this IDotnetProjectInfo self)
+	{
+		static bool IsDotnetCoreOrDotnet5Plus(string targetFramework)
+		{
+			if (targetFramework.Contains("netstandard", StringComparison.OrdinalIgnoreCase) || targetFramework.Contains("netcoreapp", StringComparison.OrdinalIgnoreCase))
+			{
+				return true;
+			}
+			else if (targetFramework.StartsWith("net", StringComparison.OrdinalIgnoreCase))
+			{
+				string[] versionSplit = targetFramework.Substring(3).Split('.');
+				if (versionSplit.Length >= 1 && Int32.TryParse(versionSplit[0], out int majorVersion))
+				{
+					return majorVersion >= 5;
+				}
+			}
+			return false;
+		}
+
+		if (self.Properties.TryGetValue("TargetFramework", out string? targetFramework))
+		{
+			return IsDotnetCoreOrDotnet5Plus(targetFramework);
+		}
+
+		if (self.Properties.TryGetValue("TargetFrameworks", out string? targetFrameworks))
+		{
+			return targetFramework?.Split(';', StringSplitOptions.RemoveEmptyEntries).Any(x => IsDotnetCoreOrDotnet5Plus(x.Trim())) ?? false;
+		}
+
+		return false;
+	}
+
+	/// <summary>
+	/// Returns the assembly name used by this project
+	/// </summary>
+	/// <returns></returns>
+	public static bool TryGetAssemblyName(this IDotnetProjectInfo self, [NotNullWhen(true)] out string? assemblyName)
+	{
+		if (self.Properties.TryGetValue("AssemblyName", out assemblyName))
+		{
+			return true;
+		}
+		else if (self.IsDotnetCoreOrDotnet5Plus())
+		{
+			assemblyName = self.ProjectPath.GetFileNameWithoutExtension();
+			return true;
+		}
+		return false;
+	}
+
+	/// <summary>
+	/// Determines all the compiled build products (executable, etc...) directly built from this project.
+	/// </summary>
+	/// <param name="self">The target .NET project info</param>
+	/// <param name="outputDir">The output directory</param>
+	/// <param name="buildProducts">Receives the set of build products</param>
+	public static void FindCompiledBuildProducts(this IDotnetProjectInfo self, DirectoryReference outputDir, HashSet<FileReference> buildProducts)
+	{
+		static void AddOptionalBuildProduct(FileReference buildProduct, HashSet<FileReference> buildProducts)
+		{
+			if (FileReference.Exists(buildProduct))
+			{
+				buildProducts.Add(buildProduct);
+			}
+		}
+
+		string? outputType, assemblyName;
+		if (self.Properties.TryGetValue("OutputType", out outputType) && self.TryGetAssemblyName(out assemblyName))
+		{
+			switch (outputType)
+			{
+				case "Exe":
+				case "WinExe":
+					string executableExtension = RuntimeInformation.IsOSPlatform(OSPlatform.Windows) ? ".exe" : "";
+					buildProducts.Add(FileReference.Combine(outputDir, assemblyName + executableExtension));
+					// dotnet outputs a apphost executable and a dll with the actual assembly
+					AddOptionalBuildProduct(FileReference.Combine(outputDir, assemblyName + ".dll"), buildProducts);
+					AddOptionalBuildProduct(FileReference.Combine(outputDir, assemblyName + ".pdb"), buildProducts);
+					AddOptionalBuildProduct(FileReference.Combine(outputDir, assemblyName + ".exe.config"), buildProducts);
+					AddOptionalBuildProduct(FileReference.Combine(outputDir, assemblyName + ".exe.mdb"), buildProducts);
+					break;
+				case "Library":
+					buildProducts.Add(FileReference.Combine(outputDir, assemblyName + ".dll"));
+					AddOptionalBuildProduct(FileReference.Combine(outputDir, assemblyName + ".pdb"), buildProducts);
+					AddOptionalBuildProduct(FileReference.Combine(outputDir, assemblyName + ".dll.config"), buildProducts);
+					AddOptionalBuildProduct(FileReference.Combine(outputDir, assemblyName + ".dll.mdb"), buildProducts);
+					break;
+			}
+		}
+	}
+
+	/// <summary>
+	/// Resolve the project's output directory
+	/// </summary>
+	/// <param name="self">The target .NET project info</param>
+	/// <param name="outputDir">If successful, receives the output directory</param>
+	/// <returns>True if the output directory was found</returns>
+	public static bool TryGetOutputDir(this IDotnetProjectInfo self, [NotNullWhen(true)] out DirectoryReference? outputDir)
+	{
+		string? outputPath;
+		if (self.Properties.TryGetValue("OutputPath", out outputPath))
+		{
+			outputDir = DirectoryReference.Combine(self.ProjectPath.Directory, outputPath);
+			return true;
+		}
+		else if (self.IsDotnetCoreOrDotnet5Plus())
+		{
+			string configuration = self.Properties.ContainsKey("Configuration") ? self.Properties["Configuration"] : "Development";
+			outputDir = DirectoryReference.Combine(self.ProjectPath.Directory, "bin", configuration, self.Properties["TargetFramework"]);
+			return true;
+		}
+		else
+		{
+			outputDir = null;
+			return false;
+		}
+	}
+
+	/// <summary>
+	/// Get the ouptut file for this project
+	/// </summary>
+	/// <param name="self">The target .NET project info</param>
+	/// <param name="file">If successful, receives the assembly path</param>
+	/// <returns>True if the output file was found</returns>
+	public static bool TryGetOutputFile(this IDotnetProjectInfo self, [NotNullWhen(true)] out FileReference? file)
+	{
+		DirectoryReference? outputDir;
+		if (!self.TryGetOutputDir(out outputDir))
+		{
+			file = null;
+			return false;
+		}
+
+		string? assemblyName;
+		if (!self.TryGetAssemblyName(out assemblyName))
+		{
+			file = null;
+			return false;
+		}
+
+		file = FileReference.Combine(outputDir, assemblyName + ".dll");
+		return true;
+	}
+
+	/// <summary>
+	/// Resolve the project's output directory
+	/// </summary>
+	/// <param name="self">The target .NET project info</param>
+	/// <param name="baseDirectory">Base directory to resolve relative paths to</param>
+	/// <returns>The configured output directory</returns>
+	public static DirectoryReference GetOutputDir(this IDotnetProjectInfo self, DirectoryReference baseDirectory)
+	{
+		string? outputPath;
+		if (self.Properties.TryGetValue("OutputPath", out outputPath))
+		{
+			return DirectoryReference.Combine(baseDirectory, outputPath);
+		}
+		else if (self.IsDotnetCoreOrDotnet5Plus())
+		{
+			string configuration = self.Properties.TryGetValue("Configuration", out string? value) ? value : "Development";
+			return !self.Properties.TryGetValue("AppendTargetFrameworkToOutputPath", out string? appendFrameworkStr) || appendFrameworkStr.Equals("true", StringComparison.OrdinalIgnoreCase)
+				? DirectoryReference.Combine(baseDirectory, "bin", configuration, self.Properties["TargetFramework"])
+				: DirectoryReference.Combine(baseDirectory, "bin", configuration);
+		}
+		else
+		{
+			return baseDirectory;
+		}
+	}
+
+	/// <summary>
+	/// Adds the given file and any additional build products to the output set
+	/// </summary>
+	/// <param name="self">The target .NET project info</param>
+	/// <param name="outputFile">The assembly to add</param>
+	/// <param name="outputFiles">Set to receive the file and support files</param>
+	public static void AddReferencedAssemblyAndSupportFiles(this IDotnetProjectInfo self, FileReference outputFile, HashSet<FileReference> outputFiles)
+	{
+		outputFiles.Add(outputFile);
+
+		FileReference symbolFile = outputFile.ChangeExtension(".pdb");
+		if (FileReference.Exists(symbolFile))
+		{
+			outputFiles.Add(symbolFile);
+		}
+
+		FileReference documentationFile = outputFile.ChangeExtension(".xml");
+		if (FileReference.Exists(documentationFile))
+		{
+			outputFiles.Add(documentationFile);
+		}
+	}
+
+	/// <summary>
+	/// Finds all build products from this project. This includes content and other assemblies marked to be copied local.
+	/// </summary>
+	/// <param name="self">The target .NET project info</param>
+	/// <param name="outputDir">The output directory</param>
+	/// <param name="buildProducts">Receives the set of build products</param>
+	/// <param name="projectFileToInfo">Map of project file to information, to resolve build products from referenced projects copied locally</param>
+	public static void FindBuildProducts(
+		this IDotnetProjectInfo self,
+		DirectoryReference outputDir,
+		HashSet<FileReference> buildProducts,
+		Dictionary<FileReference, IDotnetProjectInfo> projectFileToInfo)
+	{
+		// Add the standard build products
+		self.FindCompiledBuildProducts(outputDir, buildProducts);
+
+		// Add the referenced assemblies which are marked to be copied into the output directory. This only happens for the main project, and does not happen for referenced projects.
+		foreach (KeyValuePair<FileReference, bool> reference in self.References)
+		{
+			if (reference.Value)
+			{
+				FileReference outputFile = FileReference.Combine(outputDir, reference.Key.GetFileName());
+				self.AddReferencedAssemblyAndSupportFiles(outputFile, buildProducts);
+			}
+		}
+
+		// Copy the build products for any referenced projects. Note that this does NOT operate recursively.
+		foreach (KeyValuePair<FileReference, bool> projectReference in self.ProjectReferences)
+		{
+			IDotnetProjectInfo? otherProjectInfo;
+			if (projectFileToInfo.TryGetValue(projectReference.Key, out otherProjectInfo))
+			{
+				otherProjectInfo.FindCompiledBuildProducts(outputDir, buildProducts);
+			}
+		}
+
+		// Add any copied content. This DOES operate recursively.
+		self.FindCopiedContent(outputDir, buildProducts, projectFileToInfo);
+	}
+
+	/// <summary>
+	/// Finds all content which will be copied into the output directory for this project. This includes content from any project references as "copy local" recursively (though MSBuild only traverses a single reference for actual binaries, in such cases)
+	/// </summary>
+	/// <param name="self">The target .NET project info</param>
+	/// <param name="outputDir">The output directory</param>
+	/// <param name="outputFiles">Receives the set of build products</param>
+	/// <param name="projectFileToInfo">Map of project file to information, to resolve build products from referenced projects copied locally</param>
+	private static void FindCopiedContent(
+		this IDotnetProjectInfo self,
+		DirectoryReference outputDir,
+		HashSet<FileReference> outputFiles,
+		Dictionary<FileReference, IDotnetProjectInfo> projectFileToInfo)
+	{
+		// Copy any referenced projects too.
+		foreach (KeyValuePair<FileReference, bool> projectReference in self.ProjectReferences)
+		{
+			IDotnetProjectInfo? otherProjectInfo;
+			if (projectFileToInfo.TryGetValue(projectReference.Key, out otherProjectInfo))
+			{
+				otherProjectInfo.FindCopiedContent(outputDir, outputFiles, projectFileToInfo);
+			}
+		}
+
+		// Add the content which is copied to the output directory
+		foreach (KeyValuePair<FileReference, bool> contentReference in self.ContentReferences)
+		{
+			FileReference contentFile = contentReference.Key;
+			if (contentReference.Value)
+			{
+				outputFiles.Add(FileReference.Combine(outputDir, contentFile.GetFileName()));
+			}
+		}
+	}
+}
diff --git a/Engine/Source/Programs/Shared/EpicGames.Build/System/CompileScriptModules.cs b/Engine/Source/Programs/Shared/EpicGames.Build/System/CompileScriptModules.cs
--- a/Engine/Source/Programs/Shared/EpicGames.Build/System/CompileScriptModules.cs
+++ b/Engine/Source/Programs/Shared/EpicGames.Build/System/CompileScriptModules.cs
@@ -107,9 +107,9 @@ namespace UnrealBuildBase
 
 			DirectoryReference CsProjBuildHook.EngineDirectory => Unreal.EngineDirectory;
 
-			DirectoryReference CsProjBuildHook.DotnetDirectory => Unreal.DotnetDirectory;
+			DirectoryReference IMsBuildRegistration.DotnetDirectory => Unreal.DotnetDirectory;
 
-			FileReference CsProjBuildHook.DotnetPath => Unreal.DotnetPath;
+			FileReference IMsBuildRegistration.DotnetPath => Unreal.DotnetPath;
 		}
 
 		/// <summary>
diff --git a/Engine/Source/Programs/Shared/EpicGames.Core/CsProjBuildHook.cs b/Engine/Source/Programs/Shared/EpicGames.Core/CsProjBuildHook.cs
--- a/Engine/Source/Programs/Shared/EpicGames.Core/CsProjBuildHook.cs
+++ b/Engine/Source/Programs/Shared/EpicGames.Core/CsProjBuildHook.cs
@@ -15,7 +15,7 @@ namespace EpicGames.Core
 	/// in EpicGames.MsBuild.  If they are, there is a strong chance of running into an issue 
 	/// gathering types (Assembly.GetTypes()) on EpicGames.Build.  
 	/// </summary>
-	public interface CsProjBuildHook
+	public interface CsProjBuildHook : IMsBuildRegistration
 	{
 		/// <summary>
 		/// Test the cache for a given file to get the last write time of the given file 
@@ -62,15 +62,5 @@ namespace EpicGames.Core
 		/// Unreal engine directory
 		/// </summary>
 		DirectoryReference EngineDirectory { get; }
-
-		/// <summary>
-		/// Dotnet directory shipped with the engine
-		/// </summary>
-		DirectoryReference DotnetDirectory { get; }
-
-		/// <summary>
-		/// Dotnet program
-		/// </summary>
-		FileReference DotnetPath { get; }
 	}
 }
diff --git a/Engine/Source/Programs/Shared/EpicGames.Core/IMsBuildRegistration.cs b/Engine/Source/Programs/Shared/EpicGames.Core/IMsBuildRegistration.cs
new file mode 100644
--- /dev/null
+++ b/Engine/Source/Programs/Shared/EpicGames.Core/IMsBuildRegistration.cs
@@ -0,0 +1,20 @@
+// Copyright Epic Games, Inc. All Rights Reserved.
+
+namespace EpicGames.Core;
+
+/// <summary>
+/// See <see cref="CsProjBuildHook"/>.
+/// </summary>
+public interface IMsBuildRegistration
+{
+
+	/// <summary>
+	/// The directory of the dotnet install shipped with Unreal Engine.
+	/// </summary>
+	DirectoryReference DotnetDirectory { get; }
+
+	/// <summary>
+	/// The dotnet executable that comprises part of the dotnet install shipped with Unreal Engine.
+	/// </summary>
+	FileReference DotnetPath { get; }
+}
diff --git a/Engine/Source/Programs/Shared/EpicGames.MsBuild/CsProjBuilder.cs b/Engine/Source/Programs/Shared/EpicGames.MsBuild/CsProjBuilder.cs
--- a/Engine/Source/Programs/Shared/EpicGames.MsBuild/CsProjBuilder.cs
+++ b/Engine/Source/Programs/Shared/EpicGames.MsBuild/CsProjBuilder.cs
@@ -14,7 +14,6 @@ using System.Xml;
 using System.Xml.Linq;
 using EpicGames.Core;
 using Microsoft.Build.Evaluation;
-using Microsoft.Build.Locator;
 using Microsoft.Extensions.Logging;
 
 namespace EpicGames.MsBuild
@@ -62,11 +61,10 @@ namespace EpicGames.MsBuild
 			bool bForceCompile, out bool bBuildSuccess, CsProjBuildHook hook, IEnumerable<DirectoryReference> baseDirectories,
 			IEnumerable<string> defineConstants, ILogger logger)
 		{
-
 			// Register the MS build path prior to invoking the internal routine.  By not having the internal routine
 			// inline, we avoid having the issue of the Microsoft.Build libraries being resolved prior to the build path
 			// being set.
-			RegisterMsBuildPath(hook);
+			MsBuildSupport.RegisterMsBuildPath(hook);
 			return BuildInternal(foundProjects, bForceCompile, out bBuildSuccess, hook, baseDirectories, defineConstants, logger);
 		}
 
@@ -550,58 +548,5 @@ namespace EpicGames.MsBuild
 
 			return result;
 		}
-
-		static bool s_hasRegisteredMsBuildPath = false;
-
-		/// <summary>
-		/// Register our bundled dotnet installation to be used by Microsoft.Build
-		/// This needs to happen in a function called before the first use of any Microsoft.Build types
-		/// </summary>
-		public static void RegisterMsBuildPath(CsProjBuildHook hook)
-		{
-			if (s_hasRegisteredMsBuildPath)
-			{
-				return;
-			}
-			s_hasRegisteredMsBuildPath = true;
-
-			// Find our bundled dotnet SDK
-			List<string> listOfSdks = [];
-			ProcessStartInfo startInfo = new ProcessStartInfo
-			{
-				FileName = hook.DotnetPath.FullName,
-				RedirectStandardOutput = true,
-				UseShellExecute = false,
-				ArgumentList = { "--list-sdks" }
-			};
-			startInfo.EnvironmentVariables["DOTNET_MULTILEVEL_LOOKUP"] = "0"; // use only the bundled dotnet installation - ignore any other/system dotnet install
-
-			Process dotnetProcess = Process.Start(startInfo);
-			{
-				string line;
-				while ((line = dotnetProcess.StandardOutput.ReadLine()) != null)
-				{
-					listOfSdks.Add(line);
-				}
-			}
-			dotnetProcess.WaitForExit();
-
-			if (listOfSdks.Count != 1)
-			{
-				throw new Exception("Expected only one sdk installed for bundled dotnet");
-			}
-
-			// Expected output has this form:
-			// 3.1.403 [D:\UE5_Main\engine\binaries\ThirdParty\DotNet\Windows\sdk]
-			string sdkVersion = listOfSdks[0].Split(' ')[0];
-
-			DirectoryReference dotnetSdkDirectory = DirectoryReference.Combine(hook.DotnetDirectory, "sdk", sdkVersion);
-			if (!DirectoryReference.Exists(dotnetSdkDirectory))
-			{
-				throw new Exception($"Failed to find .NET SDK directory: {dotnetSdkDirectory.FullName}");
-			}
-
-			MSBuildLocator.RegisterMSBuildPath(dotnetSdkDirectory.FullName);
-		}
 	}
 }
diff --git a/Engine/Source/Programs/Shared/EpicGames.MsBuild/MsBuildSupport.cs b/Engine/Source/Programs/Shared/EpicGames.MsBuild/MsBuildSupport.cs
new file mode 100644
--- /dev/null
+++ b/Engine/Source/Programs/Shared/EpicGames.MsBuild/MsBuildSupport.cs
@@ -0,0 +1,68 @@
+// Copyright Epic Games, Inc. All Rights Reserved.
+
+using System;
+using System.Collections.Generic;
+using System.Diagnostics;
+using EpicGames.Core;
+using Microsoft.Build.Locator;
+
+namespace EpicGames.MsBuild;
+
+/// <summary>
+/// Support for loading MSBuild assemblies from a dotnet install.
+/// </summary>
+public static class MsBuildSupport
+{
+	static bool s_hasRegisteredMsBuildPath = false;
+
+	/// <summary>
+	/// Register our bundled dotnet installation to be used by Microsoft.Build
+	/// This needs to happen in a function called before the first use of any Microsoft.Build types
+	/// </summary>
+	public static void RegisterMsBuildPath(IMsBuildRegistration hook)
+	{
+		if (s_hasRegisteredMsBuildPath)
+		{
+			return;
+		}
+		s_hasRegisteredMsBuildPath = true;
+
+		// Find our bundled dotnet SDK
+		List<string> listOfSdks = [];
+		ProcessStartInfo startInfo = new ProcessStartInfo
+		{
+			FileName = hook.DotnetPath.FullName,
+			RedirectStandardOutput = true,
+			UseShellExecute = false,
+			ArgumentList = { "--list-sdks" }
+		};
+		startInfo.EnvironmentVariables["DOTNET_MULTILEVEL_LOOKUP"] = "0"; // use only the bundled dotnet installation - ignore any other/system dotnet install
+
+		Process dotnetProcess = Process.Start(startInfo);
+		{
+			string line;
+			while ((line = dotnetProcess.StandardOutput.ReadLine()) != null)
+			{
+				listOfSdks.Add(line);
+			}
+		}
+		dotnetProcess.WaitForExit();
+
+		if (listOfSdks.Count != 1)
+		{
+			throw new Exception("Expected only one sdk installed for bundled dotnet");
+		}
+
+		// Expected output has this form:
+		// 3.1.403 [D:\UE5_Main\engine\binaries\ThirdParty\DotNet\Windows\sdk]
+		string sdkVersion = listOfSdks[0].Split(' ')[0];
+
+		DirectoryReference dotnetSdkDirectory = DirectoryReference.Combine(hook.DotnetDirectory, "sdk", sdkVersion);
+		if (!DirectoryReference.Exists(dotnetSdkDirectory))
+		{
+			throw new Exception($"Failed to find .NET SDK directory: {dotnetSdkDirectory.FullName}");
+		}
+
+		MSBuildLocator.RegisterMSBuildPath(dotnetSdkDirectory.FullName);
+	}
+}
diff --git a/Engine/Source/Programs/UnrealBuildTool/ProjectFiles/ProjectFileGenerator.cs b/Engine/Source/Programs/UnrealBuildTool/ProjectFiles/ProjectFileGenerator.cs
--- a/Engine/Source/Programs/UnrealBuildTool/ProjectFiles/ProjectFileGenerator.cs
+++ b/Engine/Source/Programs/UnrealBuildTool/ProjectFiles/ProjectFileGenerator.cs
@@ -3851,7 +3851,7 @@ namespace UnrealBuildTool
 				{
 					// Parse the project and check if UnrealEngine.csproj.props is imported
 					bool configsFound = false;
-					if (InProject is VCSharpProjectFile csProject && csProject.IsDotNETCoreProject())
+					if (InProject is VCSharpProjectFile csProject && csProject.IsDotnetCoreOrDotnet5Plus())
 					{
 						configsFound = true;
 						bFoundDevelopmentConfig = csProject.Configurations.Contains("Development");
diff --git a/Engine/Source/Programs/UnrealBuildTool/ProjectFiles/VisualStudio/VCProject.cs b/Engine/Source/Programs/UnrealBuildTool/ProjectFiles/VisualStudio/VCProject.cs
--- a/Engine/Source/Programs/UnrealBuildTool/ProjectFiles/VisualStudio/VCProject.cs
+++ b/Engine/Source/Programs/UnrealBuildTool/ProjectFiles/VisualStudio/VCProject.cs
@@ -2627,7 +2627,7 @@ namespace UnrealBuildTool
 
 				// Parse all the configurations and platforms
 				// Parse the basic structure of the document, updating properties and recursing into other referenced projects as we go
-				if (!IsDotNETCoreProject())
+				if (!IsDotnetCoreOrDotnet5Plus())
 				{
 					foreach (XmlElement Element in Document.DocumentElement.ChildNodes.OfType<XmlElement>())
 					{
@@ -2712,16 +2712,16 @@ namespace UnrealBuildTool
 		/// <summary>
 		/// Extract information from the csproj file based on the supplied configuration
 		/// </summary>
-		public CsProjectInfo? GetProjectInfo(UnrealTargetConfiguration InConfiguration)
+		public IDotnetProjectInfo? GetProjectInfo(UnrealTargetConfiguration InConfiguration)
 		{
-			if (CachedProjectInfo.ContainsKey(InConfiguration))
+			if (CachedProjectInfo.TryGetValue(InConfiguration, out IDotnetProjectInfo? value))
 			{
-				return CachedProjectInfo[InConfiguration];
+				return value;
 			}
 
 			CsProjectInfo? Info;
 
-			Dictionary<string, string> Properties = new Dictionary<string, string>();
+			Dictionary<string, string> Properties = [];
 			Properties.Add("Platform", "AnyCPU");
 			Properties.Add("Configuration", InConfiguration.ToString());
 			Properties.Add("EngineDirectory", Unreal.EngineDirectory.FullName);
@@ -2734,12 +2734,12 @@ namespace UnrealBuildTool
 		}
 
 		/// <summary>
-		/// Determine if this project is a .NET Core project
+		/// Determine if this project is a .NET Core or .NET 5+ project
 		/// </summary>
-		public bool IsDotNETCoreProject()
+		public bool IsDotnetCoreOrDotnet5Plus()
 		{
-			CsProjectInfo Info = GetProjectInfo(UnrealTargetConfiguration.Debug)!;
-			return Info.IsDotNETCoreProject();
+			IDotnetProjectInfo Info = GetProjectInfo(UnrealTargetConfiguration.Debug)!;
+			return Info.IsDotnetCoreOrDotnet5Plus();
 		}
 
 		/// <summary>
@@ -2747,7 +2747,7 @@ namespace UnrealBuildTool
 		/// </summary>
 		public string GetProjectProperty(string property)
 		{
-			CsProjectInfo Info = GetProjectInfo(UnrealTargetConfiguration.Debug)!;
+			IDotnetProjectInfo Info = GetProjectInfo(UnrealTargetConfiguration.Debug)!;
 			Info.Properties.TryGetValue(property, out string? value);
 			return value ?? String.Empty;
 		}
@@ -2806,6 +2806,6 @@ namespace UnrealBuildTool
 		}
 
 		/// Cache of parsed info about this project
-		protected readonly Dictionary<UnrealTargetConfiguration, CsProjectInfo> CachedProjectInfo = new Dictionary<UnrealTargetConfiguration, CsProjectInfo>();
+		protected readonly Dictionary<UnrealTargetConfiguration, IDotnetProjectInfo> CachedProjectInfo = [];
 	}
 }
diff --git a/Engine/Source/Programs/UnrealBuildTool/ProjectFiles/VisualStudioCode/VSCodeProjectFileGenerator.cs b/Engine/Source/Programs/UnrealBuildTool/ProjectFiles/VisualStudioCode/VSCodeProjectFileGenerator.cs
--- a/Engine/Source/Programs/UnrealBuildTool/ProjectFiles/VisualStudioCode/VSCodeProjectFileGenerator.cs
+++ b/Engine/Source/Programs/UnrealBuildTool/ProjectFiles/VisualStudioCode/VSCodeProjectFileGenerator.cs
@@ -431,7 +431,7 @@ namespace UnrealBuildTool
 				public UnrealTargetPlatform Platform { get; set; }
 				public EOutputType OutputType { get; set; }
 
-				public CsProjectInfo? CSharpInfo { get; set; }
+				public IDotnetProjectInfo? CSharpInfo { get; set; }
 
 				public override string ToString()
 				{
@@ -558,7 +558,7 @@ namespace UnrealBuildTool
 
 					foreach (UnrealTargetConfiguration Config in Configs)
 					{
-						CsProjectInfo? Info = VCSharpProject.GetProjectInfo(Config)!;
+						IDotnetProjectInfo? Info = VCSharpProject.GetProjectInfo(Config)!;
 
 						{
 							ProjectData.EOutputType OutputType;
