EpicGames.UHT: Replace cityhash with IoHash

#rnx
#Synced-CL 48947859

diff --git a/Engine/Source/Programs/Shared/EpicGames.UHT/Exporters/CodeGen/UhtCodeGenerator.cs b/Engine/Source/Programs/Shared/EpicGames.UHT/Exporters/CodeGen/UhtCodeGenerator.cs
--- a/Engine/Source/Programs/Shared/EpicGames.UHT/Exporters/CodeGen/UhtCodeGenerator.cs
+++ b/Engine/Source/Programs/Shared/EpicGames.UHT/Exporters/CodeGen/UhtCodeGenerator.cs
@@ -150,7 +150,7 @@ namespace EpicGames.UHT.Exporters.CodeGen
 			public Task? Task { get; set; }
 			public string IncludePath { get; set; }
 			public string FileId { get; set; }
-			public uint BodyHash { get; set; }
+			public IoHash BodyHash { get; set; }
 			public bool NeedsPushModelHeaders { get; set; }
 			public bool NeedsFastArrayHeaders { get; set; }
 			public bool NeedsVerseValue { get; set; }
@@ -185,7 +185,7 @@ namespace EpicGames.UHT.Exporters.CodeGen
 			public string StaticsName { get; set; }
 			public UhtClass? NativeInterface { get; set; }
 			public UhtProperty? FastArrayProperty { get; set; }
-			public uint Hash { get; set; }
+			public IoHash Hash { get; set; }
 		}
 		public ObjectInfo[] ObjectInfos { get; set; }
 
diff --git a/Engine/Source/Programs/Shared/EpicGames.UHT/Exporters/CodeGen/UhtHeaderCodeGeneratorCppFile.cs b/Engine/Source/Programs/Shared/EpicGames.UHT/Exporters/CodeGen/UhtHeaderCodeGeneratorCppFile.cs
--- a/Engine/Source/Programs/Shared/EpicGames.UHT/Exporters/CodeGen/UhtHeaderCodeGeneratorCppFile.cs
+++ b/Engine/Source/Programs/Shared/EpicGames.UHT/Exporters/CodeGen/UhtHeaderCodeGeneratorCppFile.cs
@@ -394,7 +394,7 @@ namespace EpicGames.UHT.Exporters.CodeGen
 					string name = $"Z_CompiledInDeferFile_{headerInfo.FileId}_{PackageInfos[package.PackageTypeIndex].StrippedName}";
 					string staticsName = $"{name}_Statics";
 
-					uint combinedHash = UInt32.MaxValue;
+					IoHash combinedHash = new(UInt64.MaxValue, UInt64.MaxValue, UInt32.MaxValue);
 
 					using UhtCodeBlockComment blockComment = new(builder, "Registration");
 					builder.Append($"struct {staticsName}\r\n");
@@ -421,13 +421,13 @@ namespace EpicGames.UHT.Exporters.CodeGen
 							builder => builder.Append("\tstatic constexpr FEnumRegisterCompiledInInfo EnumInfo[] = {\r\n"),
 							(builder, enumObj) =>
 							{ 
-								uint hash = ObjectInfos[enumObj.ObjectTypeIndex].Hash;
+								IoHash hash = ObjectInfos[enumObj.ObjectTypeIndex].Hash;
 								builder
 									.Append($"\t\t{{ {enumObj.SourceName}_StaticEnum, ")
 									.Append($"TEXT(\"{enumObj.EngineName}\"), ")
 									.Append($"&Z_Registration_Info_UEnum_{enumObj.SourceName}, ")
-									.Append($"CONSTRUCT_RELOAD_VERSION_INFO(FEnumReloadVersionInfo, {hash}U) }},\r\n");
-								combinedHash = HashCombine(combinedHash, hash);
+									.Append($"CONSTRUCT_RELOAD_VERSION_INFO(FEnumReloadVersionInfo, {(uint)hash}U) }},\r\n");
+								combinedHash = IoHash.Combine(combinedHash, hash);
 							},
 							builder => builder.Append("\t};\r\n"));
 
@@ -435,14 +435,14 @@ namespace EpicGames.UHT.Exporters.CodeGen
 							builder => builder.Append("\tstatic constexpr FStructRegisterCompiledInInfo ScriptStructInfo[] = {\r\n"),
 							(builder, scriptStruct) =>
 							{
-								uint hash = ObjectInfos[scriptStruct.ObjectTypeIndex].Hash;
+								IoHash hash = ObjectInfos[scriptStruct.ObjectTypeIndex].Hash;
 								builder
 									.Append($"\t\t{{ {scriptStruct.FullyQualifiedSourceName}::StaticStruct, ")
 									.Append($"{GetSingletonName(scriptStruct, UhtSingletonType.Statics)}::NewStructOps, ")
 									.Append($"TEXT(\"{scriptStruct.EngineName}\"),")
 									.Append($"&Z_Registration_Info_UScriptStruct_{scriptStruct.SourceName}, ")
-									.Append($"CONSTRUCT_RELOAD_VERSION_INFO(FStructReloadVersionInfo, sizeof({scriptStruct.FullyQualifiedSourceName}), {hash}U) }},\r\n");
-								combinedHash = HashCombine(combinedHash, hash);
+									.Append($"CONSTRUCT_RELOAD_VERSION_INFO(FStructReloadVersionInfo, sizeof({scriptStruct.FullyQualifiedSourceName}), {(uint)hash}U) }},\r\n");
+								combinedHash = IoHash.Combine(combinedHash, hash);
 							},
 							builder => builder.Append("\t};\r\n"));
 
@@ -450,14 +450,14 @@ namespace EpicGames.UHT.Exporters.CodeGen
 							builder => builder.Append("\tstatic constexpr FClassRegisterCompiledInInfo ClassInfo[] = {\r\n"),
 							(builder, classObj) =>
 							{
-								uint hash = ObjectInfos[classObj.ObjectTypeIndex].Hash;
+								IoHash hash = ObjectInfos[classObj.ObjectTypeIndex].Hash;
 								builder
 									.Append($"\t\t{{ {GetSingletonName(classObj, UhtSingletonType.Registered)}, ")
 									.Append($"{classObj.FullyQualifiedSourceName}::StaticClass, ")
 									.Append($"TEXT(\"{(classObj.IsVerseField ? classObj.EngineName : classObj.SourceName)}\"), ")
 									.Append($"&Z_Registration_Info_UClass_{classObj.SourceName}, ")
-									.Append($"CONSTRUCT_RELOAD_VERSION_INFO(FClassReloadVersionInfo, sizeof({classObj.FullyQualifiedSourceName}), {hash}U) }},\r\n");
-								combinedHash = HashCombine(combinedHash, hash);
+									.Append($"CONSTRUCT_RELOAD_VERSION_INFO(FClassReloadVersionInfo, sizeof({classObj.FullyQualifiedSourceName}), {(uint)hash}U) }},\r\n");
+								combinedHash = IoHash.Combine(combinedHash, hash);
 							},
 							builder => builder.Append("\t};\r\n"));
 					}
@@ -521,7 +521,7 @@ namespace EpicGames.UHT.Exporters.CodeGen
 					}
 
 					// Save the hash of the generated body 
-					HeaderInfos[HeaderFile.HeaderFileTypeIndex].BodyHash = UhtHash.GenenerateTextHash(generatedBody.Span[generatedBodyStart..generatedBodyEnd]);
+					HeaderInfos[HeaderFile.HeaderFileTypeIndex].BodyHash = IoHash.Compute(generatedBody.Span[generatedBodyStart..generatedBodyEnd]);
 				}
 			}
 		}
@@ -821,7 +821,7 @@ namespace EpicGames.UHT.Exporters.CodeGen
 
 			{
 				using UhtRentedPoolBuffer<char> borrowBuffer = builder.RentPoolBuffer(hashCodeBlockStart, builder.Length - hashCodeBlockStart);
-				ObjectInfos[enumObj.ObjectTypeIndex].Hash = UhtHash.GenenerateTextHash(borrowBuffer.Buffer.Memory.Span);
+				ObjectInfos[enumObj.ObjectTypeIndex].Hash = IoHash.Compute<char>(borrowBuffer.Buffer.Memory.Span);
 			}
 			return builder;
 		}
@@ -1179,7 +1179,7 @@ namespace EpicGames.UHT.Exporters.CodeGen
 
 			using (UhtRentedPoolBuffer<char> borrowBuffer = builder.RentPoolBuffer(hashCodeBlockStart, builder.Length - hashCodeBlockStart))
 			{
-				ObjectInfos[scriptStruct.ObjectTypeIndex].Hash = UhtHash.GenenerateTextHash(borrowBuffer.Buffer.Memory.Span);
+				ObjectInfos[scriptStruct.ObjectTypeIndex].Hash = IoHash.Compute<char>(borrowBuffer.Buffer.Memory.Span);
 			}
 
 			// if this struct has RigVM methods we need to implement both the 
@@ -1507,7 +1507,7 @@ namespace EpicGames.UHT.Exporters.CodeGen
 
 			using (UhtRentedPoolBuffer<char> borrowBuffer = builder.RentPoolBuffer(hashCodeBlockStart, builder.Length - hashCodeBlockStart))
 			{
-				ObjectInfos[function.ObjectTypeIndex].Hash = UhtHash.GenenerateTextHash(borrowBuffer.Buffer.Memory.Span);
+				ObjectInfos[function.ObjectTypeIndex].Hash = IoHash.Compute<char>(borrowBuffer.Buffer.Memory.Span);
 			}
 			return builder;
 		}
@@ -3089,7 +3089,7 @@ namespace EpicGames.UHT.Exporters.CodeGen
 
 				// Append base class' hash at the end of the generated code, this will force update derived classes
 				// when base class changes during hot-reload.
-				uint baseClassHash = 0;
+				IoHash baseClassHash = IoHash.Zero;
 				if (classObj.SuperClass != null && !classObj.SuperClass.ClassFlags.HasAnyFlags(EClassFlags.Intrinsic))
 				{
 					baseClassHash = ObjectInfos[classObj.SuperClass.ObjectTypeIndex].Hash;
@@ -3134,7 +3134,7 @@ namespace EpicGames.UHT.Exporters.CodeGen
 				// Calculate generated class initialization code hash so that we know when it changes after hot-reload
 				{
 					using UhtRentedPoolBuffer<char> borrowBuffer = hashBuilder.RentPoolBuffer();
-					ObjectInfos[classObj.ObjectTypeIndex].Hash = UhtHash.GenenerateTextHash(borrowBuffer.Buffer.Memory.Span);
+					ObjectInfos[classObj.ObjectTypeIndex].Hash = IoHash.Compute<char>(borrowBuffer.Buffer.Memory.Span);
 				}
 			}
 
@@ -3385,7 +3385,7 @@ namespace EpicGames.UHT.Exporters.CodeGen
 
 			public int GetObjectIndex(UhtObject obj) => _codeGenerator.GetObjectIndex(obj);
 
-			public uint GetTypeHash(UhtObject obj)
+			public IoHash GetTypeHash(UhtObject obj)
 			{
 				return _codeGenerator.ObjectInfos[obj.ObjectTypeIndex].Hash;
 			}
diff --git a/Engine/Source/Programs/Shared/EpicGames.UHT/Exporters/CodeGen/UhtPackageCodeGenerator.cs b/Engine/Source/Programs/Shared/EpicGames.UHT/Exporters/CodeGen/UhtPackageCodeGenerator.cs
--- a/Engine/Source/Programs/Shared/EpicGames.UHT/Exporters/CodeGen/UhtPackageCodeGenerator.cs
+++ b/Engine/Source/Programs/Shared/EpicGames.UHT/Exporters/CodeGen/UhtPackageCodeGenerator.cs
@@ -126,38 +126,6 @@ namespace EpicGames.UHT.Exporters.CodeGen
 			}
 			return false;
 		}
-
-		/// <summary>
-		/// Combines two hash values to get a third.
-		/// Note - this function is not commutative.
-		///
-		/// This function cannot change for backward compatibility reasons.
-		/// You may want to choose HashCombineFast for a better in-memory hash combining function.
-		/// 
-		/// NOTE: This is a copy of the method in TypeHash.h
-		/// </summary>
-		/// <param name="A">Hash to merge</param>
-		/// <param name="C">Previously combined hash</param>
-		/// <returns>Resulting hash value</returns>
-		[System.Diagnostics.CodeAnalysis.SuppressMessage("Style", "IDE1006:Naming Styles", Justification = "<Pending>")]
-		[System.Diagnostics.CodeAnalysis.SuppressMessage("Style", "IDE2001:Embedded statements must be on their own line", Justification = "<Pending>")]
-		public static uint HashCombine(uint A, uint C)
-		{
-			uint B = 0x9e3779b9;
-			A += B;
-
-			A -= B; A -= C; A ^= (C >> 13);
-			B -= C; B -= A; B ^= (A << 8);
-			C -= A; C -= B; C ^= (B >> 13);
-			A -= B; A -= C; A ^= (C >> 12);
-			B -= C; B -= A; B ^= (A << 16);
-			C -= A; C -= B; C ^= (B >> 5);
-			A -= B; A -= C; A ^= (C >> 3);
-			B -= C; B -= A; B ^= (A << 10);
-			C -= A; C -= B; C ^= (B >> 15);
-
-			return C;
-		}
 		#endregion
 	}
 
diff --git a/Engine/Source/Programs/Shared/EpicGames.UHT/Exporters/CodeGen/UhtPackageCodeGeneratorCppFile.cs b/Engine/Source/Programs/Shared/EpicGames.UHT/Exporters/CodeGen/UhtPackageCodeGeneratorCppFile.cs
--- a/Engine/Source/Programs/Shared/EpicGames.UHT/Exporters/CodeGen/UhtPackageCodeGeneratorCppFile.cs
+++ b/Engine/Source/Programs/Shared/EpicGames.UHT/Exporters/CodeGen/UhtPackageCodeGeneratorCppFile.cs
@@ -39,7 +39,7 @@ namespace EpicGames.UHT.Exporters.CodeGen
 				// Collect information from all of the headers
 				List<UhtField> singletons = new();
 				StringBuilder declarations = new();
-				uint bodiesHash = 0;
+				IoHash bodiesHash = IoHash.Zero;
 				foreach (UhtHeaderFile headerFile in packageSortedHeaders)
 				{
 					ref UhtCodeGenerator.HeaderInfo headerInfo = ref HeaderInfos[headerFile.HeaderFileTypeIndex];
@@ -56,15 +56,15 @@ namespace EpicGames.UHT.Exporters.CodeGen
 
 					singletons.AddRange(headerFile.References.Singletons);
 
-					uint bodyHash = HeaderInfos[headerFile.HeaderFileTypeIndex].BodyHash;
-					if (bodiesHash == 0)
+					IoHash bodyHash = HeaderInfos[headerFile.HeaderFileTypeIndex].BodyHash;
+					if (bodiesHash == IoHash.Zero)
 					{
 						// Don't combine in the first case because it keeps GUID backwards compatibility
 						bodiesHash = bodyHash;
 					}
 					else
 					{
-						bodiesHash = HashCombine(bodyHash, bodiesHash);
+						bodiesHash = IoHash.Combine(bodyHash, bodiesHash);
 					}
 				}
 
@@ -78,7 +78,7 @@ namespace EpicGames.UHT.Exporters.CodeGen
 					}
 					return;
 				}
-				uint declarationsHash = UhtHash.GenenerateTextHash(declarations.ToString());
+				IoHash declarationsHash = IoHash.Compute(declarations.ToString());
 
 				singletons.Sort((UhtField lhs, UhtField rhs) =>
 				{
@@ -229,8 +229,8 @@ namespace EpicGames.UHT.Exporters.CodeGen
 							.Append("\t\t},\r\n")
 							.Append($"\t\t.Flags = EPackageFlags(PKG_CompiledIn | 0x{(uint)flags:X8}),\r\n")
 							.Append($"\t\tIF_WITH_METADATA(.MetaData = ").AppendMetaDataArrayView(package.MetaData, null, MetaDataParamsName, 2, ",)\r\n")
-							.Append($"\t\t.BodiesHash = 0x{bodiesHash:X8},\r\n")
-							.Append($"\t\t.DeclarationsHash = 0x{declarationsHash:X8},\r\n")
+							.Append($"\t\t.BodiesHash = 0x{(uint)bodiesHash:X8},\r\n")
+							.Append($"\t\t.DeclarationsHash = 0x{(uint)declarationsHash:X8},\r\n")
 							.Append("});\r\n")
 							.Append("\r\n");
 					}
@@ -259,8 +259,8 @@ namespace EpicGames.UHT.Exporters.CodeGen
 						builder.Append("\t\t\t").Append(packageSingletons.Length != 0 ? "SingletonFuncArray" : "nullptr").Append(",\r\n");
 						builder.Append("\t\t\t").Append(packageSingletons.Length != 0 ? "UE_ARRAY_COUNT(SingletonFuncArray)" : "0").Append(",\r\n");
 						builder.Append("\t\t\tPKG_CompiledIn | ").Append($"0x{(uint)flags:X8}").Append(",\r\n");
-						builder.Append("\t\t\t").Append($"0x{bodiesHash:X8}").Append(",\r\n");
-						builder.Append("\t\t\t").Append($"0x{declarationsHash:X8}").Append(",\r\n");
+						builder.Append("\t\t\t").Append($"0x{(uint)bodiesHash:X8}").Append(",\r\n");
+						builder.Append("\t\t\t").Append($"0x{(uint)declarationsHash:X8}").Append(",\r\n");
 						builder.Append("\t\t\t").AppendMetaDataParams(package, null, MetaDataParamsName).Append("\r\n");
 						builder.Append("\t\t};\r\n");
 						builder.Append("\t\tUECodeGen_Private::ConstructUPackage(Z_Registration_Info_UPackage_").Append(strippedName).Append(".OuterSingleton, PackageParams);\r\n");
@@ -271,7 +271,7 @@ namespace EpicGames.UHT.Exporters.CodeGen
 						// Do not change the Z_CompiledInDeferPackage_UPackage_ without changing LC_SymbolPatterns
 						builder.Append("static FRegisterCompiledInInfo Z_CompiledInDeferPackage_UPackage_").Append(strippedName).Append('(').Append(singletonName)
 							.Append(", TEXT(\"").Append(package.SourceName).Append("\"), Z_Registration_Info_UPackage_").Append(strippedName).Append(", CONSTRUCT_RELOAD_VERSION_INFO(FPackageReloadVersionInfo, ")
-							.Append($"0x{bodiesHash:X8}, 0x{declarationsHash:X8}").Append("));\r\n");
+							.Append($"0x{(uint)bodiesHash:X8}, 0x{(uint)declarationsHash:X8}").Append("));\r\n");
 					}
 				}
 
diff --git a/Engine/Source/Programs/Shared/EpicGames.UHT/Types/UhtProperty.cs b/Engine/Source/Programs/Shared/EpicGames.UHT/Types/UhtProperty.cs
--- a/Engine/Source/Programs/Shared/EpicGames.UHT/Types/UhtProperty.cs
+++ b/Engine/Source/Programs/Shared/EpicGames.UHT/Types/UhtProperty.cs
@@ -684,7 +684,7 @@ namespace EpicGames.UHT.Types
 		/// </summary>
 		/// <param name="obj">Object in question</param>
 		/// <returns>Hash code</returns>
-		public uint GetTypeHash(UhtObject obj);
+		public IoHash GetTypeHash(UhtObject obj);
 
 	}
 
@@ -2847,8 +2847,8 @@ namespace EpicGames.UHT.Types
 			}
 
 			builder.Append(builder.Length == startingLength ? " // " : " ");
-			uint hash = context.GetTypeHash(obj);
-			if (hash == 0)
+			IoHash hash = context.GetTypeHash(obj);
+			if (hash == IoHash.Zero)
 			{
 				string type = referingType is UhtProperty ? "property" : "object";
 				referingType.LogError($"The {type} \"{referingType.SourceName}\" references type \"{obj.SourceName}\" but the code generation hash is zero.  Check for circular dependencies or missing includes.");
diff --git a/Engine/Source/Programs/Shared/EpicGames.UHT/Utils/UhtHash.cs b/Engine/Source/Programs/Shared/EpicGames.UHT/Utils/UhtHash.cs
deleted file mode 100644
