Wrap flushing JsonTraceSink in a try catch block as we occasionally see NullReferenceExceptions on Mac that don't seem to make sense and we don't want the job to fail. Replace the Epic custom JSON writer with the System.Text.Json equivalent in case this is contributing to the above issue.

#Synced-CL 48984147

diff --git a/Engine/Source/Programs/Shared/EpicGames.Core/TraceSpan.cs b/Engine/Source/Programs/Shared/EpicGames.Core/TraceSpan.cs
--- a/Engine/Source/Programs/Shared/EpicGames.Core/TraceSpan.cs
+++ b/Engine/Source/Programs/Shared/EpicGames.Core/TraceSpan.cs
@@ -6,6 +6,8 @@ using System.Diagnostics;
 using System.Globalization;
 using System.IO;
 using System.Reflection;
+using System.Text.Json;
+using Microsoft.Extensions.Logging;
 
 namespace EpicGames.Core
 {
@@ -111,6 +113,22 @@ namespace EpicGames.Core
 		/// Saves all the scope information to a file
 		/// </summary>
 		public void Flush()
+		{
+			// We've seen this method fail on Horde when running on Mac due to a NullReferenceException.
+			// We've been unable to find out what's causing the NRE, so try and catch the whole function
+			// if it NREs because not failing the job is more critical than getting the trace file written.
+			// We shouldn't broadly catch other exception types as they have not been observed so far.
+			try
+			{
+				FlushImpl();
+			}
+			catch (NullReferenceException ex)
+			{
+				Log.Logger.LogInformation(ex, "Caught NullReferenceException when flushing JsonTraceSink: {Exception}", ex);
+			}
+		}
+
+		private void FlushImpl()
 		{
 			FileReference file;
 			using (Process process = Process.GetCurrentProcess())
@@ -121,38 +139,37 @@ namespace EpicGames.Core
 				file = FileReference.Combine(_telemetryDir, fileName);
 			}
 
-			using (JsonWriter writer = new JsonWriter(file))
+			using FileStream stream = new(file.FullName, FileMode.Create, FileAccess.Write, FileShare.Read);
+			using Utf8JsonWriter writer = new(stream, new JsonWriterOptions { Indented = true });
+			writer.WriteStartObject();
+			writer.WriteStartArray("Spans");
+			foreach (TraceSpanImpl span in _spans)
 			{
-				writer.WriteObjectStart();
-				writer.WriteArrayStart("Spans");
-				foreach (TraceSpanImpl span in _spans)
+				if (span.FinishTime != null)
 				{
-					if (span.FinishTime != null)
+					writer.WriteStartObject();
+					writer.WriteString("Name", span.Name);
+					if (span.Resource != null)
+					{
+						writer.WriteString("Resource", span.Resource);
+					}
+					if (span.Service != null)
+					{
+						writer.WriteString("Service", span.Service);
+					}
+					writer.WriteString("StartTime", span.StartTime.ToString("o", CultureInfo.InvariantCulture));
+					writer.WriteString("FinishTime", span.FinishTime.Value.ToString("o", CultureInfo.InvariantCulture));
+					writer.WriteStartObject("Metadata");
+					foreach (KeyValuePair<string, string> pair in span.Metadata)
 					{
-						writer.WriteObjectStart();
-						writer.WriteValue("Name", span.Name);
-						if (span.Resource != null)
-						{
-							writer.WriteValue("Resource", span.Resource);
-						}
-						if (span.Service != null)
-						{
-							writer.WriteValue("Service", span.Service);
-						}
-						writer.WriteValue("StartTime", span.StartTime.ToString("o", CultureInfo.InvariantCulture));
-						writer.WriteValue("FinishTime", span.FinishTime.Value.ToString("o", CultureInfo.InvariantCulture));
-						writer.WriteObjectStart("Metadata");
-						foreach (KeyValuePair<string, string> pair in span.Metadata)
-						{
-							writer.WriteValue(pair.Key, pair.Value);
-						}
-						writer.WriteObjectEnd();
-						writer.WriteObjectEnd();
+						writer.WriteString(pair.Key, pair.Value);
 					}
+					writer.WriteEndObject();
+					writer.WriteEndObject();
 				}
-				writer.WriteArrayEnd();
-				writer.WriteObjectEnd();
 			}
+			writer.WriteEndArray();
+			writer.WriteEndObject();
 		}
 	}
 
