Add an EpicGames.Build.Tests project with tests for the MSBuild project parsing logic.

#Synced-CL 48984147

diff --git a/Engine/Source/Programs/Shared/EpicGames.Build.Tests/DotnetProjectInfoTests.cs b/Engine/Source/Programs/Shared/EpicGames.Build.Tests/DotnetProjectInfoTests.cs
new file mode 100644
--- /dev/null
+++ b/Engine/Source/Programs/Shared/EpicGames.Build.Tests/DotnetProjectInfoTests.cs
@@ -0,0 +1,58 @@
+// Copyright Epic Games, Inc. All Rights Reserved.
+
+using System.Collections.Generic;
+using System.Linq;
+using System.Reflection;
+using EpicGames.Core;
+using Microsoft.VisualStudio.TestTools.UnitTesting;
+using UnrealBuildBase;
+
+namespace EpicGames.Build.Tests;
+
+[TestClass]
+public class DotnetProjectInfoTests
+{
+	[TestMethod]
+	public void TestParseOwnProject()
+	{
+		FileReference ownProjectPath = GetOwnProjectPath();
+		DotnetProjectInfo projectInfo = DotnetProjectInfo.Read(ownProjectPath, new Dictionary<string, string>());
+
+		// Test properties from the Directory.Build.props file
+		CollectionAssert.Contains(projectInfo.Properties.Keys.ToList(), "IsPackable");
+		Assert.AreEqual(projectInfo.Properties["IsPackable"], "false");
+	}
+
+	[TestMethod]
+	public void TestCanParseSameProjectRepeatedly()
+	{
+		FileReference ownProjectPath = GetOwnProjectPath();
+		_ = DotnetProjectInfo.Read(ownProjectPath, new Dictionary<string, string>());
+		_ = DotnetProjectInfo.Read(ownProjectPath, new Dictionary<string, string>());
+		_ = DotnetProjectInfo.Read(ownProjectPath, new Dictionary<string, string>());
+	}
+
+	private static FileReference GetOwnProjectPath()
+	{
+		DirectoryReference? searchDirectory = new DirectoryReference(Assembly.GetExecutingAssembly().Location);
+		FileReference? maybeOwnProjectPath = null;
+		while (searchDirectory is not null)
+		{
+			maybeOwnProjectPath = FileReference.Combine(searchDirectory, "EpicGames.Build.Tests.csproj");
+			if (FileReference.Exists(maybeOwnProjectPath))
+			{
+				break;
+			}
+
+			searchDirectory = searchDirectory.ParentDirectory;
+			maybeOwnProjectPath = null;
+		}
+
+		if (maybeOwnProjectPath is null)
+		{
+			Assert.Inconclusive("Could not find own project file to test");
+		}
+
+		return maybeOwnProjectPath;
+	}
+}
diff --git a/Engine/Source/Programs/Shared/EpicGames.Build.Tests/EpicGames.Build.Tests.csproj b/Engine/Source/Programs/Shared/EpicGames.Build.Tests/EpicGames.Build.Tests.csproj
new file mode 100644
--- /dev/null
+++ b/Engine/Source/Programs/Shared/EpicGames.Build.Tests/EpicGames.Build.Tests.csproj
@@ -0,0 +1,27 @@
+﻿<Project Sdk="Microsoft.NET.Sdk">
+  <Import Project="..\UnrealEngine.csproj.props" />
+  
+  <PropertyGroup>
+    <TargetFramework>net8.0</TargetFramework>
+    <OutputType>Library</OutputType>
+    <Nullable>enable</Nullable>
+    <IsPackable>false</IsPackable>
+    <NoWarn>$(NoWarn);1591</NoWarn>
+  </PropertyGroup>
+
+  <ItemGroup>
+    <PackageReference Include="Microsoft.NET.Test.Sdk" Version="17.14.1" />
+    <PackageReference Include="MSTest.TestAdapter" Version="3.6.4">
+      <TreatAsUsed>true</TreatAsUsed>
+    </PackageReference>
+    <PackageReference Include="MSTest.TestFramework" Version="3.6.4" />
+    <PackageReference Include="coverlet.collector" Version="6.0.4">
+      <PrivateAssets>all</PrivateAssets>
+      <IncludeAssets>runtime; build; native; contentfiles; analyzers; buildtransitive</IncludeAssets>
+    </PackageReference>
+  </ItemGroup>
+
+  <ItemGroup>
+    <ProjectReference Include="..\EpicGames.Build\EpicGames.Build.csproj" />
+  </ItemGroup>
+</Project>
