Fix project generation on Macs by explicitly adding configurations to the EpicGames.Build.Tests project.

#rnx
#Synced-CL 48985497

diff --git a/Engine/Source/Programs/Shared/EpicGames.Build.Tests/EpicGames.Build.Tests.csproj b/Engine/Source/Programs/Shared/EpicGames.Build.Tests/EpicGames.Build.Tests.csproj
--- a/Engine/Source/Programs/Shared/EpicGames.Build.Tests/EpicGames.Build.Tests.csproj
+++ b/Engine/Source/Programs/Shared/EpicGames.Build.Tests/EpicGames.Build.Tests.csproj
@@ -3,6 +3,10 @@
   
   <PropertyGroup>
     <TargetFramework>net8.0</TargetFramework>
+    <!-- Configurations are required to be in the project itself despite being declared in UnrealEngine.csproj.props due to
+      a reported Mac-specific issue with project generation. -->
+    <Configuration Condition=" '$(Configuration)' == '' ">Development</Configuration>
+    <Configurations>Debug;Release;Development;Analyze</Configurations>
     <OutputType>Library</OutputType>
     <Nullable>enable</Nullable>
     <IsPackable>false</IsPackable>
