EpicGames.AspNet: Skip generating documentation file in Analyze configuration

Required for silencing CS1591 warning as <NoWarn> is not respected.

diff --git a/Engine/Source/Programs/Shared/EpicGames.AspNet/EpicGames.AspNet.csproj b/Engine/Source/Programs/Shared/EpicGames.AspNet/EpicGames.AspNet.csproj
--- a/Engine/Source/Programs/Shared/EpicGames.AspNet/EpicGames.AspNet.csproj
+++ b/Engine/Source/Programs/Shared/EpicGames.AspNet/EpicGames.AspNet.csproj
@@ -30,8 +30,7 @@
     <RunAnalyzersDuringBuild>True</RunAnalyzersDuringBuild>
     <TreatWarningsAsErrors>False</TreatWarningsAsErrors>
     <EnforceCodeStyleInBuild>True</EnforceCodeStyleInBuild>
-    <GenerateDocumentationFile>True</GenerateDocumentationFile>
-    <DocumentationFile></DocumentationFile>
+    <GenerateDocumentationFile>False</GenerateDocumentationFile>
     <NoWarn>$(NoWarn);1701;1702;1591</NoWarn>
   </PropertyGroup>
 
