Add an EpicGames.Benchmarks project with some benchmarks for various hash functions currently available.

#Synced-CL 48985497

diff --git a/Engine/Source/Programs/Shared/EpicGames.Benchmarks/EpicGames.Benchmarks.csproj b/Engine/Source/Programs/Shared/EpicGames.Benchmarks/EpicGames.Benchmarks.csproj
new file mode 100644
--- /dev/null
+++ b/Engine/Source/Programs/Shared/EpicGames.Benchmarks/EpicGames.Benchmarks.csproj
@@ -0,0 +1,22 @@
+﻿<Project Sdk="Microsoft.NET.Sdk">
+  <Import Project="..\UnrealEngine.csproj.props" />
+
+  <PropertyGroup>
+    <TargetFramework>net8.0</TargetFramework>
+    <OutputType>Exe</OutputType>
+    <Nullable>enable</Nullable>
+    <Configuration Condition=" '$(Configuration)' == '' ">Development</Configuration>
+    <Configurations>Debug;Release;Development;Analyze</Configurations>
+    <NoWarn>$(NoWarn);1591</NoWarn>
+  </PropertyGroup>
+
+  <ItemGroup>
+    <PackageReference Include="BenchmarkDotNet" Version="0.15.2" />
+    <PackageReference Include="Microsoft.VisualStudio.DiagnosticsHub.BenchmarkDotNetDiagnosers" Version="18.0.36313.1" />
+  </ItemGroup>
+
+  <ItemGroup>
+    <ProjectReference Include="..\EpicGames.Core\EpicGames.Core.csproj" />
+    <ProjectReference Include="..\EpicGames.IoHash\EpicGames.IoHash.csproj" />
+  </ItemGroup>
+</Project>
diff --git a/Engine/Source/Programs/Shared/EpicGames.Benchmarks/HashBenchmarks.cs b/Engine/Source/Programs/Shared/EpicGames.Benchmarks/HashBenchmarks.cs
new file mode 100644
--- /dev/null
+++ b/Engine/Source/Programs/Shared/EpicGames.Benchmarks/HashBenchmarks.cs
@@ -0,0 +1,50 @@
+// Copyright Epic Games, Inc. All Rights Reserved.
+
+using System;
+using BenchmarkDotNet.Attributes;
+using EpicGames.Core;
+using Microsoft.VSDiagnostics;
+
+namespace EpicGames.Benchmarks;
+
+// For more information on the VS BenchmarkDotNet Diagnosers see https://learn.microsoft.com/visualstudio/profiling/profiling-with-benchmark-dotnet
+[CPUUsageDiagnoser]
+[MemoryDiagnoser]
+public class HashBenchmarks
+{
+	[Params(1000, 10_000, 1_000_000)]
+	public int DataLength { get; set; }
+
+	private byte[]? _data;
+
+	[GlobalSetup]
+	public void Setup()
+	{
+		_data = new byte[DataLength];
+		new Random(0).NextBytes(_data);
+	}
+
+	[Benchmark]
+	public IoHash IOHash()
+	{
+		return IoHash.Compute(_data);
+	}
+
+	[Benchmark]
+	public ContentHash ContentHashSHA1()
+	{
+		return ContentHash.SHA1(_data!);
+	}
+
+	[Benchmark]
+	public ContentHash ContentHashMD5()
+	{
+		return ContentHash.MD5(_data!);
+	}
+
+	[Benchmark]
+	public Md5Hash MD5Hash()
+	{
+		return Md5Hash.Compute(_data!);
+	}
+}
diff --git a/Engine/Source/Programs/Shared/EpicGames.Benchmarks/Program.cs b/Engine/Source/Programs/Shared/EpicGames.Benchmarks/Program.cs
new file mode 100644
--- /dev/null
+++ b/Engine/Source/Programs/Shared/EpicGames.Benchmarks/Program.cs
@@ -0,0 +1,13 @@
+// Copyright Epic Games, Inc. All Rights Reserved.
+
+using BenchmarkDotNet.Running;
+
+namespace EpicGames.Benchmarks;
+
+internal class Program
+{
+	static void Main(string[] args)
+	{
+		_ = BenchmarkRunner.Run(typeof(Program).Assembly, args: args);
+	}
+}
