EpicGames.AspNet: Analysis requires generating documentation

#rnx
#Synced-CL 49007367

diff --git a/Engine/Source/Programs/Shared/EpicGames.AspNet/EpicGames.AspNet.csproj b/Engine/Source/Programs/Shared/EpicGames.AspNet/EpicGames.AspNet.csproj
--- a/Engine/Source/Programs/Shared/EpicGames.AspNet/EpicGames.AspNet.csproj
+++ b/Engine/Source/Programs/Shared/EpicGames.AspNet/EpicGames.AspNet.csproj
@@ -30,7 +30,6 @@
     <RunAnalyzersDuringBuild>True</RunAnalyzersDuringBuild>
     <TreatWarningsAsErrors>False</TreatWarningsAsErrors>
     <EnforceCodeStyleInBuild>True</EnforceCodeStyleInBuild>
-    <GenerateDocumentationFile>False</GenerateDocumentationFile>
     <NoWarn>$(NoWarn);1701;1702;1591</NoWarn>
   </PropertyGroup>
 
