EpicGames.Build: Begin addressing analysis issues

#rnx
#Synced-CL 49007367

diff --git a/Engine/Source/Programs/Shared/EpicGames.Build/Automation/CommandLine.cs b/Engine/Source/Programs/Shared/EpicGames.Build/Automation/CommandLine.cs
--- a/Engine/Source/Programs/Shared/EpicGames.Build/Automation/CommandLine.cs
+++ b/Engine/Source/Programs/Shared/EpicGames.Build/Automation/CommandLine.cs
@@ -2,7 +2,6 @@
 
 using System;
 using System.Collections.Generic;
-using System.Text;
 
 namespace UnrealBuildBase
 {
@@ -11,9 +10,9 @@ namespace UnrealBuildBase
 	/// </summary>
 	public class CommandInfo(string CommandNameIn)
 	{
-		public string CommandName = CommandNameIn;
+		public string CommandName { get; } = CommandNameIn;
 
-		public List<string> Arguments = [];
+		public List<string> Arguments { get; set; } = [];
 
 		public override string ToString()
 		{
@@ -49,7 +48,7 @@ namespace UnrealBuildBase
 			private set;
 		} = new HashSet<string>(IgnoredGlobalParameters, StringComparer.InvariantCultureIgnoreCase);
 
-		HashSet<string> GlobalParameterValues = new HashSet<string>(StringComparer.InvariantCultureIgnoreCase);
+		readonly HashSet<string> GlobalParameterValues = new HashSet<string>(StringComparer.InvariantCultureIgnoreCase);
 		
 		public void SetGlobal(string Parameter)
 		{
diff --git a/Engine/Source/Programs/Shared/EpicGames.Build/Automation/ExitCode.cs b/Engine/Source/Programs/Shared/EpicGames.Build/Automation/ExitCode.cs
--- a/Engine/Source/Programs/Shared/EpicGames.Build/Automation/ExitCode.cs
+++ b/Engine/Source/Programs/Shared/EpicGames.Build/Automation/ExitCode.cs
@@ -1,12 +1,9 @@
 // Copyright Epic Games, Inc. All Rights Reserved.
 
-using System;
-using System.Collections.Generic;
-using System.Text;
-
 namespace UnrealBuildBase
 {
 	// NOTE: this needs to be kept in sync with EditorAnalytics.h and iPhonePackager.cs
+	[System.Diagnostics.CodeAnalysis.SuppressMessage("Naming", "CA1707:Identifiers should not contain underscores", Justification = "Intentional")]
 	public enum ExitCode
     {
         Error_UATNotFound = -1,
diff --git a/Engine/Source/Programs/Shared/EpicGames.Build/CsProjectInfo.cs b/Engine/Source/Programs/Shared/EpicGames.Build/CsProjectInfo.cs
--- a/Engine/Source/Programs/Shared/EpicGames.Build/CsProjectInfo.cs
+++ b/Engine/Source/Programs/Shared/EpicGames.Build/CsProjectInfo.cs
@@ -428,7 +428,7 @@ namespace EpicGames.Core
 						// and we also want to filter out Intermediate, Binaries, etc
 
 						ConcurrentBag<DirectoryReference> bag = new();
-						ParallelQueue queue = new();
+						using ParallelQueue queue = new();
 						GetDirectories(DirectoryItem.GetItemByDirectoryReference(existingPath), bag, option, queue);
 						queue.Drain();
 
diff --git a/Engine/Source/Programs/Shared/EpicGames.Build/DotnetProjectInfo.cs b/Engine/Source/Programs/Shared/EpicGames.Build/DotnetProjectInfo.cs
--- a/Engine/Source/Programs/Shared/EpicGames.Build/DotnetProjectInfo.cs
+++ b/Engine/Source/Programs/Shared/EpicGames.Build/DotnetProjectInfo.cs
@@ -11,7 +11,6 @@ using EpicGames.Core;
 using EpicGames.MsBuild;
 using Microsoft.Build.Definition;
 using Microsoft.Build.Evaluation;
-using OpenTracing;
 
 namespace UnrealBuildBase;
 
diff --git a/Engine/Source/Programs/Shared/EpicGames.Build/Exceptions/BuildException.cs b/Engine/Source/Programs/Shared/EpicGames.Build/Exceptions/BuildException.cs
--- a/Engine/Source/Programs/Shared/EpicGames.Build/Exceptions/BuildException.cs
+++ b/Engine/Source/Programs/Shared/EpicGames.Build/Exceptions/BuildException.cs
@@ -4,7 +4,6 @@ using System;
 using System.Diagnostics.CodeAnalysis;
 using EpicGames.Core;
 using Microsoft.Extensions.Logging;
-using UnrealBuildBase;
 
 // This namespace is intentional for backwards compatibility.
 // BuildException will eventually be marked Obsolete in favor of BuildLogEventException
diff --git a/Engine/Source/Programs/Shared/EpicGames.Build/Exceptions/BuildExceptionExtensions.cs b/Engine/Source/Programs/Shared/EpicGames.Build/Exceptions/BuildExceptionExtensions.cs
--- a/Engine/Source/Programs/Shared/EpicGames.Build/Exceptions/BuildExceptionExtensions.cs
+++ b/Engine/Source/Programs/Shared/EpicGames.Build/Exceptions/BuildExceptionExtensions.cs
@@ -1,7 +1,6 @@
 ﻿// Copyright Epic Games, Inc. All Rights Reserved.
 
 using System;
-using System.Linq;
 using System.Text.Json;
 using EpicGames.Core;
 using Microsoft.Extensions.Logging;
diff --git a/Engine/Source/Programs/Shared/EpicGames.Build/System/CompileScriptModules.cs b/Engine/Source/Programs/Shared/EpicGames.Build/System/CompileScriptModules.cs
--- a/Engine/Source/Programs/Shared/EpicGames.Build/System/CompileScriptModules.cs
+++ b/Engine/Source/Programs/Shared/EpicGames.Build/System/CompileScriptModules.cs
@@ -7,7 +7,6 @@ using System.Diagnostics;
 using System.IO;
 using System.Linq;
 using System.Runtime.CompilerServices;
-using System.Text;
 using System.Text.Json;
 using EpicGames.Core;
 using EpicGames.MsBuild;
@@ -174,7 +173,7 @@ namespace UnrealBuildBase
 			{
 				FoundProjects = FoundProjects.Where(x => (AutomationProjectIncludeFilter.Contains(x.GetFileNameWithoutAnyExtensions(), StringComparer.OrdinalIgnoreCase))).ToHashSet();
 			}
-			Logger.LogDebug("Find {Count} projects time: {TimeSeconds:0.00} s", FoundProjects.Count(), Timer.Elapsed.TotalSeconds);
+			Logger.LogDebug("Find {Count} projects time: {TimeSeconds:0.00} s", FoundProjects.Count, Timer.Elapsed.TotalSeconds);
 
 			BuildFlags BuildFlags = BuildFlags.None;
 			if (bForceCompile)
@@ -585,7 +584,6 @@ namespace UnrealBuildBase
 				return;
 			}
 
-
 			// If the engine is installed, and the project being validated is under the engine directory, treat it as valid
 			// so it is not built as a dependency when building an external project.
 			// It is assumed these projects exist and have already been verified.
diff --git a/Engine/Source/Programs/Shared/EpicGames.Build/System/DirectoryEnumerator.cs b/Engine/Source/Programs/Shared/EpicGames.Build/System/DirectoryEnumerator.cs
--- a/Engine/Source/Programs/Shared/EpicGames.Build/System/DirectoryEnumerator.cs
+++ b/Engine/Source/Programs/Shared/EpicGames.Build/System/DirectoryEnumerator.cs
@@ -5,8 +5,6 @@ using System.Buffers;
 using System.Collections.Generic;
 using System.ComponentModel;
 using System.IO;
-using System.IO.Enumeration;
-using System.Runtime.CompilerServices;
 using System.Runtime.InteropServices;
 
 namespace UnrealBuildBase
@@ -343,7 +341,7 @@ namespace UnrealBuildBase
 		}
 
 		[StructLayout(LayoutKind.Sequential, CharSet = CharSet.Unicode)]
-		public struct FILE_FULL_DIR_INFORMATION
+		private struct FILE_FULL_DIR_INFORMATION
 		{
 			public uint NextEntryOffset;
 			public uint FileIndex;
diff --git a/Engine/Source/Programs/Shared/EpicGames.Build/System/DirectoryItem.cs b/Engine/Source/Programs/Shared/EpicGames.Build/System/DirectoryItem.cs
--- a/Engine/Source/Programs/Shared/EpicGames.Build/System/DirectoryItem.cs
+++ b/Engine/Source/Programs/Shared/EpicGames.Build/System/DirectoryItem.cs
@@ -19,12 +19,12 @@ namespace UnrealBuildBase
 		/// <summary>
 		/// Full path to the directory on disk
 		/// </summary>
-		public readonly DirectoryReference Location;
+		public required DirectoryReference Location { get; init; }
 
 		/// <summary>
 		/// The name of this directory
 		/// </summary>
-		public readonly string Name;
+		public required string Name { get; init; }
 
 		/// <summary>
 		/// Accessor for map of name to subdirectory item
@@ -202,6 +202,7 @@ namespace UnrealBuildBase
 			return Directories.TryGetValue(Name, out OutDirectory);
 		}
 
+		[SetsRequiredMembers]
 		DirectoryItem(DirectoryReference location, string name, Info? info)
 		{
 			Location = location;
diff --git a/Engine/Source/Programs/Shared/EpicGames.Build/System/DirectoryLookupCache.cs b/Engine/Source/Programs/Shared/EpicGames.Build/System/DirectoryLookupCache.cs
--- a/Engine/Source/Programs/Shared/EpicGames.Build/System/DirectoryLookupCache.cs
+++ b/Engine/Source/Programs/Shared/EpicGames.Build/System/DirectoryLookupCache.cs
@@ -1,6 +1,5 @@
 // Copyright Epic Games, Inc. All Rights Reserved.
 
-using System;
 using System.IO;
 using System.Collections.Generic;
 using EpicGames.Core;
@@ -8,35 +7,35 @@ using System.Linq;
 
 namespace UnrealBuildBase
 {
-	static public class DirectoryLookupCache
+	public static class DirectoryLookupCache
 	{
-		static public bool FileExists(FileReference File)
+		public static bool FileExists(FileReference File)
 		{
 			return FileItem.GetItemByFileReference(File).Exists;
 		}
 
-		static public bool DirectoryExists(DirectoryReference Directory)
+		public static bool DirectoryExists(DirectoryReference Directory)
 		{
 			return DirectoryItem.GetItemByDirectoryReference(Directory).Exists;
 		}
 
-		static public bool DirectoryExistsAndContainsFiles(DirectoryReference Directory, SearchOption searchOption = SearchOption.TopDirectoryOnly)
+		public static bool DirectoryExistsAndContainsFiles(DirectoryReference Directory, SearchOption searchOption = SearchOption.TopDirectoryOnly)
 		{
 			DirectoryItem item = DirectoryItem.GetItemByDirectoryReference(Directory);
 			return item.Exists && item.ContainsFiles(searchOption);
 		}
 
-		static public IEnumerable<FileReference> EnumerateFiles(DirectoryReference Directory)
+		public static IEnumerable<FileReference> EnumerateFiles(DirectoryReference Directory)
 		{
 			return DirectoryItem.GetItemByDirectoryReference(Directory).EnumerateFiles().Select(x => x.Location);
 		}
 
-		static public IEnumerable<DirectoryReference> EnumerateDirectories(DirectoryReference Directory)
+		public static IEnumerable<DirectoryReference> EnumerateDirectories(DirectoryReference Directory)
 		{
 			return DirectoryItem.GetItemByDirectoryReference(Directory).EnumerateDirectories().Select(x => x.Location);
 		}
 
-        static public void InvalidateCachedDirectory(DirectoryReference Directory)
+        public static void InvalidateCachedDirectory(DirectoryReference Directory)
         {
 			DirectoryItem.GetItemByDirectoryReference(Directory).ResetCachedInfo();
         }
diff --git a/Engine/Source/Programs/Shared/EpicGames.Build/System/EnumeratePlugins.cs b/Engine/Source/Programs/Shared/EpicGames.Build/System/EnumeratePlugins.cs
--- a/Engine/Source/Programs/Shared/EpicGames.Build/System/EnumeratePlugins.cs
+++ b/Engine/Source/Programs/Shared/EpicGames.Build/System/EnumeratePlugins.cs
@@ -6,7 +6,6 @@ using System;
 using System.Collections.Concurrent;
 using System.Collections.Generic;
 using System.Linq;
-using System.Text;
 
 namespace UnrealBuildBase
 {
diff --git a/Engine/Source/Programs/Shared/EpicGames.Build/System/FileHasher.cs b/Engine/Source/Programs/Shared/EpicGames.Build/System/FileHasher.cs
--- a/Engine/Source/Programs/Shared/EpicGames.Build/System/FileHasher.cs
+++ b/Engine/Source/Programs/Shared/EpicGames.Build/System/FileHasher.cs
@@ -2,7 +2,6 @@
 
 using System.Collections.Concurrent;
 using System.Collections.Generic;
-using System.Diagnostics;
 using System.IO;
 using System.Linq;
 using System.Text.Json;
diff --git a/Engine/Source/Programs/Shared/EpicGames.Build/System/FileItem.cs b/Engine/Source/Programs/Shared/EpicGames.Build/System/FileItem.cs
--- a/Engine/Source/Programs/Shared/EpicGames.Build/System/FileItem.cs
+++ b/Engine/Source/Programs/Shared/EpicGames.Build/System/FileItem.cs
@@ -4,6 +4,7 @@ using System;
 using System.Collections.Concurrent;
 using System.Collections.Generic;
 using System.Diagnostics;
+using System.Diagnostics.CodeAnalysis;
 using System.IO;
 using System.Linq;
 using System.Runtime.CompilerServices;
@@ -22,17 +23,17 @@ namespace UnrealBuildBase
 		/// <summary>
 		/// The directory containing this file
 		/// </summary>
-		Lazy<DirectoryItem> CachedDirectory;
+		Lazy<DirectoryItem> _cachedDirectory;
 
 		/// <summary>
 		/// Location of this file
 		/// </summary>
-		public readonly FileReference Location;
+		public required FileReference Location { get; init; }
 
 		/// <summary>
 		/// Name of this file
 		/// </summary>
-		public readonly string Name;
+		public required string Name { get; init; }
 
 		/// <summary>
 		/// Full name of this file
@@ -47,7 +48,7 @@ namespace UnrealBuildBase
 		/// <summary>
 		/// Gets the directory that this file is in
 		/// </summary>
-		public DirectoryItem Directory => CachedDirectory.Value;
+		public DirectoryItem Directory => _cachedDirectory.Value;
 
 		/// <summary>
 		/// Whether the file exists.
@@ -94,9 +95,9 @@ namespace UnrealBuildBase
 		public void UpdateCachedDirectory(DirectoryItem Directory)
 		{
 			Debug.Assert(Directory.Location == Location.Directory);
-			if (!CachedDirectory.IsValueCreated)
+			if (!_cachedDirectory.IsValueCreated)
 			{
-				CachedDirectory = new(Directory);
+				_cachedDirectory = new(Directory);
 			}
 		}
 
@@ -217,20 +218,16 @@ namespace UnrealBuildBase
 		/// </summary>
 		public bool HasInfo => _info != null;
 
+		[SetsRequiredMembers]
 		internal FileItem(FileReference Location, string name, Info? info, DirectoryItem? directory)
 		{
 			this.Location = Location;
 			Name = name;
 			_info = info;
 
-			if (directory != null)
-			{
-				CachedDirectory = new(directory);
-			}
-			else
-			{
-				CachedDirectory = new(() => DirectoryItem.GetItemByDirectoryReference(Location.Directory)); // Take this in
-			}
+			_cachedDirectory = directory != null
+				? new(directory)
+				: new(() => DirectoryItem.GetItemByDirectoryReference(Location.Directory));
 		}
 
 		private Info Update()
diff --git a/Engine/Source/Programs/Shared/EpicGames.Build/System/GlobalSingleInstanceMutex.cs b/Engine/Source/Programs/Shared/EpicGames.Build/System/GlobalSingleInstanceMutex.cs
--- a/Engine/Source/Programs/Shared/EpicGames.Build/System/GlobalSingleInstanceMutex.cs
+++ b/Engine/Source/Programs/Shared/EpicGames.Build/System/GlobalSingleInstanceMutex.cs
@@ -14,7 +14,7 @@ namespace UnrealBuildBase
 		/// <summary>
 		/// The global mutex instance
 		/// </summary>
-		Mutex? GlobalMutex;
+		Mutex? _globalMutex;
 
 		/// <summary>
 		/// Constructor. Attempts to acquire the global mutex
@@ -25,7 +25,7 @@ namespace UnrealBuildBase
 		{
 			// Try to create the mutex, with it initially locked
 			bool bCreatedMutex;
-			GlobalMutex = new Mutex(true, mutexName, out bCreatedMutex);
+			_globalMutex = new Mutex(true, mutexName, out bCreatedMutex);
 
 			// If we didn't create the mutex, we can wait for it or fail immediately
 			if (!bCreatedMutex)
@@ -34,7 +34,7 @@ namespace UnrealBuildBase
 				{
 					try
 					{
-						GlobalMutex.WaitOne();
+						_globalMutex.WaitOne();
 					}
 					catch (AbandonedMutexException)
 					{
@@ -82,9 +82,9 @@ namespace UnrealBuildBase
 		{
 			if (disposing)
 			{
-				GlobalMutex?.ReleaseMutex();
-				GlobalMutex?.Dispose();
-				GlobalMutex = null;
+				_globalMutex?.ReleaseMutex();
+				_globalMutex?.Dispose();
+				_globalMutex = null;
 			}
 		}
 	}
diff --git a/Engine/Source/Programs/Shared/EpicGames.Build/System/NativeProjectsBase.cs b/Engine/Source/Programs/Shared/EpicGames.Build/System/NativeProjectsBase.cs
--- a/Engine/Source/Programs/Shared/EpicGames.Build/System/NativeProjectsBase.cs
+++ b/Engine/Source/Programs/Shared/EpicGames.Build/System/NativeProjectsBase.cs
@@ -2,11 +2,8 @@
 
 using System;
 using System.Collections.Generic;
-using System.Diagnostics.CodeAnalysis;
 using System.IO;
 using System.Linq;
-using System.Text;
-using System.Threading.Tasks;
 using EpicGames.Core;
 using Microsoft.Extensions.Logging;
 
diff --git a/Engine/Source/Programs/Shared/EpicGames.Build/System/Rules.cs b/Engine/Source/Programs/Shared/EpicGames.Build/System/Rules.cs
--- a/Engine/Source/Programs/Shared/EpicGames.Build/System/Rules.cs
+++ b/Engine/Source/Programs/Shared/EpicGames.Build/System/Rules.cs
@@ -5,9 +5,7 @@ using System;
 using System.Collections.Generic;
 using System.IO;
 using System.Linq;
-using System.Text;
 using System.Threading;
-using System.Threading.Tasks;
 
 namespace UnrealBuildBase
 {
@@ -151,7 +149,6 @@ namespace UnrealBuildBase
 			public List<FileReference> UbtPlugins = [];
 		}
 
-
 		class Directory
 		{
 			internal Files? Files;
@@ -161,7 +158,6 @@ namespace UnrealBuildBase
 		static Directory Root = new();
 		static ReaderWriterLockSlim RootLock = new();
 
-
 		static Directory? FindDirectory(DirectoryReference dirRef, bool createIfNotFound)
 		{
 			string[] parts = dirRef.FullName.Split(Path.DirectorySeparatorChar);
@@ -192,9 +188,9 @@ namespace UnrealBuildBase
 		public static void Populate(IEnumerable<DirectoryReference> dirs)
 		{
 			RootLock.EnterWriteLock();
+			ParallelQueue? queue = null;
 			try
 			{
-				ParallelQueue? queue = null;
 				foreach (DirectoryReference dirRef in dirs)
 				{
 					Directory dir = FindDirectory(dirRef, true)!;
@@ -221,6 +217,7 @@ namespace UnrealBuildBase
 			}
 			finally
 			{
+				queue?.Dispose();
 				RootLock.ExitWriteLock();
 			}
 		}
@@ -245,7 +242,7 @@ namespace UnrealBuildBase
 						GetRecursively(files, dir, type);
 					}
 				}
-				files.Sort(static (A, B) => A.FullName.CompareTo(B.FullName));
+				files.Sort(static (A, B) => A.CompareTo(B));
 				return files;
 			}
 			catch (Exception)
@@ -413,7 +410,7 @@ namespace UnrealBuildBase
 						continue;
 					}
 				}
-				else if (subDirectory.Name.StartsWith("__External")) // This is a bit annoying but version control likes to leave around these huuuge hierarchies of empty directories
+				else if (subDirectory.Name.StartsWith("__External", StringComparison.Ordinal)) // This is a bit annoying but version control likes to leave around these huuuge hierarchies of empty directories
 				{
 					//Console.WriteLine("Most likely leftovers: " + Directory.Location);
 					continue;
diff --git a/Engine/Source/Programs/Shared/EpicGames.Build/UEBuildPlatformSDK.cs b/Engine/Source/Programs/Shared/EpicGames.Build/UEBuildPlatformSDK.cs
--- a/Engine/Source/Programs/Shared/EpicGames.Build/UEBuildPlatformSDK.cs
+++ b/Engine/Source/Programs/Shared/EpicGames.Build/UEBuildPlatformSDK.cs
@@ -72,7 +72,6 @@ namespace EpicGames.Core
 			}
 		}
 
-
 		public SDKDescriptor(string InName, string? InMin, string? InMax)
 			: this(InName, InMin, InMax, null, null, null)
 		{
@@ -187,7 +186,6 @@ namespace EpicGames.Core
 			return AutoSdk?.Validity == SDKStatus.Valid;
 		}
 
-
 		public string ToString(bool bIncludeCurrent = true)
 		{
 			// by default print something like "MinVersion=1.0, MaxVersion=2.0"
@@ -262,7 +260,6 @@ namespace EpicGames.Core
 		{
 		}
 
-
 		public override string DefaultName => "Software";
 
 		public override bool IsValid(string Version, SDKDescriptor Info)
@@ -271,11 +268,10 @@ namespace EpicGames.Core
 		}
 	}
 
-
 	/// <summary>
 	/// SDK for a platform
 	/// </summary>
-	abstract public class UEBuildPlatformSDK(ILogger InLogger)
+	public abstract class UEBuildPlatformSDK(ILogger InLogger)
 	{
 		// Public SDK handling, not specific to AutoSDK
 
@@ -294,7 +290,7 @@ namespace EpicGames.Core
 		public static void RegisterSDKForPlatform(UEBuildPlatformSDK SDK, string PlatformName, bool bIsSdkAllowedOnHost)
 		{
 			// verify that neither platform or sdk were added before
-			if (SDKRegistry.Count(x => x.Key == PlatformName || x.Value == SDK) > 0)
+			if (SDKRegistry.Any(x => x.Key == PlatformName || x.Value == SDK))
 			{
 				throw new Exception(String.Format("Re-registering SDK for {0}. All Platforms must have a unique SDK object", PlatformName));
 			}
@@ -394,7 +390,6 @@ namespace EpicGames.Core
 			return (T?)SDK;
 		}
 
-
 		/// <summary>
 		/// Gets the set of all known SDKs
 		/// </summary>
@@ -492,7 +487,6 @@ namespace EpicGames.Core
 		//	SDKCollection AllSoftwares =
 		//}
 
-
 		public string? GetInstalledVersion(out bool bIsAutoSDK)
 		{
 			bIsAutoSDK = HasSetupAutoSDK();
@@ -529,7 +523,6 @@ namespace EpicGames.Core
 		//			ManualSDKVersion = CachedManualSDKVersion;
 		//		}
 
-
 		public virtual bool IsVersionValid(string? Version, string SDKType)
 		{
 			return IsVersionValidInternal(Version, SDKType, null);
@@ -613,7 +606,7 @@ namespace EpicGames.Core
 		/// <returns></returns>
 		protected virtual string? GetInstalledSDKVersion()
 		{
-			if (GetInstalledSDKVersions().FullSdks.Count() == 1)
+			if (GetInstalledSDKVersions().FullSdks.Count == 1)
 			{
 				return GetInstalledSDKVersions().FullSdks[0].Current;
 			}
@@ -739,7 +732,6 @@ namespace EpicGames.Core
 			return false;
 		}
 
-
 		/// <summary>
 		/// For a platform that doesn't use properly named AutoSDK directories, the directory name may not be convertible to an integer,
 		/// and IsVersionValid checks could fail when checking AutoSDK version for an exact match. GetMainVersion() would return the 
@@ -752,7 +744,6 @@ namespace EpicGames.Core
 			return GetVersionFromConfig("AutoSDKDirectory") ?? GetMainVersion();
 		}
 
-
 		///// <summary>
 		///// Gets the valid (integer) range of Sdk versions. Must be an integer to easily check a range vs a particular version
 		///// </summary>
@@ -786,7 +777,6 @@ namespace EpicGames.Core
 		//	}
 		//}
 
-
 		// Let platform override behavior to determine if a version is a valid (useful for non-numeric versions)
 		protected virtual bool IsVersionValidInternal(string? Version, string? SDKType, SDKDescriptor? VersionInfo)
 		{
@@ -857,7 +847,6 @@ namespace EpicGames.Core
 			return IntVersion >= VersionInfo.MinInt && IntVersion <= VersionInfo.MaxInt;
 		}
 
-
 		/// <summary>
 		/// Only the platform can convert a version string into an integer that is usable for comparison
 		/// </summary>
@@ -916,7 +905,6 @@ namespace EpicGames.Core
 			}
 		}
 
-
 		/// <summary>
 		/// Allow the platform SDK to override the name it will use in AutoSDK, but default to the platform name
 		/// </summary>
@@ -1021,10 +1009,6 @@ namespace EpicGames.Core
 
 		#endregion
 
-
-
-
-
 		#region Private/Protected general functionality
 
 		// this is the SDK version that was set before activating AutoSDK, since AutoSDK may remove ability to retrieve the Manual SDK version
@@ -1232,7 +1216,6 @@ namespace EpicGames.Core
 			return VersionNumber.Parse(VersionString);
 		}
 
-
 		private VersionNumberRange? ParseVersionNumberRange(string Range)
 		{
 			string[] Versions = Range.Split("-");
@@ -1333,7 +1316,6 @@ namespace EpicGames.Core
 		protected const string AutoSetupEnvVar = "AutoSDKSetup";
 		protected const string ManualSetupEnvVar = "ManualSDKSetup";
 
-
 		private static string GetAutoSDKHostPlatform()
 		{
 			if (OperatingSystem.IsWindows())
@@ -1360,9 +1342,9 @@ namespace EpicGames.Core
 			return false;
 		}
 
-		static private bool bCheckedAutoSDKRootEnvVar = false;
-		static private bool bAutoSDKSystemEnabled = false;
-		static private bool HasAutoSDKSystemEnabled()
+		private static bool bCheckedAutoSDKRootEnvVar = false;
+		private static bool bAutoSDKSystemEnabled = false;
+		private static bool HasAutoSDKSystemEnabled()
 		{
 			if (!bCheckedAutoSDKRootEnvVar)
 			{
@@ -1412,7 +1394,7 @@ namespace EpicGames.Core
 			string? SDKRoot = Environment.GetEnvironmentVariable(SDKRootEnvVar);
 			if (SDKRoot != null)
 			{
-				if (SDKRoot != "")
+				if (!String.IsNullOrEmpty(SDKRoot))
 				{
 					SDKPath = Path.Combine(SDKRoot, "Host" + GetAutoSDKHostPlatform(), GetAutoSDKPlatformName());
 				}
@@ -1682,7 +1664,6 @@ namespace EpicGames.Core
 			return true;
 		}
 
-
 		protected bool SetLastRunAutoSDKScriptVersion(string LastRunScriptVersion)
 		{
 			string PlatformSDKRoot = GetPathToPlatformAutoSDKs();
@@ -1753,7 +1734,7 @@ namespace EpicGames.Core
 			}
 		}
 
-		static private string[] AutoSDKLevels =
+		private static string[] AutoSDKLevels =
 		[
 			"NONE",
 			"BUILD",
@@ -1827,7 +1808,7 @@ namespace EpicGames.Core
 				Logger.LogDebug("{Platform} attempted to run SDK hook which could have damaged manual SDK install!", GetAutoSDKPlatformName());
 				return false;
 			}
-			if (SDKVersionString != "")
+			if (!String.IsNullOrEmpty(SDKVersionString))
 			{
 				string SDKDirectory = Path.Combine(PlatformSDKRoot, SDKVersionString);
 				string HookExe = Path.Combine(SDKDirectory, GetHookExecutableName(Hook));
@@ -1838,7 +1819,7 @@ namespace EpicGames.Core
 					using ITimelineEvent _ = Timeline.ScopeEvent("RunAutoSDKHooks");
 
 					// run it
-					Process HookProcess = new Process();
+					using Process HookProcess = new();
 					HookProcess.StartInfo.WorkingDirectory = SDKDirectory;
 					HookProcess.StartInfo.FileName = HookExe;
 					HookProcess.StartInfo.Arguments = AutoSDKLevel;
@@ -1963,7 +1944,6 @@ namespace EpicGames.Core
 						Environment.SetEnvironmentVariable(EnvVarName, EnvVarValue);
 					}
 
-
 					// actually perform the PATH stripping / adding.
 					string? OrigPathVar = Environment.GetEnvironmentVariable("PATH");
 					string PathDelimiter = OperatingSystem.IsWindows() ? ";" : ":";
@@ -2086,7 +2066,7 @@ namespace EpicGames.Core
 			if (PlatformSupportsAutoSDKs() && HasAutoSDKSystemEnabled())
 			{
 				string AutoSDKRoot = GetPathToPlatformAutoSDKs();
-				if (AutoSDKRoot != "")
+				if (!String.IsNullOrEmpty(AutoSDKRoot))
 				{
 					string DesiredSDKLevel = GetAutoSDKLevelForPlatform(AutoSDKRoot);
 					// if the user doesn't want AutoSDK for this platform, then return that it is not installed, even if it actually is
