[Automation] Add append option for retrieving cloud artifacts
#Synced-CL 48993796

diff --git a/Engine/Source/Programs/AutomationTool/BuildGraph/Tasks/RetrieveCloudArtifactTask.cs b/Engine/Source/Programs/AutomationTool/BuildGraph/Tasks/RetrieveCloudArtifactTask.cs
--- a/Engine/Source/Programs/AutomationTool/BuildGraph/Tasks/RetrieveCloudArtifactTask.cs
+++ b/Engine/Source/Programs/AutomationTool/BuildGraph/Tasks/RetrieveCloudArtifactTask.cs
@@ -146,6 +146,13 @@ namespace AutomationTool.Tasks
 		/// </summary>
 		[TaskParameter(Optional = true, ValidationType = TaskParameterValidationType.FileSpec)]
 		public string? Exclude { get; set; } = null;
+
+		/// <summary>
+		/// Whether to append downloads to an existing directory (true) or to clear the directory before download (false)
+		/// Set this to true if you want to download an artifact to a directory, then download an additional artifact from another without wiping state between
+		/// </summary>
+		[TaskParameter(Optional = true)]
+		public bool Append { get; set; } = false;
 	}
 
 	/// <summary>
@@ -399,7 +406,8 @@ namespace AutomationTool.Tasks
 					BoostWorkers = CommandUtils.IsBuildMachine || _parameters.BoostWorkers,
 					// boost worker memory if boosting workers as well
 					BoostWorkerMemory = CommandUtils.IsBuildMachine || _parameters.BoostWorkers,
-					TraceFile = traceFile
+					TraceFile = traceFile,
+					Append = _parameters.Append
 				}, optionsOverride: optionsOverride);
 		}
 
diff --git a/Engine/Source/Programs/Shared/EpicGames.ProjectStore/CloudStorage/CloudStorage.cs b/Engine/Source/Programs/Shared/EpicGames.ProjectStore/CloudStorage/CloudStorage.cs
--- a/Engine/Source/Programs/Shared/EpicGames.ProjectStore/CloudStorage/CloudStorage.cs
+++ b/Engine/Source/Programs/Shared/EpicGames.ProjectStore/CloudStorage/CloudStorage.cs
@@ -190,6 +190,7 @@ namespace EpicGames.ProjectStore
 			string boostWorkersOption = downloadOptions?.BoostWorkers ?? false ? "--boost-workers" : "";
 			string boostWorkerMemoryOption = downloadOptions?.BoostWorkerMemory ?? false ? "--boost-worker-memory" : "";
 			string allowMultipartOption = downloadOptions?.AllowMultipart != null ? $"--allow-multipart={downloadOptions.AllowMultipart}" : "";
+			string append = downloadOptions?.Append ?? false ? "--append=true" : String.Empty;
 			// trace file option needs to be passed before other arguments including the verbs
 			string traceOption = downloadOptions?.TraceFile != null ? $"--tracefile={downloadOptions.TraceFile.FullName} " : "";
 			string zenFolderPath = downloadOptions?.ZenStateFolder != null ? $"--zen-folder-path={downloadOptions.ZenStateFolder.FullName} " : "";
@@ -197,7 +198,7 @@ namespace EpicGames.ProjectStore
 			string wildcardOption = includeFiles == null ? "" : $"--wildcard \"{String.Join(';', includeFiles.Select(p => ToZenWildcard(outputDir, p)))}\"";
 			string excludeOption =  excludeFiles == null ? "" : $"--exclude-wildcard \"{String.Join(';', excludeFiles.Select(p => ToZenWildcard(outputDir, p)))}\"";
 			// pass the access token via environment variable to avoid showing it on the cli
-			string cmdline = $"{traceOption}builds download {_options.Value.GetHostCliArgs(optionsOverride)} --namespace {ns} --bucket {bucketId} {wildcardOption} {excludeOption} {allowMultipartOption} {http2Options} {boostWorkersOption} {boostWorkerMemoryOption} {zenFolderPath} \"{outputDir}\" {buildId}";
+			string cmdline = $"{traceOption}builds download {_options.Value.GetHostCliArgs(optionsOverride)} --namespace {ns} --bucket {bucketId} {wildcardOption} {excludeOption} {allowMultipartOption} {append} {http2Options} {boostWorkersOption} {boostWorkerMemoryOption} {zenFolderPath} \"{outputDir}\" {buildId}";
 
 			int exitCode = await RunZenCliAsync(cmdline, progress: progress, optionsOverride: optionsOverride, cancellationToken: cancellationToken);
 			EnsureZenExitCodeSuccess(exitCode);
diff --git a/Engine/Source/Programs/Shared/EpicGames.ProjectStore/CloudStorage/ICloudStorage.cs b/Engine/Source/Programs/Shared/EpicGames.ProjectStore/CloudStorage/ICloudStorage.cs
--- a/Engine/Source/Programs/Shared/EpicGames.ProjectStore/CloudStorage/ICloudStorage.cs
+++ b/Engine/Source/Programs/Shared/EpicGames.ProjectStore/CloudStorage/ICloudStorage.cs
@@ -213,6 +213,12 @@ namespace EpicGames.ProjectStore
 		/// Override directory location for the zen state folder
 		/// </summary>
 		public DirectoryReference? ZenStateFolder { get; set; }
+
+		/// <summary>
+		/// Whether to append downloads to an existing directory (true) or to clear the directory before download (false)
+		/// Set this to true if you want to download an artifact to a directory, then download an additional artifact from another without wiping state between
+		/// </summary>
+		public bool Append { get; set; }
 	}
 
 	/// <summary>
