Continue addressing analysis issues

#rnx
#Synced-CL 49039288

diff --git a/Engine/Source/Programs/Shared/EpicGames.Build/Automation/FormatCommandLine.cs b/Engine/Source/Programs/Shared/EpicGames.Build/Automation/FormatCommandLine.cs
--- a/Engine/Source/Programs/Shared/EpicGames.Build/Automation/FormatCommandLine.cs
+++ b/Engine/Source/Programs/Shared/EpicGames.Build/Automation/FormatCommandLine.cs
@@ -6,57 +6,52 @@ using System.Text;
 
 namespace UnrealBuildBase
 {
-	public partial class CommandLine
+	public static partial class CommandLine
 	{
 		/// <summary>
 		/// Converts a list of arguments to a string where each argument is separated with a space character.
 		/// </summary>
-		/// <param name="Arguments">Arguments</param>
+		/// <param name="arguments">Arguments</param>
 		/// <returns>Single string containing all arguments separated with a space.</returns>
-		public static string FormatCommandLine(IEnumerable<string> Arguments)
+		public static string FormatCommandLine(IEnumerable<string> arguments)
 		{
-			StringBuilder Result = new StringBuilder();
-			foreach(string Argument in Arguments)
+			StringBuilder result = new StringBuilder();
+			foreach(string argument in arguments)
 			{
-				if(Result.Length > 0)
+				if(result.Length > 0)
 				{
-					Result.Append(" ");
+					result.Append(' ');
 				}
-				Result.Append(FormatArgumentForCommandLine(Argument));
+				result.Append(FormatArgumentForCommandLine(argument));
 			}
-			return Result.ToString();
+			return result.ToString();
 		}
 
 		/// <summary>
 		/// Format a single argument for passing on the command line, inserting quotes as necessary.
 		/// </summary>
-		/// <param name="Argument">The argument to quote</param>
+		/// <param name="argument">The argument to quote</param>
 		/// <returns>The argument, with quotes if necessary</returns>
-		public static string FormatArgumentForCommandLine(string Argument)
+		public static string FormatArgumentForCommandLine(string argument)
 		{
 			// Check if the argument contains a space. If not, we can just pass it directly.
-			int SpaceIdx = Argument.IndexOf(' ');
-			if(SpaceIdx == -1)
+			int spaceIdx = argument.IndexOf(' ', StringComparison.Ordinal);
+			if(spaceIdx == -1)
 			{
-				return Argument;
+				return argument;
 			}
 
 			// If the argument has quotes in it, the nested quotes must be escaped (unless already escaped)
-			if (Argument.Contains('"') && !Argument.Contains("\\\""))
+			if (argument.Contains('"', StringComparison.Ordinal) && !argument.Contains("\\\"", StringComparison.Ordinal))
 			{
-				Argument = Argument.Replace("\"", "\\\"");
+				argument = argument.Replace("\"", "\\\"", StringComparison.Ordinal);
 			}
 
 			// If it does have a space, and it's formatted as an option (ie. -Something=), try to insert quotes after the equals character
-			int EqualsIdx = Argument.IndexOf('=');
-			if(Argument.StartsWith('-') && EqualsIdx != -1 && EqualsIdx < SpaceIdx)
-			{
-				return String.Format("{0}=\"{1}\"", Argument.Substring(0, EqualsIdx), Argument.Substring(EqualsIdx + 1));
-			}
-			else
-			{
-				return String.Format("\"{0}\"", Argument);
-			}
+			int equalsIdx = argument.IndexOf('=', StringComparison.Ordinal);
+			return argument.StartsWith('-') && equalsIdx != -1 && equalsIdx < spaceIdx
+				? $"{argument.Substring(0, equalsIdx)}=\"{argument.Substring(equalsIdx + 1)}\""
+				: $"\"{argument}\"";
 		}
 	}
 }
diff --git a/Engine/Source/Programs/Shared/EpicGames.Build/UEBuildPlatformSDK.cs b/Engine/Source/Programs/Shared/EpicGames.Build/UEBuildPlatformSDK.cs
--- a/Engine/Source/Programs/Shared/EpicGames.Build/UEBuildPlatformSDK.cs
+++ b/Engine/Source/Programs/Shared/EpicGames.Build/UEBuildPlatformSDK.cs
@@ -54,129 +54,129 @@ namespace EpicGames.Core
 			Name = "Unknown";
 		}
 
-		public SDKDescriptor(string InName, string? InMin, string? InMax, string? InCurrent, string? InGroupName, SDKCollection? Collection)
+		public SDKDescriptor(string inName, string? inMin, string? inMax, string? inCurrent, string? inGroupName, SDKCollection? collection)
 		{
-			Name = InName;
-			Min = InMin;
-			Max = InMax;
-			Current = InCurrent;
-			GroupName = InGroupName;
+			Name = inName;
+			Min = inMin;
+			Max = inMax;
+			Current = inCurrent;
+			GroupName = inGroupName;
 
-			if (Collection != null)
+			if (collection != null)
 			{
-				MinInt = Collection.PlatformSDK.ConvertVersionToInt(Min, Name);
+				MinInt = collection.PlatformSDK.ConvertVersionToInt(Min, Name);
 				// null will convert to 0, but we want MaxValue on null
-				MaxInt = Max == null ? UInt64.MaxValue : Collection.PlatformSDK.ConvertVersionToInt(Max, Name);
-				CurrentInt = Collection.PlatformSDK.ConvertVersionToInt(Current, Name);
-				UpdateValidity(Collection);
+				MaxInt = Max == null ? UInt64.MaxValue : collection.PlatformSDK.ConvertVersionToInt(Max, Name);
+				CurrentInt = collection.PlatformSDK.ConvertVersionToInt(Current, Name);
+				UpdateValidity(collection);
 			}
 		}
 
-		public SDKDescriptor(string InName, string? InMin, string? InMax)
-			: this(InName, InMin, InMax, null, null, null)
+		public SDKDescriptor(string inName, string? inMin, string? inMax)
+			: this(inName, inMin, inMax, null, null, null)
 		{
 		}
 
-		public SDKDescriptor(string InName, string? InCurrent)
-			: this(InName, null, null, InCurrent, null, null)
+		public SDKDescriptor(string inName, string? inCurrent)
+			: this(inName, null, null, inCurrent, null, null)
 		{
 		}
 
-		public void UpdateCurrent(string InCurrent, string? Hint, SDKCollection? Collection)
+		public void UpdateCurrent(string inCurrent, string? hint, SDKCollection? collection)
 		{
-			Current = InCurrent;
-			if (Collection != null)
+			Current = inCurrent;
+			if (collection != null)
 			{
-				CurrentInt = Collection.PlatformSDK.ConvertVersionToInt(Current, Hint);
-				UpdateValidity(Collection);
+				CurrentInt = collection.PlatformSDK.ConvertVersionToInt(Current, hint);
+				UpdateValidity(collection);
 			}
 		}
 
-		private void UpdateValidity(SDKCollection Collection)
+		private void UpdateValidity(SDKCollection collection)
 		{
 			// instead of just checking the range numerically, we allow platforms to override the IsValid check, which makes this a little convoluted (calling back into the Collection 
 			// so it can call back into the PlatformSDK with the proper SDK vs Software function)
-			Validity = (Current != null && Collection.IsValid(Current, this)) ? SDKStatus.Valid : SDKStatus.Invalid;
+			Validity = (Current != null && collection.IsValid(Current, this)) ? SDKStatus.Valid : SDKStatus.Invalid;
 		}
 
 		public string ToString(bool bIncludeCurrent = true)
 		{
 			return $"MinVersion={Min ?? ""}, MaxVersion={Max ?? ""}" + (bIncludeCurrent ? $", Current={Current ?? ""}" : "");
 		}
-		public string ToString(string Descriptor, string? CurrentDescriptor, bool bIncludeName)
+		public string ToString(string descriptor, string? currentDescriptor, bool bIncludeName)
 		{
-			string DisplayName = bIncludeName ? "_" + Name : "";
+			string displayName = bIncludeName ? "_" + Name : "";
 
 			// if they are different, print something like:
 			//   MinAllowed_Toolchain=1.0, MaxAllowed_Toolchain=2.0
 			if (Min != Max)
 			{
-				return $"Min{Descriptor}{DisplayName}={Min ?? ""}, Max{Descriptor}{DisplayName}={Max ?? ""}" + (CurrentDescriptor != null ? $", Current{CurrentDescriptor}{DisplayName}={Current ?? ""}" : "");
+				return $"Min{descriptor}{displayName}={Min ?? ""}, Max{descriptor}{displayName}={Max ?? ""}" + (currentDescriptor != null ? $", Current{currentDescriptor}{displayName}={Current ?? ""}" : "");
 			}
 			// if they are the same, print something like:
 			//   Allowed_Toolchain=1.0
-			return $"{Descriptor}{DisplayName}={Min}" + (CurrentDescriptor != null ? $", Current{CurrentDescriptor}{DisplayName}={Current ?? ""}" : "");
+			return $"{descriptor}{displayName}={Min}" + (currentDescriptor != null ? $", Current{currentDescriptor}{displayName}={Current ?? ""}" : "");
 		}
 	}
-	public class SDKCollection(UEBuildPlatformSDK InPlatformSDK)
+	public class SDKCollection(UEBuildPlatformSDK inPlatformSDK)
 	{
 		public virtual string DefaultName => "Sdk";
-		public virtual bool IsValid(string Version, SDKDescriptor Info)
+		public virtual bool IsValid(string version, SDKDescriptor info)
 		{
-			return PlatformSDK.IsVersionValid(Version, Info);
+			return PlatformSDK.IsVersionValid(version, info);
 		}
 
-		public UEBuildPlatformSDK PlatformSDK = InPlatformSDK;
+		public UEBuildPlatformSDK PlatformSDK = inPlatformSDK;
 		public List<SDKDescriptor> Sdks = [];
-		public List<SDKDescriptor> FullSdks => Sdks.Where(x => String.Compare(x.Name, "AutoSdk", true) != 0).ToList();
-		public SDKDescriptor? AutoSdk => Sdks.FirstOrDefault(x => String.Compare(x.Name, "AutoSdk", true) == 0);
+		public List<SDKDescriptor> FullSdks => Sdks.Where(x => !String.Equals(x.Name, "AutoSdk", StringComparison.OrdinalIgnoreCase)).ToList();
+		public SDKDescriptor? AutoSdk => Sdks.FirstOrDefault(x => String.Equals(x.Name, "AutoSdk", StringComparison.OrdinalIgnoreCase));
 
-		public SDKCollection(SDKCollection Other, UEBuildPlatformSDK PlatformSDK)
-			: this(PlatformSDK)
+		public SDKCollection(SDKCollection other, UEBuildPlatformSDK platformSDK)
+			: this(platformSDK)
 		{
-			foreach (SDKDescriptor Desc in Other.Sdks)
+			foreach (SDKDescriptor desc in other.Sdks)
 			{
-				SetupSDK(Desc.Name, Desc.Min, Desc.Max, Desc.Current, Desc.GroupName);
+				SetupSDK(desc.Name, desc.Min, desc.Max, desc.Current, desc.GroupName);
 			}
 		}
 
 		// convenience for single unnamed Sdk range
-		public SDKCollection(string? Min, string? Max, UEBuildPlatformSDK PlatformSDK)
-			: this(PlatformSDK)
+		public SDKCollection(string? min, string? max, UEBuildPlatformSDK platformSDK)
+			: this(platformSDK)
 		{
-			Sdks.Add(new SDKDescriptor(DefaultName, Min, Max, null, null, this));
+			Sdks.Add(new SDKDescriptor(DefaultName, min, max, null, null, this));
 		}
 
 		// convenience for single unnamed Sdk current value
-		public SDKCollection(string? Current, UEBuildPlatformSDK PlatformSDK)
-			: this(PlatformSDK)
+		public SDKCollection(string? current, UEBuildPlatformSDK platformSDK)
+			: this(platformSDK)
 		{
-			if (Current != null)
+			if (current != null)
 			{
-				Sdks.Add(new SDKDescriptor(DefaultName, null, null, Current, null, this));
+				Sdks.Add(new SDKDescriptor(DefaultName, null, null, current, null, this));
 			}
 		}
 
 		public bool AreAllManualSDKsValid()
 		{
-			IEnumerable<SDKDescriptor>? Sdks = FullSdks;
+			IEnumerable<SDKDescriptor>? sdks = FullSdks;
 
-			if (Sdks == null)
+			if (sdks == null)
 			{
 				return true;
 			}
 
 			// check if all the SDKs not in groups are valid (All returns true if empty set)
-			bool bAllUngroupedAreValid = Sdks
-				.Where(Desc => Desc.GroupName == null)
-				.All(Desc => Desc.Validity == SDKStatus.Valid);
+			bool bAllUngroupedAreValid = sdks
+				.Where(desc => desc.GroupName == null)
+				.All(desc => desc.Validity == SDKStatus.Valid);
 
 			// for each groupname (ToHashSet removes duplicates), check if at least 1 SDK with that group name is valid
-			bool bAllGroupsHaveOneValid = Sdks
-				.Where(Desc => Desc.GroupName != null)
-				.Select(Desc => Desc.GroupName)
+			bool bAllGroupsHaveOneValid = sdks
+				.Where(desc => desc.GroupName != null)
+				.Select(desc => desc.GroupName)
 				.ToHashSet()
-				.All(GroupName => Sdks.Any(Desc => Desc.GroupName == GroupName && Desc.Validity == SDKStatus.Valid));
+				.All(groupName => sdks.Any(desc => desc.GroupName == groupName && desc.Validity == SDKStatus.Valid));
 
 			return bAllUngroupedAreValid && bAllGroupsHaveOneValid;
 		}
@@ -192,90 +192,90 @@ namespace EpicGames.Core
 			return ToString("Version", "Version");
 		}
 
-		public string ToString(string Descriptor, string? CurrentDescriptor = null)
+		public string ToString(string descriptor, string? currentDescriptor = null)
 		{
 			if (Sdks.Count == 1)
 			{
-				return Sdks[0].ToString(Descriptor, CurrentDescriptor, false);
+				return Sdks[0].ToString(descriptor, currentDescriptor, false);
 			}
-			return String.Join(", ", Sdks.Select(x => x.ToString(Descriptor, CurrentDescriptor, true)));
+			return String.Join(", ", Sdks.Select(x => x.ToString(descriptor, currentDescriptor, true)));
 		}
 
-		public string ToMultilineString(bool bIncludeCurrent = true)
+		public static string ToMultilineString(bool bIncludeCurrent = true)
 		{
 			return $"";// MinVersion={Min ?? ""}, MaxVersion={Max ?? ""}" + (bIncludeCurrent ? $", Current={Current ?? ""}" : "");
 		}
 
-		public void SetupSDK(string Name, string? Min, string? Max, string? Current, string? GroupName)
+		public void SetupSDK(string name, string? min, string? max, string? current, string? groupName)
 		{
-			Sdks.Add(new SDKDescriptor(Name, Min, Max, Current, GroupName, this));
+			Sdks.Add(new SDKDescriptor(name, min, max, current, groupName, this));
 		}
 
-		public void SetupCurrent(string Name, string Current)
+		public void SetupCurrent(string name, string current)
 		{
-			Sdks.Add(new SDKDescriptor(Name, null, null, Current, null, this));
+			Sdks.Add(new SDKDescriptor(name, null, null, current, null, this));
 		}
 
-		public bool UpdateCurrentForSingle(string Name, string CurrentVersion)
+		public bool UpdateCurrentForSingle(string name, string currentVersion)
 		{
 			if (PlatformSDK == null)
 			{
 				throw new Exception("Cannot call UpdateCurrentForSingle in a SDKCollection without a PlatformSDK set");
 			}
 
-			SDKDescriptor? Info;
+			SDKDescriptor? info;
 
 			// for ease, we assume that if there is only one "Sdk" or "Software" in this Collection, then any Name passed in will work with it
 			// so just use that one entry
 			if (Sdks.Count == 1 && (Sdks[0].Name == "Sdk" || Sdks[0].Name == "Software"))
 			{
-				Info = Sdks[0];
+				info = Sdks[0];
 			}
 			else
 			{
-				Info = Sdks.FirstOrDefault(x => String.Compare(x.Name, Name, true) == 0);
-				if (Info == null)
+				info = Sdks.FirstOrDefault(x => String.Equals(x.Name, name, StringComparison.OrdinalIgnoreCase));
+				if (info == null)
 				{
 					return false;
 				}
 			}
 
-			Info.UpdateCurrent(CurrentVersion, Name, this);
+			info.UpdateCurrent(currentVersion, name, this);
 			return true;
 		}
 	}
 
 	public class SoftwareCollection : SDKCollection
 	{
-		public SoftwareCollection(UEBuildPlatformSDK PlatformSDK)
-			: base(PlatformSDK)
+		public SoftwareCollection(UEBuildPlatformSDK platformSDK)
+			: base(platformSDK)
 		{
 		}
-		public SoftwareCollection(string? Min, string? Max, UEBuildPlatformSDK PlatformSDK)
-			: base(Min, Max, PlatformSDK)
+		public SoftwareCollection(string? min, string? max, UEBuildPlatformSDK platformSDK)
+			: base(min, max, platformSDK)
 		{
 		}
-		public SoftwareCollection(string? Current, UEBuildPlatformSDK PlatformSDK)
-			: base(Current, PlatformSDK)
+		public SoftwareCollection(string? current, UEBuildPlatformSDK platformSDK)
+			: base(current, platformSDK)
 		{
 		}
 
 		public override string DefaultName => "Software";
 
-		public override bool IsValid(string Version, SDKDescriptor Info)
+		public override bool IsValid(string version, SDKDescriptor info)
 		{
-			return PlatformSDK.IsSoftwareVersionValid(Version, Info);
+			return PlatformSDK.IsSoftwareVersionValid(version, info);
 		}
 	}
 
 	/// <summary>
 	/// SDK for a platform
 	/// </summary>
-	public abstract class UEBuildPlatformSDK(ILogger InLogger)
+	public abstract class UEBuildPlatformSDK(ILogger inLogger)
 	{
 		// Public SDK handling, not specific to AutoSDK
 
-		protected readonly ILogger Logger = InLogger;
+		protected readonly ILogger Logger = inLogger;
 
 		protected readonly HashSet<FileReference> LoadedConfigFiles = [];
 
@@ -284,25 +284,25 @@ namespace EpicGames.Core
 		/// <summary>
 		/// Registers the SDK for a given platform (as a string, but equivalent to UnrealTargetPlatform)
 		/// </summary>
-		/// <param name="SDK">SDK object</param>
-		/// <param name="PlatformName">Platform name for this SDK</param>
+		/// <param name="sdk">SDK object</param>
+		/// <param name="platformName">Platform name for this SDK</param>
 		/// <param name="bIsSdkAllowedOnHost"></param>
-		public static void RegisterSDKForPlatform(UEBuildPlatformSDK SDK, string PlatformName, bool bIsSdkAllowedOnHost)
+		public static void RegisterSDKForPlatform(UEBuildPlatformSDK sdk, string platformName, bool bIsSdkAllowedOnHost)
 		{
 			// verify that neither platform or sdk were added before
-			if (SDKRegistry.Any(x => x.Key == PlatformName || x.Value == SDK))
+			if (SDKRegistry.Any(x => x.Key == platformName || x.Value == sdk))
 			{
-				throw new Exception(String.Format("Re-registering SDK for {0}. All Platforms must have a unique SDK object", PlatformName));
+				throw new Exception(String.Format("Re-registering SDK for {0}. All Platforms must have a unique SDK object", platformName));
 			}
 
-			SDKRegistry.Add(PlatformName, SDK);
+			SDKRegistry.Add(platformName, sdk);
 
-			SDK.Init(PlatformName, bIsSdkAllowedOnHost);
+			sdk.Init(platformName, bIsSdkAllowedOnHost);
 		}
 
-		private void Init(string InPlatformName, bool bInIsSdkAllowedOnHost)
+		private void Init(string inPlatformName, bool bInIsSdkAllowedOnHost)
 		{
-			PlatformName = InPlatformName;
+			PlatformName = inPlatformName;
 			bIsSdkAllowedOnHost = bInIsSdkAllowedOnHost;
 
 			// load the SDK config file
@@ -312,41 +312,41 @@ namespace EpicGames.Core
 			}
 
 			// if the parent set up autosdk, the env vars will be wrong, but we can still get the manual SDK version from before it was setup
-			string? ParentManualSDKVersions = Environment.GetEnvironmentVariable(GetPlatformManualSDKSetupEnvVar());
-			if (!String.IsNullOrEmpty(ParentManualSDKVersions))
+			string? parentManualSDKVersions = Environment.GetEnvironmentVariable(GetPlatformManualSDKSetupEnvVar());
+			if (!String.IsNullOrEmpty(parentManualSDKVersions))
 			{
 
 				// we pass along __None to indicate the parent didn't have a manual sdk installed
-				if (ParentManualSDKVersions == "__None")
+				if (parentManualSDKVersions == "__None")
 				{
 					// empty it out
 					CachedManualSDKVersions = [];
 				}
 				else
 				{
-					Logger.LogInformation("Processing parent manual sdk: {SDKVersion}", ParentManualSDKVersions);
-					CachedManualSDKVersions = JsonSerializer.Deserialize<Dictionary<string, string>>(ParentManualSDKVersions) ?? [];
+					Logger.LogInformation("Processing parent manual sdk: {SDKVersion}", parentManualSDKVersions);
+					CachedManualSDKVersions = JsonSerializer.Deserialize<Dictionary<string, string>>(parentManualSDKVersions) ?? [];
 				}
 			}
 			else
 			{
 				// allow the platform to select the most appropriate manually installed SDK if it hasn't already been set up by AutoSDK
 				bool bHasAutoSDK = PlatformSupportsAutoSDKs() && HasAutoSDKSystemEnabled(); // @todo: PlatformSupportsAutoSDKs & HasAutoSDKSystemEnabled check is not technically needed for manual sdk switching but the editor-side code is heavily tied into this and would require further refactoring first
-				if ( (bHasAutoSDK || !ManualSDKAutoSwitchRequiresAutoSDK) && !HasParentProcessSetupAutoSDK(out _)) 
+				if ((bHasAutoSDK || !ManualSDKAutoSwitchRequiresAutoSDK) && !HasParentProcessSetupAutoSDK(out _))
 				{
-					if (TrySelectBestManualSDK( out string? SelectedSDKVersion ))
+					if (TrySelectBestManualSDK(out string? selectedSDKVersion))
 					{
-						Logger.LogInformation("Auto-selected best installed sdk for {PlatformName} - {SDKVersion}", GetAutoSDKPlatformName(), SelectedSDKVersion);
+						Logger.LogInformation("Auto-selected best installed sdk for {PlatformName} - {SDKVersion}", GetAutoSDKPlatformName(), selectedSDKVersion);
 					}
 				}
 
 				CachedManualSDKVersions = [];
 
 				// if there was no parent, get the SDK version before we run AutoSDK to get the manual version
-				SDKCollection InstalledVersions = GetInstalledSDKVersions();
-				foreach (SDKDescriptor Desc in InstalledVersions.FullSdks)
+				SDKCollection installedVersions = GetInstalledSDKVersions();
+				foreach (SDKDescriptor desc in installedVersions.FullSdks)
 				{
-					CachedManualSDKVersions[Desc.Name] = Desc.Current!;
+					CachedManualSDKVersions[desc.Name] = desc.Current!;
 				}
 			}
 		}
@@ -358,36 +358,36 @@ namespace EpicGames.Core
 		/// <summary>
 		/// Retrieves a previously registered SDK for a given platform
 		/// </summary>
-		/// <param name="PlatformName">String name of the platform (equivalent to UnrealTargetPlatform)</param>
+		/// <param name="platformName">String name of the platform (equivalent to UnrealTargetPlatform)</param>
 		/// <returns></returns>
-		public static UEBuildPlatformSDK? GetSDKForPlatform(string PlatformName)
+		public static UEBuildPlatformSDK? GetSDKForPlatform(string platformName)
 		{
-			UEBuildPlatformSDK? SDK;
-			SDKRegistry.TryGetValue(PlatformName, out SDK);
+			UEBuildPlatformSDK? sdk;
+			SDKRegistry.TryGetValue(platformName, out sdk);
 
-			return SDK;
+			return sdk;
 		}
 
-		public static T? GetSDKForPlatformOrMakeTemp<T>(string PlatformName) where T : UEBuildPlatformSDK
+		public static T? GetSDKForPlatformOrMakeTemp<T>(string platformName) where T : UEBuildPlatformSDK
 		{
-			UEBuildPlatformSDK? SDK;
-			SDKRegistry.TryGetValue(PlatformName, out SDK);
+			UEBuildPlatformSDK? sdk;
+			SDKRegistry.TryGetValue(platformName, out sdk);
 
 			// make a temp one if needed, this is not expected to happen often at all
-			if (SDK == null)
+			if (sdk == null)
 			{
-				if (!TempSDKRegistry.TryGetValue(PlatformName, out SDK))
+				if (!TempSDKRegistry.TryGetValue(platformName, out sdk))
 				{
 					object[] parameter = [Log.Logger];
-					SDK = (UEBuildPlatformSDK)Activator.CreateInstance(typeof(T), parameter)!;
+					sdk = (UEBuildPlatformSDK)Activator.CreateInstance(typeof(T), parameter)!;
 					// by setting this to false, we don't require any of the RequiredVersions to exist, but if they do, they will be read
 					// this is useful on other platforms that 
-					SDK.bIsSdkAllowedOnHost = false;
-					SDK.LoadJsonFile(PlatformName);
-					TempSDKRegistry.Add(PlatformName, SDK);
+					sdk.bIsSdkAllowedOnHost = false;
+					sdk.LoadJsonFile(platformName);
+					TempSDKRegistry.Add(platformName, sdk);
 				}
 			}
-			return (T?)SDK;
+			return (T?)sdk;
 		}
 
 		/// <summary>
@@ -415,71 +415,71 @@ namespace EpicGames.Core
 
 		public SDKCollection GetAllSDKInfo()
 		{
-			SDKCollection AllSdks = new SDKCollection(this);
+			SDKCollection allSdks = new SDKCollection(this);
 
 			if (bIsSdkAllowedOnHost)
 			{
 				// walk over each one version the platform supports, and get it's current installed info
-				foreach (SDKDescriptor Desc in GetValidVersions().Sdks)
+				foreach (SDKDescriptor desc in GetValidVersions().Sdks)
 				{
-					AllSdks.SetupSDK(Desc.Name, Desc.Min, Desc.Max, CachedManualSDKVersions.GetValueOrDefault(Desc.Name), Desc.GroupName);
+					allSdks.SetupSDK(desc.Name, desc.Min, desc.Max, CachedManualSDKVersions.GetValueOrDefault(desc.Name), desc.GroupName);
 				}
 
 				// now get current AutoSDK version (if autosdk is set up, then the GetInstalledSDKVersion will return AutoSDK version)
 				bool bIsAutoSDK = false;
-				string? CurrentAutoSDKVersion = (PlatformSupportsAutoSDKs() && HasRequiredAutoSDKInstalled() == SDKStatus.Valid && HasSetupAutoSDK()) ? GetInstalledVersion(out bIsAutoSDK) : null;
-				AllSdks.SetupSDK("AutoSdk", GetMainVersion(), GetMainVersion(), CurrentAutoSDKVersion, null);
+				string? currentAutoSDKVersion = (PlatformSupportsAutoSDKs() && HasRequiredAutoSDKInstalled() == SDKStatus.Valid && HasSetupAutoSDK()) ? GetInstalledVersion(out bIsAutoSDK) : null;
+				allSdks.SetupSDK("AutoSdk", GetMainVersion(), GetMainVersion(), currentAutoSDKVersion, null);
 
 				// verify some assumptions
-				if (CurrentAutoSDKVersion != null && bIsAutoSDK == false)
+				if (currentAutoSDKVersion != null && bIsAutoSDK == false)
 				{
-					throw new Exception($"AutoSDK was indicated to be setup ({CurrentAutoSDKVersion}), but GetInstalledSDKVersion returned false for bIsAutoSDK");
+					throw new Exception($"AutoSDK was indicated to be setup ({currentAutoSDKVersion}), but GetInstalledSDKVersion returned false for bIsAutoSDK");
 				}
-				if (CurrentAutoSDKVersion != null && CurrentAutoSDKVersion != GetMainVersion())
+				if (currentAutoSDKVersion != null && currentAutoSDKVersion != GetMainVersion())
 				{
-					throw new Exception($"AutoSDK was indicated to be setup, but the version if returned ({CurrentAutoSDKVersion}) doesn't equal the MainVersion ({GetMainVersion()}");
+					throw new Exception($"AutoSDK was indicated to be setup, but the version if returned ({currentAutoSDKVersion}) doesn't equal the MainVersion ({GetMainVersion()}");
 				}
 			}
 
-			return AllSdks;
+			return allSdks;
 		}
 
-		public SDKDescriptor? GetSDKInfo(string SDKName)
+		public SDKDescriptor? GetSDKInfo(string sdkName)
 		{
-			return GetAllSDKInfo().Sdks.FirstOrDefault(x => String.Compare(x.Name, SDKName, true) == 0);
+			return GetAllSDKInfo().Sdks.FirstOrDefault(x => String.Equals(x.Name, sdkName, StringComparison.OrdinalIgnoreCase));
 		}
 
-		public SDKDescriptor? GetSoftwareInfo(string? SDKName = null)
+		public SDKDescriptor? GetSoftwareInfo(string? sdkName = null)
 		{
 			// todo: make this queryable?
-			SDKName ??= "Software";
-			return GetAllSoftwareInfo().Sdks.FirstOrDefault(x => String.Compare(x.Name, SDKName, true) == 0);
+			sdkName ??= "Software";
+			return GetAllSoftwareInfo().Sdks.FirstOrDefault(x => String.Equals(x.Name, sdkName, StringComparison.OrdinalIgnoreCase));
 		}
 
-		public SDKCollection GetAllSoftwareInfo(string? DeviceType = null, string? Current = null)
+		public SDKCollection GetAllSoftwareInfo(string? deviceType = null, string? current = null)
 		{
-			SDKCollection AllSoftwares = new SDKCollection(this);
+			SDKCollection allSoftwares = new SDKCollection(this);
 
 			// walk over the valid software ranges, potentially setting the current verison if one was passed in
 			// if there is only one known SDK named "Software" we allow any DeviceType
-			IEnumerable<SDKDescriptor> SoftwareSDKs = GetValidSoftwareVersions().Sdks;
+			IEnumerable<SDKDescriptor> softwareSDKs = GetValidSoftwareVersions().Sdks;
 
-			foreach (SDKDescriptor Desc in SoftwareSDKs)
+			foreach (SDKDescriptor desc in softwareSDKs)
 			{
 				// are we restricting this to one type?
-				if (DeviceType == null || string.Compare(Desc.Name, DeviceType, true) == 0)
+				if (deviceType == null || String.Equals(desc.Name, deviceType, StringComparison.OrdinalIgnoreCase))
 				{
-					AllSoftwares.SetupSDK(Desc.Name, Desc.Min, Desc.Max, null, Desc.GroupName);
+					allSoftwares.SetupSDK(desc.Name, desc.Min, desc.Max, null, desc.GroupName);
 				}
 
-				if (Current != null)
+				if (current != null)
 				{
 					// set the current for this one (likely it will only be for DeviceType, but support passing in null DeviceType, and non-null Current
-					AllSoftwares.UpdateCurrentForSingle(Desc.Name, Current);
+					allSoftwares.UpdateCurrentForSingle(desc.Name, current);
 				}
 			}
 
-			return AllSoftwares;
+			return allSoftwares;
 		}
 
 		//public SDKCollection GetAllSoftwareInfo(string DeviceId)
@@ -498,9 +498,9 @@ namespace EpicGames.Core
 			return GetInstalledSDKVersion();
 		}
 
-		public string? GetInstalledVersion(string SDKName)
+		public string? GetInstalledVersion(string sdkName)
 		{
-			return GetSDKInfo(SDKName)?.Current;
+			return GetSDKInfo(sdkName)?.Current;
 		}
 
 		//		public void GetInstalledVersions(out string? ManualSDKVersion, out string? AutoSDKVersion)
@@ -523,22 +523,13 @@ namespace EpicGames.Core
 		//			ManualSDKVersion = CachedManualSDKVersion;
 		//		}
 
-		public virtual bool IsVersionValid(string? Version, string SDKType)
-		{
-			return IsVersionValidInternal(Version, SDKType, null);
-		}
-		public virtual bool IsVersionValid(string? Version, SDKDescriptor Info)
-		{
-			return IsVersionValidInternal(Version, null, Info);
-		}
-		public virtual bool IsSoftwareVersionValid(string Version, string DeviceType)
-		{
-			return IsSoftwareVersionValidInternal(Version, DeviceType, null);
-		}
-		public virtual bool IsSoftwareVersionValid(string Version, SDKDescriptor Info)
-		{
-			return IsSoftwareVersionValidInternal(Version, null, Info);
-		}
+		public virtual bool IsVersionValid(string? version, string sdkType) => IsVersionValidInternal(version, sdkType, null);
+
+		public virtual bool IsVersionValid(string? version, SDKDescriptor info) => IsVersionValidInternal(version, null, info);
+
+		public virtual bool IsSoftwareVersionValid(string version, string deviceType) => IsSoftwareVersionValidInternal(version, deviceType, null);
+
+		public virtual bool IsSoftwareVersionValid(string version, SDKDescriptor info) => IsSoftwareVersionValidInternal(version, null, info);
 
 		public void ReactivateAutoSDK()
 		{
@@ -563,41 +554,38 @@ namespace EpicGames.Core
 		/// <summary>
 		/// Gets the valid string range of Sdk versions. TryConvertVersionToInt() will need to succeed to make this usable for range checks
 		/// </summary>
-		/// <param name="MinVersion">Smallest version allowed</param>
-		/// <param name="MaxVersion">Largest version allowed (inclusive)</param>
-		protected virtual void GetValidVersionRange(out string? MinVersion, out string? MaxVersion)
+		/// <param name="minVersion">Smallest version allowed</param>
+		/// <param name="maxVersion">Largest version allowed (inclusive)</param>
+		protected virtual void GetValidVersionRange(out string? minVersion, out string? maxVersion)
 		{
 			// by default, use SDK config file values
-			MinVersion = GetVersionFromConfig("MinVersion");
-			MaxVersion = GetVersionFromConfig("MaxVersion");
+			minVersion = GetVersionFromConfig("MinVersion");
+			maxVersion = GetVersionFromConfig("MaxVersion");
 		}
 
 		protected virtual SDKCollection GetValidVersions()
 		{
 			// if the platform doesn't override this, then it only has one sdk, so put it into the collection as the single sdk (this is very much the standard behavior)
-			string? MinVersion, MaxVersion;
-			GetValidVersionRange(out MinVersion, out MaxVersion);
-			return new SDKCollection(MinVersion, MaxVersion, this);
+			GetValidVersionRange(out string? minVersion, out string? maxVersion);
+			return new SDKCollection(minVersion, maxVersion, this);
 		}
 
 		/// <summary>
 		/// Gets the valid string range of software/flash versions. TryConvertVersionToInt() will need to succeed to make this usable for range checks
 		/// </summary>
-		/// <param name="MinVersion">Smallest version allowed, or null if no minmum (in other words, 0 - MaxVersion)</param>
-		/// <param name="MaxVersion">Largest version allowed (inclusive), or null if no maximum (in other words, MinVersion - infinity)y</param>
-		protected virtual void GetValidSoftwareVersionRange(out string? MinVersion, out string? MaxVersion)
+		/// <param name="minVersion">Smallest version allowed, or null if no minmum (in other words, 0 - MaxVersion)</param>
+		/// <param name="maxVersion">Largest version allowed (inclusive), or null if no maximum (in other words, MinVersion - infinity)y</param>
+		protected virtual void GetValidSoftwareVersionRange(out string? minVersion, out string? maxVersion)
 		{
-			MinVersion = GetVersionFromConfig("MinSoftwareVersion");
-			MaxVersion = GetVersionFromConfig("MaxSoftwareVersion");
+			minVersion = GetVersionFromConfig("MinSoftwareVersion");
+			maxVersion = GetVersionFromConfig("MaxSoftwareVersion");
 		}
 
 		protected virtual SoftwareCollection GetValidSoftwareVersions()
 		{
 			// if the platform doesn't override this, then it only has one sdk, so put it into the collection as the single software (this is very much the standard behavior)
-			string? MinVersion, MaxVersion;
-			GetValidSoftwareVersionRange(out MinVersion, out MaxVersion);
-			SoftwareCollection SoftwareVersions = new SoftwareCollection(MinVersion, MaxVersion, this);
-			return SoftwareVersions;
+			GetValidSoftwareVersionRange(out string? minVersion, out string? maxVersion);
+			return new SoftwareCollection(minVersion, maxVersion, this);
 		}
 
 		/// <summary>
@@ -618,10 +606,10 @@ namespace EpicGames.Core
 			if (bIsSdkAllowedOnHost)
 			{
 				// if the platform doesn't override this, then it only has one sdk, so put it into the collection as the single sdk (this is very much the standard behavior)
-				string? Version = GetInstalledSDKVersion();
-				if (Version != null)
+				string? version = GetInstalledSDKVersion();
+				if (version != null)
 				{
-					return new SDKCollection(Version, this);
+					return new SDKCollection(version, this);
 				}
 			}
 
@@ -632,11 +620,11 @@ namespace EpicGames.Core
 		/// <summary>
 		/// Returns a platform-specific version string that can be retrieved from anywhere without needing to typecast the SDK object
 		/// </summary>
-		/// <param name="VersionType">Descriptor of the version you want</param>
+		/// <param name="versionType">Descriptor of the version you want</param>
 		/// <returns>Version string, or empty string if not known/supported</returns>
-		public virtual string GetPlatformSpecificVersion(string VersionType)
+		public virtual string GetPlatformSpecificVersion(string versionType)
 		{
-			return GetRequiredVersionFromConfig(VersionType);
+			return GetRequiredVersionFromConfig(versionType);
 		}
 
 		/// <summary>
@@ -647,56 +635,56 @@ namespace EpicGames.Core
 		{
 			return [];
 		}
-		
+
 		/// <summary>
 		/// Find the highest valid manually-installed SDK, prioritizing the main version if it's available. Requires that the platform SDK implements GetAllInstalledSDKVersions
 		/// </summary>
 		/// <returns>True if successful</returns>
 		protected virtual string? FindBestInstalledSDKVersion()
 		{
-			string[] InstalledSDKVersions = GetAllInstalledSDKVersions();
-			if (InstalledSDKVersions.Length == 0)
+			string[] installedSDKVersions = GetAllInstalledSDKVersions();
+			if (installedSDKVersions.Length == 0)
 			{
 				return null;
 			}
 
 			// see if the main version is available
-			string MainVersion = GetMainVersion();
-			if (InstalledSDKVersions.Contains(MainVersion))
+			string mainVersion = GetMainVersion();
+			if (installedSDKVersions.Contains(mainVersion))
 			{
-				return MainVersion;
+				return mainVersion;
 			}
 
 			// main version is not available - find highest valid version
-			GetValidVersionRange(out string? MinVersion, out string? MaxVersion);
-			if (TryConvertVersionToInt(MinVersion, out ulong MinVersionInt, null) && TryConvertVersionToInt(MaxVersion, out ulong MaxVersionInt, null))
+			GetValidVersionRange(out string? minVersion, out string? maxVersion);
+			if (TryConvertVersionToInt(minVersion, out ulong minVersionInt, null) && TryConvertVersionToInt(maxVersion, out ulong maxVersionInt, null))
 			{
-				List<Tuple<ulong,ulong>> BannedSdkVersionRanges = [];
-				foreach (string SdkVersionPair in GetStringArrayFromConfig("BannedSdkVersions"))
+				List<Tuple<ulong, ulong>> bannedSdkVersionRanges = [];
+				foreach (string sdkVersionPair in GetStringArrayFromConfig("BannedSdkVersions"))
 				{
-					string[] Versions = SdkVersionPair.Split("-");
-					if (Versions.Length > 0 && Versions.Length <= 2)
+					string[] versions = sdkVersionPair.Split("-");
+					if (versions.Length > 0 && versions.Length <= 2)
 					{
-						if (TryConvertVersionToInt(Versions.First(), out ulong MinInvalidVersion, null) && TryConvertVersionToInt(Versions.Last(), out ulong MaxInvalidVersion, null))
+						if (TryConvertVersionToInt(versions.First(), out ulong minInvalidVersion, null) && TryConvertVersionToInt(versions.Last(), out ulong maxInvalidVersion, null))
 						{
-							BannedSdkVersionRanges.Add( new( MinInvalidVersion, MaxInvalidVersion));
+							bannedSdkVersionRanges.Add(new(minInvalidVersion, maxInvalidVersion));
 						}
 					}
 				}
 
-				IGrouping<ulong,string>? ResultGroup = InstalledSDKVersions
-					.GroupBy( X => TryConvertVersionToInt(X, out ulong Value, null) ? Value : 0 )            // group by integer version
-					.Where( X => X.Key >= MinVersionInt && X.Key <= MaxVersionInt)                           // select only valid versions
-					.OrderByDescending( X => X.Key )                                                         // sort highest version first
-					.OrderBy(X => BannedSdkVersionRanges.Any(R => X.Key >= R.Item1 && X.Key <= R.Item2))     // put banned versions at the end of the list - they'll be excluded later
+				IGrouping<ulong, string>? resultGroup = installedSDKVersions
+					.GroupBy(x => TryConvertVersionToInt(x, out ulong value, null) ? value : 0)            // group by integer version
+					.Where(x => x.Key >= minVersionInt && x.Key <= maxVersionInt)                           // select only valid versions
+					.OrderByDescending(x => x.Key)                                                         // sort highest version first
+					.OrderBy(x => bannedSdkVersionRanges.Any(r => x.Key >= r.Item1 && x.Key <= r.Item2))     // put banned versions at the end of the list - they'll be excluded later
 					.FirstOrDefault();
-				if (ResultGroup?.FirstOrDefault() != null && BannedSdkVersionRanges.Any(R => ResultGroup.Key >= R.Item1 && ResultGroup.Key <= R.Item2))
+				if (resultGroup?.FirstOrDefault() != null && bannedSdkVersionRanges.Any(r => resultGroup.Key >= r.Item1 && resultGroup.Key <= r.Item2))
 				{
-					Logger.LogInformation("The best installed SDK for {Platform} is {Version} but this version is forbidden", PlatformName, ResultGroup.FirstOrDefault());
+					Logger.LogInformation("The best installed SDK for {Platform} is {Version} but this version is forbidden", PlatformName, resultGroup.FirstOrDefault());
 					return null;
 				}
 
-				return ResultGroup?.FirstOrDefault();
+				return resultGroup?.FirstOrDefault();
 			}
 
 			return null;
@@ -705,10 +693,10 @@ namespace EpicGames.Core
 		/// <summary>
 		/// Switch to another version of the SDK than what GetInstalledSDKVersion() returns. This will be one of the versions returned from GetAllInstalledSDKVersions()
 		/// </summary>
-		/// <param name="Version">String name of the version to switch to</param>
+		/// <param name="version">String name of the version to switch to</param>
 		/// <param name="bSwitchForThisProcessOnly">If true, only switch for this process (usually via process environment variable). If false, switch permanently (usually via system-wide environment variable)</param>
 		/// <returns>True if successful</returns>
-		public virtual bool SwitchToAlternateSDK(string Version, bool bSwitchForThisProcessOnly)
+		public virtual bool SwitchToAlternateSDK(string version, bool bSwitchForThisProcessOnly)
 		{
 			return false;
 		}
@@ -724,11 +712,11 @@ namespace EpicGames.Core
 		/// Allows a platform to switch a different version of a manually-installed SDK, if the platform supports side-by-side installations of different versions
 		/// </summary>
 		/// <returns>True if successful</returns>
-		protected virtual bool TryGetEnvironmentForManualSDK(string SelectedSDKVersion, out Dictionary<string, string>? EnvVarValues, out List<string>? PathAdds, out List<string>? PathRemoves)
+		protected virtual bool TryGetEnvironmentForManualSDK(string selectedSDKVersion, out Dictionary<string, string>? envVarValues, out List<string>? pathAdds, out List<string>? pathRemoves)
 		{
-			EnvVarValues = null;
-			PathAdds = null;
-			PathRemoves = null;
+			envVarValues = null;
+			pathAdds = null;
+			pathRemoves = null;
 			return false;
 		}
 
@@ -778,102 +766,100 @@ namespace EpicGames.Core
 		//}
 
 		// Let platform override behavior to determine if a version is a valid (useful for non-numeric versions)
-		protected virtual bool IsVersionValidInternal(string? Version, string? SDKType, SDKDescriptor? VersionInfo)
+		protected virtual bool IsVersionValidInternal(string? version, string? sdkType, SDKDescriptor? versionInfo)
 		{
 			// we could have null if no SDK is installed at all, etc, which is always a failure
-			if (Version == null)
+			if (version == null)
 			{
 				return false;
 			}
 
 			// AutoSDK must match the desired version exactly, since that is the only one we will use
-			if (String.Compare(SDKType, "AutoSdk", true) == 0)
+			if (String.Equals(sdkType, "AutoSdk", StringComparison.OrdinalIgnoreCase))
 			{
 				// if integer version checking failed, then we can detect valid autosdk if the version matches the autosdk directory by name
-				return String.Compare(Version, GetAutoSDKDirectoryForMainVersion(), true) == 0;
+				return String.Equals(version, GetAutoSDKDirectoryForMainVersion(), StringComparison.OrdinalIgnoreCase);
 			}
 
 			// look for a range for this hinted type of SDK
-			VersionInfo = VersionInfo ?? GetSDKVersionForHint(GetValidVersions(), SDKType);
+			versionInfo = versionInfo ?? GetSDKVersionForHint(GetValidVersions(), sdkType);
 
 			// if we couldn't find the Info, we can't do anything, so fail
-			if (VersionInfo == null)
+			if (versionInfo == null)
 			{
 				return false;
 			}
 
 			// convert it to an integer
-			ulong IntVersion;
-			if (!TryConvertVersionToInt(Version, out IntVersion, VersionInfo.Name))
+			if (!TryConvertVersionToInt(version, out ulong intVersion, versionInfo.Name))
 			{
 				return false;
 			}
 
 			// short circuit range check if the Version is the desired version already
-			if (IntVersion == ConvertVersionToInt(GetMainVersion()))
+			if (intVersion == ConvertVersionToInt(GetMainVersion()))
 			{
 				return true;
 			}
 
 			// finally do the numeric comparison
-			return IntVersion >= VersionInfo.MinInt && IntVersion <= VersionInfo.MaxInt;
+			return intVersion >= versionInfo.MinInt && intVersion <= versionInfo.MaxInt;
 		}
 
-		protected virtual bool IsSoftwareVersionValidInternal(string? Version, string? DeviceType, SDKDescriptor? VersionInfo)
+		protected virtual bool IsSoftwareVersionValidInternal(string? version, string? deviceType, SDKDescriptor? versionInfo)
 		{
 			// we could have null if no SDK is installed at all, etc, which is always a failure
-			if (Version == null)
+			if (version == null)
 			{
 				return false;
 			}
 
 			// convert it to an integer
-			ulong IntVersion;
-			if (!TryConvertVersionToInt(Version, out IntVersion))
+			if (!TryConvertVersionToInt(version, out ulong intVersion))
 			{
 				return false;
 			}
 
 			// look for a range for this hinted type of SDK
-			VersionInfo = VersionInfo ?? GetSDKVersionForHint(GetValidSoftwareVersions(), DeviceType);
+			versionInfo = versionInfo ?? GetSDKVersionForHint(GetValidSoftwareVersions(), deviceType);
 
 			// if we couldn't find the Info, we can't do anything, so fail
-			if (VersionInfo == null)
+			if (versionInfo == null)
 			{
 				return false;
 			}
 
 			// finally do the numeric comparison
-			return IntVersion >= VersionInfo.MinInt && IntVersion <= VersionInfo.MaxInt;
+			return intVersion >= versionInfo.MinInt && intVersion <= versionInfo.MaxInt;
 		}
 
 		/// <summary>
 		/// Only the platform can convert a version string into an integer that is usable for comparison
 		/// </summary>
-		/// <param name="StringValue">Version that comes from the installed SDK or a Turnkey manifest or the like</param>
-		/// <param name="OutValue"></param>
-		/// <param name="Hint">A platform specific hint that can help guide conversion (usually SDKName or device type)</param>
+		/// <param name="stringValue">Version that comes from the installed SDK or a Turnkey manifest or the like</param>
+		/// <param name="outValue"></param>
+		/// <param name="hint">A platform specific hint that can help guide conversion (usually SDKName or device type)</param>
 		/// <returns>If the StringValue was able to be be converted to an integer</returns>
-		public abstract bool TryConvertVersionToInt(string? StringValue, out ulong OutValue, string? Hint = null);
+		public abstract bool TryConvertVersionToInt(string? stringValue, out ulong outValue, string? hint = null);
 
 		/// <summary>
 		/// Like TryConvertVersionToInt, but will throw an exception on failure
 		/// </summary>
-		/// <param name="StringValue">Version that comes from the PlatformSDK class (should not be used with Manifest or other user supplied versions)</param>
-		/// <param name="Hint">A platform specific hint that can help guide conversion (usually SDKName or device type)</param>
+		/// <param name="stringValue">Version that comes from the PlatformSDK class (should not be used with Manifest or other user supplied versions)</param>
+		/// <param name="hint">A platform specific hint that can help guide conversion (usually SDKName or device type)</param>
 		/// <returns>The integer version of StringValue, can be used to compare against a valid range</returns>
-		public ulong ConvertVersionToInt(string? StringValue, string? Hint = null)
+		public ulong ConvertVersionToInt(string? stringValue, string? hint = null)
 		{
 			// quickly handle null input, which is not necessarily an error case
-			if (StringValue == null)
+			if (stringValue == null)
 			{
 				return 0;
 			}
 
 			ulong Result;
-			if (!TryConvertVersionToInt(StringValue, out Result, Hint))
+			if (!TryConvertVersionToInt(stringValue, out Result, hint))
 			{
-				throw new Exception($"Unable to convert {GetType()} version '{StringValue}' to an integer. Likely this version was supplied by code, and is expected to be valid.");
+				throw new Exception($"Unable to convert {GetType()} version '{stringValue}' to an integer. Likely this version was supplied by code, and is expected to be valid.");
 			}
 			return Result;
 		}
@@ -882,20 +868,19 @@ namespace EpicGames.Core
 		/// Compare two sdk versions, for Sort() purposes
 		/// e.g. SdkVersions.Sort( (A,B) => PlatformSDK.SdkVersionsCompare(A,B) );
 		/// </summary>
-		/// <param name="StringValueA">First Version to compare</param>
-		/// <param name="StringValueB">Second Version to compare</param>
-		/// <param name="Hint">A platform specific hint that can help guide conversion (usually SDKName or device type)</param>
+		/// <param name="stringValueA">First Version to compare</param>
+		/// <param name="stringValueB">Second Version to compare</param>
+		/// <param name="hint">A platform specific hint that can help guide conversion (usually SDKName or device type)</param>
 		/// <returns>Comparison integer</returns>
-		public int SdkVersionsCompare(string? StringValueA, string? StringValueB, string? Hint = null)
+		public int SdkVersionsCompare(string? stringValueA, string? stringValueB, string? hint = null)
 		{
-			ulong ValueA, ValueB;
-			TryConvertVersionToInt(StringValueA, out ValueA, Hint);
-			TryConvertVersionToInt(StringValueB, out ValueB, Hint);
-			if (ValueA == ValueB)
+			TryConvertVersionToInt(stringValueA, out ulong valueA, hint);
+			TryConvertVersionToInt(stringValueB, out ulong valueB, hint);
+			if (valueA == valueB)
 			{
 				return 0;
 			}
-			else if (ValueA < ValueB)
+			else if (valueA < valueB)
 			{
 				return -1;
 			}
@@ -919,9 +904,9 @@ namespace EpicGames.Core
 		#region Per-project SDK support
 
 		private static List<FileReference> ProjectsToCheckForSDKOverrides = [];
-		public static void InitializePerProjectSDKVersions(IEnumerable<FileReference> ProjectsToCheckForOverrides)
+		public static void InitializePerProjectSDKVersions(IEnumerable<FileReference> projectsToCheckForOverrides)
 		{
-			ProjectsToCheckForSDKOverrides = ProjectsToCheckForOverrides.ToList();
+			ProjectsToCheckForSDKOverrides = projectsToCheckForOverrides.ToList();
 		}
 
 		#endregion
@@ -931,63 +916,63 @@ namespace EpicGames.Core
 		private SDKStatus? SDKInfoValidity = null;
 		public static bool bSuppressSDKWarnings = false;
 
-		public virtual SDKStatus PrintSDKInfoAndReturnValidity(LogEventType Verbosity = LogEventType.Console, LogFormatOptions Options = LogFormatOptions.None,
-			LogEventType ErrorVerbosity = LogEventType.Error, LogFormatOptions ErrorOptions = LogFormatOptions.None, bool bBriefInvalidSDKWarnings = false)
+		public virtual SDKStatus PrintSDKInfoAndReturnValidity(LogEventType verbosity = LogEventType.Console, LogFormatOptions options = LogFormatOptions.None,
+			LogEventType errorVerbosity = LogEventType.Error, LogFormatOptions errorOptions = LogFormatOptions.None, bool bBriefInvalidSDKWarnings = false)
 		{
 			if (SDKInfoValidity != null)
 			{
 				return SDKInfoValidity.Value;
 			}
 
-			SDKCollection SDKInfo = GetAllSDKInfo();
+			SDKCollection sdkInfo = GetAllSDKInfo();
 
 			// will mark invalid below if needed
 			SDKInfoValidity = SDKStatus.Valid;
 
-			SDKDescriptor? AutoSDKInfo = SDKInfo.AutoSdk;
-			if (AutoSDKInfo != null && AutoSDKInfo.Validity == SDKStatus.Valid)
+			SDKDescriptor? autoSDKInfo = sdkInfo.AutoSdk;
+			if (autoSDKInfo != null && autoSDKInfo.Validity == SDKStatus.Valid)
 			{
-				string PlatformSDKRoot = GetPathToPlatformAutoSDKs();
-				Log.WriteLine(Verbosity, Options, "{0} using Auto SDK {1} from: {2} 0x{3:X}", PlatformName, AutoSDKInfo.Current, Path.Combine(PlatformSDKRoot, GetAutoSDKDirectoryForMainVersion()), AutoSDKInfo.CurrentInt);
+				string platformSDKRoot = GetPathToPlatformAutoSDKs();
+				Log.WriteLine(verbosity, options, "{0} using Auto SDK {1} from: {2} 0x{3:X}", PlatformName, autoSDKInfo.Current, Path.Combine(platformSDKRoot, GetAutoSDKDirectoryForMainVersion()), autoSDKInfo.CurrentInt);
 			}
 			else
 			{
-				if (SDKInfo.AreAllManualSDKsValid())
+				if (sdkInfo.AreAllManualSDKsValid())
 				{
-					Log.WriteLine(Verbosity, Options, "{0} Installed SDK(s): {1}", PlatformName, SDKInfo.ToString(true));
+					Log.WriteLine(verbosity, options, "{0} Installed SDK(s): {1}", PlatformName, sdkInfo.ToString(true));
 				}
 				else
 				{
 					SDKInfoValidity = SDKStatus.Invalid;
 
-					StringBuilder Msg = new StringBuilder();
-					Msg.AppendLine($"Unable to find valid SDK(s) for {PlatformName}:");
+					StringBuilder msg = new StringBuilder();
+					msg.AppendLine($"Unable to find valid SDK(s) for {PlatformName}:");
 
-					foreach (SDKDescriptor Desc in SDKInfo.Sdks)
+					foreach (SDKDescriptor desc in sdkInfo.Sdks)
 					{
-						if (Desc.Validity == SDKStatus.Valid)
+						if (desc.Validity == SDKStatus.Valid)
 						{
-							Msg.AppendLine($"  {Desc.Name} is valid ({Desc}");
+							msg.AppendLine($"  {desc.Name} is valid ({desc}");
 						}
 						else
 						{
-							Msg.Append($"  Found {Desc.Name} Version");
+							msg.Append($"  Found {desc.Name} Version");
 
-							if (Desc.Current != null)
+							if (desc.Current != null)
 							{
-								Msg.Append($"={Desc.Current}");
+								msg.Append($"={desc.Current}");
 							}
 
-							Msg.AppendLine($", {Desc.ToString("Required", null, false)}.");
+							msg.AppendLine($", {desc.ToString("Required", null, false)}.");
 						}
 					}
 
 					if (!bBriefInvalidSDKWarnings && !bHasShownTurnkey)
 					{
-						Msg.AppendLine("  If your Studio has it set up, you can run this command to find the SDK to install:");
-						Msg.AppendLine("    RunUAT Turnkey -command=InstallSdk -platform={0} -BestAvailable", PlatformName!);
+						msg.AppendLine("  If your Studio has it set up, you can run this command to find the SDK to install:");
+						msg.AppendLine("    RunUAT Turnkey -command=InstallSdk -platform={0} -BestAvailable", PlatformName!);
 
-						if ((ErrorOptions & LogFormatOptions.NoConsoleOutput) == LogFormatOptions.None)
+						if ((errorOptions & LogFormatOptions.NoConsoleOutput) == LogFormatOptions.None)
 						{
 							bHasShownTurnkey = true;
 						}
@@ -996,11 +981,11 @@ namespace EpicGames.Core
 					// Reducing warnings to log to help prevent warnings locally or in Horde about SDKs we might not currently be concerned about
 					if (bSuppressSDKWarnings)
 					{
-						ErrorVerbosity = LogEventType.Log;
+						errorVerbosity = LogEventType.Log;
 					}
 
 					// always print errors to the screen
-					Log.WriteLine(ErrorVerbosity, ErrorOptions, Msg.ToString());
+					Log.WriteLine(errorVerbosity, errorOptions, msg.ToString());
 				}
 			}
 
@@ -1019,22 +1004,22 @@ namespace EpicGames.Core
 		// like in MicrosoftPlatofrmSDK.Version.cs, if one of the versions is needed, C# will go construct every version, needing the SDK multiple times
 		private static Dictionary<string, UEBuildPlatformSDK> TempSDKRegistry = [];
 
-		private SDKDescriptor? GetSDKVersionForHint(SDKCollection Collection, string? Hint)
+		private static SDKDescriptor? GetSDKVersionForHint(SDKCollection collection, string? hint)
 		{
 			// if the hint is found, use it always
-			SDKDescriptor? SDKDesc = Collection.Sdks.FirstOrDefault(x => String.Compare(x.Name, Hint, true) == 0);
-			if (SDKDesc != null)
+			SDKDescriptor? sdkDesc = collection.Sdks.FirstOrDefault(x => String.Equals(x.Name, hint, StringComparison.OrdinalIgnoreCase));
+			if (sdkDesc != null)
 			{
-				return SDKDesc;
+				return sdkDesc;
 			}
 
 			// only use AUtoSDK if explicitly asked for above, otherwise, remove it and look at what's left
-			IEnumerable<SDKDescriptor> FullSdks = Collection.FullSdks;
+			IEnumerable<SDKDescriptor> fullSdks = collection.FullSdks;
 
 			// if there's one with the special name, then use it as it's generic (common case here)
-			if (FullSdks.Count() == 1 && (FullSdks.First().Name == "Sdk" || FullSdks.First().Name == "Software"))
+			if (fullSdks.Count() == 1 && (fullSdks.First().Name == "Sdk" || fullSdks.First().Name == "Software"))
 			{
-				return FullSdks.First();
+				return fullSdks.First();
 			}
 
 			// finally we have to give up
@@ -1045,239 +1030,238 @@ namespace EpicGames.Core
 		private Dictionary<string, string> ConfigSDKVersions = new(StringComparer.OrdinalIgnoreCase);
 		private Dictionary<string, string[]> ConfigSDKVersionArrays = new(StringComparer.OrdinalIgnoreCase);
 
-		private void LoadJsonFile(string Platform)
+		private void LoadJsonFile(string platform)
 		{
 			// fixup for Windows
-			if (Platform == "Win64")
+			if (platform == "Win64")
 			{
-				Platform = "Windows";
+				platform = "Windows";
 			}
 
-			FileReference? MakeConfigFilename(DirectoryReference RootDir, bool bIsRequired)
+			FileReference? MakeConfigFilename(DirectoryReference rootDir, bool bIsRequired)
 			{
-				FileReference PlatformExtensionLocation = FileReference.Combine(RootDir, "Platforms", Platform, "Config", $"{Platform}_SDK.json");
-				if (FileReference.Exists(PlatformExtensionLocation))
+				FileReference platformExtensionLocation = FileReference.Combine(rootDir, "Platforms", platform, "Config", $"{platform}_SDK.json");
+				if (FileReference.Exists(platformExtensionLocation))
 				{
-					return PlatformExtensionLocation;
+					return platformExtensionLocation;
 				}
-				FileReference StandardLocation = FileReference.Combine(RootDir, "Config", Platform, $"{Platform}_SDK.json");
-				if (FileReference.Exists(StandardLocation))
+				FileReference standardLocation = FileReference.Combine(rootDir, "Config", platform, $"{platform}_SDK.json");
+				if (FileReference.Exists(standardLocation))
 				{
-					return StandardLocation;
+					return standardLocation;
 				}
 				if (bIsRequired)
 				{
-					throw new Exception($"Failed to find required SDK.json for {Platform}. Looked in '{StandardLocation}' and '{PlatformExtensionLocation}'.");
+					throw new Exception($"Failed to find required SDK.json for {platform}. Looked in '{standardLocation}' and '{platformExtensionLocation}'.");
 				}
 				return null;
 			}
 
 			// if the SDK isn't allowed on the host, then allow it to not exist
-			FileReference EngineSDKConfigFile = MakeConfigFilename(Unreal.EngineDirectory, bIsSdkAllowedOnHost)!;
+			FileReference engineSDKConfigFile = MakeConfigFilename(Unreal.EngineDirectory, bIsSdkAllowedOnHost)!;
 
 			// load the file, along with any chained group file
-			ProcessJsonFile(EngineSDKConfigFile, ConfigSDKVersions, ConfigSDKVersionArrays);
+			ProcessJsonFile(engineSDKConfigFile, ConfigSDKVersions, ConfigSDKVersionArrays);
 
 			// copy off the versions to the defaults, so we can check if overridden by the project
-			Dictionary<string, string> DefaultConfigSDKVersions = new(ConfigSDKVersions);
+			Dictionary<string, string> defaultConfigSDKVersions = new(ConfigSDKVersions);
 			foreach (string Key in ConfigSDKVersions.Keys)
 			{
-				DefaultConfigSDKVersions[Key] = ConfigSDKVersions[Key];
+				defaultConfigSDKVersions[Key] = ConfigSDKVersions[Key];
 			}
 
 			// now read overrides into the array
-			foreach (FileReference ProjectFile in ProjectsToCheckForSDKOverrides)
+			foreach (FileReference projectFile in ProjectsToCheckForSDKOverrides)
 			{
-				FileReference? ProjectSDKConfigFile = MakeConfigFilename(ProjectFile.Directory, false);
-				if (ProjectSDKConfigFile != null)
+				FileReference? projectSDKConfigFile = MakeConfigFilename(projectFile.Directory, false);
+				if (projectSDKConfigFile != null)
 				{
 					// load a project's SDK file if thre is one, into a temp dictionary
-					Dictionary<string, string> OverrideConfigSDKVersions = new(StringComparer.OrdinalIgnoreCase);
-					Dictionary<string, string[]> OverrideConfigSDKVersionArrays = new(StringComparer.OrdinalIgnoreCase);
-					ProcessJsonFile(ProjectSDKConfigFile, OverrideConfigSDKVersions, OverrideConfigSDKVersionArrays);
-					if (OverrideConfigSDKVersionArrays.Count > 0)
+					Dictionary<string, string> overrideConfigSDKVersions = new(StringComparer.OrdinalIgnoreCase);
+					Dictionary<string, string[]> overrideConfigSDKVersionArrays = new(StringComparer.OrdinalIgnoreCase);
+					ProcessJsonFile(projectSDKConfigFile, overrideConfigSDKVersions, overrideConfigSDKVersionArrays);
+					if (overrideConfigSDKVersionArrays.Count > 0)
 					{
-						throw new Exception($"Overriding version arrays, in project '{ProjectFile.GetFileNameWithoutExtension()}', platform {Platform} is not currently supported");
+						throw new Exception($"Overriding version arrays, in project '{projectFile.GetFileNameWithoutExtension()}', platform {platform} is not currently supported");
 					}
 
 					// currently only care about MainVersion in the overrides
-					string? OverrideMainVersion;
 					const string MainVersionKey = "MainVersion";
-					if (OverrideConfigSDKVersions.TryGetValue(MainVersionKey, out OverrideMainVersion))
+					if (overrideConfigSDKVersions.TryGetValue(MainVersionKey, out string? overrideMainVersion))
 					{
 						// if different from default, then remmber it, and mark that it's been overridden
-						if (!OverrideMainVersion.Equals(DefaultConfigSDKVersions[MainVersionKey], StringComparison.OrdinalIgnoreCase))
+						if (!overrideMainVersion.Equals(defaultConfigSDKVersions[MainVersionKey], StringComparison.OrdinalIgnoreCase))
 						{
 							// now check if it was already overridden, in which case we have a conflict we can't resolve, so error
-							if (ProjectsThatOverrodeSDK.Count > 0 && !OverrideMainVersion.Equals(ConfigSDKVersions[MainVersionKey], StringComparison.OrdinalIgnoreCase))
+							if (ProjectsThatOverrodeSDK.Count > 0 && !overrideMainVersion.Equals(ConfigSDKVersions[MainVersionKey], StringComparison.OrdinalIgnoreCase))
 							{
-								throw new Exception($"Project {ProjectFile.GetFileNameWithoutAnyExtensions()} wants to override {Platform} SDK to {OverrideMainVersion}, but it was already overridden to version {ConfigSDKVersions[MainVersionKey]}");
+								throw new Exception($"Project {projectFile.GetFileNameWithoutAnyExtensions()} wants to override {platform} SDK to {overrideMainVersion}, but it was already overridden to version {ConfigSDKVersions[MainVersionKey]}");
 							}
 
-							Logger.LogWarning("Project {Project} is overriding {Platform} SDK to {OverrideMainVersion}", ProjectFile.GetFileNameWithoutAnyExtensions(), Platform, OverrideMainVersion);
+							Logger.LogWarning("Project {Project} is overriding {Platform} SDK to {OverrideMainVersion}", projectFile.GetFileNameWithoutAnyExtensions(), platform, overrideMainVersion);
 
-							ConfigSDKVersions[MainVersionKey] = OverrideMainVersion;
+							ConfigSDKVersions[MainVersionKey] = overrideMainVersion;
 							UEBuildPlatformSDK.bHasAnySDKOverride = true;
 							bHasSDKOverride = true;
-							ProjectsThatOverrodeSDK.Add(ProjectFile);
+							ProjectsThatOverrodeSDK.Add(projectFile);
 						}
 					}
 				}
 			}
 		}
 
-		private void ProcessJsonFile(FileReference SDKConfigFile, Dictionary<string, string> VersionMap, Dictionary<string, string[]> VersionArrayMap)
+		private void ProcessJsonFile(FileReference sdkConfigFile, Dictionary<string, string> versionMap, Dictionary<string, string[]> versionArrayMap)
 		{
 			// load the json
-			if (!JsonObject.TryRead(SDKConfigFile, out JsonObject? jsonObject))
+			if (!JsonObject.TryRead(sdkConfigFile, out JsonObject? jsonObject))
 			{
-				throw new Exception($"Failed to parse SDK version file '{SDKConfigFile}'");
+				throw new Exception($"Failed to parse SDK version file '{sdkConfigFile}'");
 			}
 
-			LoadedConfigFiles.Add(SDKConfigFile);
+			LoadedConfigFiles.Add(sdkConfigFile);
 
 			// look for special key to chain to other - after processing everything else (we only add from the "parent" settings not already in the map)
-			List<string> Parents = [];
-			foreach (string Key in jsonObject.KeyNames)
+			List<string> parents = [];
+			foreach (string key in jsonObject.KeyNames)
 			{
-				if (jsonObject.TryGetStringField(Key, out string? StringValue))
+				if (jsonObject.TryGetStringField(key, out string? stringValue))
 				{
-					if (Key.Equals("ParentSDKFile", StringComparison.OrdinalIgnoreCase))
+					if (key.Equals("ParentSDKFile", StringComparison.OrdinalIgnoreCase))
 					{
-						Parents.Add(StringValue);
+						parents.Add(stringValue);
 					}
 					// add this key if it's not already in the (case-insensitive) dictionary 
-					VersionMap.TryAdd(Key, StringValue);
+					versionMap.TryAdd(key, stringValue);
 				}
-				else if (jsonObject.TryGetStringArrayField(Key, out string[]? ArrayValue))
+				else if (jsonObject.TryGetStringArrayField(key, out string[]? ArrayValue))
 				{
-					VersionArrayMap.Add(Key, ArrayValue);
+					versionArrayMap.Add(key, ArrayValue);
 				}
 			}
 
 			// now load parents, filling in unset properties
-			foreach (string Parent in Parents)
+			foreach (string parent in parents)
 			{
-				FileReference ParentConfigFile = FileReference.Combine(SDKConfigFile.Directory, Parent);
-				ProcessJsonFile(ParentConfigFile, VersionMap, VersionArrayMap);
+				FileReference parentConfigFile = FileReference.Combine(sdkConfigFile.Directory, parent);
+				ProcessJsonFile(parentConfigFile, versionMap, versionArrayMap);
 			}
 		}
 
-		private string GetHostSpecificVersionName(string VersionName)
+		private static string GetHostSpecificVersionName(string versionName)
 		{
-			string HostPlatform = "Win64";
+			string hostPlatform = "Win64";
 			if (OperatingSystem.IsMacOS())
 			{
-				HostPlatform = "Mac";
+				hostPlatform = "Mac";
 			}
 			else if (OperatingSystem.IsLinux())
 			{
-				HostPlatform = "Linux";
+				hostPlatform = "Linux";
 			}
-			return $"{VersionName}_{HostPlatform}";
+			return $"{versionName}_{hostPlatform}";
 		}
-		public string GetRequiredVersionFromConfig(string VersionName)
+		public string GetRequiredVersionFromConfig(string versionName)
 		{
 			// when bIsRequired is true, then we know it will return non-null
-			return GetVersionFromConfig(VersionName, bIsRequired: true)!;
+			return GetVersionFromConfig(versionName, bIsRequired: true)!;
 		}
 
-		public string? GetVersionFromConfig(string VersionName, bool bIsRequired = false)
+		public string? GetVersionFromConfig(string versionName, bool bIsRequired = false)
 		{
-			string? Version;
 			// look up both Version_Host and Version (Host specific version wins)
-			if (ConfigSDKVersions!.TryGetValue(GetHostSpecificVersionName(VersionName), out Version) || ConfigSDKVersions.TryGetValue(VersionName, out Version))
+			string? version;
+			if (ConfigSDKVersions!.TryGetValue(GetHostSpecificVersionName(versionName), out version) || ConfigSDKVersions.TryGetValue(versionName, out version))
 			{
-				return Version;
+				return version;
 			}
 
 			if (bIsRequired)
 			{
-				throw new Exception($"Unable to find required SDK version '{VersionName}' for platform {PlatformName}. Check your SDK.json files");
+				throw new Exception($"Unable to find required SDK version '{versionName}' for platform {PlatformName}. Check your SDK.json files");
 			}
 			return null;
 		}
 
-		protected VersionNumber GetRequiredVersionNumberFromConfig(string VersionName)
+		protected VersionNumber GetRequiredVersionNumberFromConfig(string versionName)
 		{
 			// required won't ever return null
-			return GetVersionNumberFromConfig(VersionName, true)!;
+			return GetVersionNumberFromConfig(versionName, true)!;
 		}
 
-		public VersionNumber? GetVersionNumberFromConfig(string VersionName, bool bIsRequired = false)
+		public VersionNumber? GetVersionNumberFromConfig(string versionName, bool bIsRequired = false)
 		{
-			string? VersionString = GetVersionFromConfig(VersionName, bIsRequired);
+			string? versionString = GetVersionFromConfig(versionName, bIsRequired);
 
-			if (VersionString == null)
+			if (versionString == null)
 			{
 				return null;
 			}
 
-			return VersionNumber.Parse(VersionString);
+			return VersionNumber.Parse(versionString);
 		}
 
-		private VersionNumberRange? ParseVersionNumberRange(string Range)
+		private static VersionNumberRange? ParseVersionNumberRange(string range)
 		{
-			string[] Versions = Range.Split("-");
-			if (Versions.Length != 2)
+			string[] versions = range.Split("-");
+			if (versions.Length != 2)
 			{
 				return null;
 			}
 
-			return VersionNumberRange.Parse(Versions[0], Versions[1]);
+			return VersionNumberRange.Parse(versions[0], versions[1]);
 		}
 
-		public VersionNumberRange? GetRequiredVersionNumberRangeFromConfig(string VersionName, bool bIsRequired = false)
+		public VersionNumberRange? GetRequiredVersionNumberRangeFromConfig(string versionName, bool bIsRequired = false)
 		{
 			// required won't ever return null
-			return GetVersionNumberRangeFromConfig(VersionName, true)!;
+			return GetVersionNumberRangeFromConfig(versionName, true)!;
 		}
 
-		public VersionNumberRange? GetVersionNumberRangeFromConfig(string VersionName, bool bIsRequired = false)
+		public VersionNumberRange? GetVersionNumberRangeFromConfig(string versionName, bool bIsRequired = false)
 		{
-			string? VersionRange;
-			if (!ConfigSDKVersions!.TryGetValue(GetHostSpecificVersionName(VersionName), out VersionRange) && !ConfigSDKVersions.TryGetValue(VersionName, out VersionRange))
+			string? versionRange;
+			if (!ConfigSDKVersions!.TryGetValue(GetHostSpecificVersionName(versionName), out versionRange) && !ConfigSDKVersions.TryGetValue(versionName, out versionRange))
 			{
 				if (bIsRequired)
 				{
-					throw new Exception($"Unable to find required SDK version range '{VersionName}' for platform {PlatformName}. Check your SDK.json files");
+					throw new Exception($"Unable to find required SDK version range '{versionName}' for platform {PlatformName}. Check your SDK.json files");
 				}
 				return null;
 			}
 
-			VersionNumberRange? Range = ParseVersionNumberRange(VersionRange);
-			if (Range == null && bIsRequired)
+			VersionNumberRange? range = ParseVersionNumberRange(versionRange);
+			if (range == null && bIsRequired)
 			{
-				throw new Exception($"Unable to parse the version number range for required version '{VersionName}' for platform {PlatformName}. Check your SDK.json files");
+				throw new Exception($"Unable to parse the version number range for required version '{versionName}' for platform {PlatformName}. Check your SDK.json files");
 			}
-			return Range;
+			return range;
 		}
 
-		public VersionNumberRange[] GetVersionNumberRangeArrayFromConfig(string VersionName)
+		public VersionNumberRange[] GetVersionNumberRangeArrayFromConfig(string versionName)
 		{
-			string[]? VersionRanges;
-			List<VersionNumberRange> Ranges = [];
-			if (ConfigSDKVersionArrays.TryGetValue(GetHostSpecificVersionName(VersionName), out VersionRanges) || ConfigSDKVersionArrays.TryGetValue(VersionName, out VersionRanges))
+			string[]? versionRanges;
+			List<VersionNumberRange> ranges = [];
+			if (ConfigSDKVersionArrays.TryGetValue(GetHostSpecificVersionName(versionName), out versionRanges) || ConfigSDKVersionArrays.TryGetValue(versionName, out versionRanges))
 			{
-				foreach (string VersionRange in VersionRanges)
+				foreach (string versionRange in versionRanges)
 				{
-					VersionNumberRange? Range = ParseVersionNumberRange(VersionRange);
-					if (Range != null)
+					VersionNumberRange? range = ParseVersionNumberRange(versionRange);
+					if (range != null)
 					{
-						Ranges.Add(Range);
+						ranges.Add(range);
 					}
 				}
 			}
 
-			return [.. Ranges];
+			return [.. ranges];
 		}
 
-		public string[] GetStringArrayFromConfig(string Name)
+		public string[] GetStringArrayFromConfig(string name)
 		{
-			string[]? Results;
-			if (ConfigSDKVersionArrays.TryGetValue(GetHostSpecificVersionName(Name), out Results) || ConfigSDKVersionArrays.TryGetValue(Name, out Results))
+			string[]? results;
+			if (ConfigSDKVersionArrays.TryGetValue(GetHostSpecificVersionName(name), out results) || ConfigSDKVersionArrays.TryGetValue(name, out results))
 			{
-				return Results;
+				return results;
 			}
 
 			return [];
@@ -1348,8 +1332,8 @@ namespace EpicGames.Core
 		{
 			if (!bCheckedAutoSDKRootEnvVar)
 			{
-				string? SDKRoot = Environment.GetEnvironmentVariable(SDKRootEnvVar);
-				if (SDKRoot != null)
+				string? sdkRoot = Environment.GetEnvironmentVariable(SDKRootEnvVar);
+				if (sdkRoot != null)
 				{
 					bAutoSDKSystemEnabled = true;
 				}
@@ -1373,15 +1357,10 @@ namespace EpicGames.Core
 		{
 			const string UnspecifiedVersion = "UnspecifiedScriptVersion";
 
-			string VersionFilename = Path.Combine(GetPathToPlatformAutoSDKs(), GetAutoSDKDirectoryForMainVersion(), ScriptVersionFilename);
-			if (File.Exists(VersionFilename))
-			{
-				return File.ReadAllLines(VersionFilename).FirstOrDefault() ?? UnspecifiedVersion;
-			}
-			else
-			{
-				return UnspecifiedVersion;
-			}
+			string versionFilename = Path.Combine(GetPathToPlatformAutoSDKs(), GetAutoSDKDirectoryForMainVersion(), ScriptVersionFilename);
+			return File.Exists(versionFilename)
+				? File.ReadAllLines(versionFilename).FirstOrDefault() ?? UnspecifiedVersion
+				: UnspecifiedVersion;
 		}
 
 		/// <summary>
@@ -1390,33 +1369,33 @@ namespace EpicGames.Core
 		/// <returns>Valid SDK string</returns>
 		protected string GetPathToPlatformAutoSDKs()
 		{
-			string SDKPath = "";
-			string? SDKRoot = Environment.GetEnvironmentVariable(SDKRootEnvVar);
-			if (SDKRoot != null)
+			string sdkPath = "";
+			string? sdkRoot = Environment.GetEnvironmentVariable(SDKRootEnvVar);
+			if (sdkRoot != null)
 			{
-				if (!String.IsNullOrEmpty(SDKRoot))
+				if (!String.IsNullOrEmpty(sdkRoot))
 				{
-					SDKPath = Path.Combine(SDKRoot, "Host" + GetAutoSDKHostPlatform(), GetAutoSDKPlatformName());
+					sdkPath = Path.Combine(sdkRoot, "Host" + GetAutoSDKHostPlatform(), GetAutoSDKPlatformName());
 				}
 			}
-			return SDKPath;
+			return sdkPath;
 		}
 
 		/// <summary>
 		/// Returns path to platform SDKs
 		/// </summary>
 		/// <returns>Valid SDK string</returns>
-		public static bool TryGetHostPlatformAutoSDKDir([NotNullWhen(true)] out DirectoryReference? OutPlatformDir)
+		public static bool TryGetHostPlatformAutoSDKDir([NotNullWhen(true)] out DirectoryReference? outPlatformDir)
 		{
-			string? SDKRoot = Environment.GetEnvironmentVariable(SDKRootEnvVar);
-			if (String.IsNullOrEmpty(SDKRoot))
+			string? sdkRoot = Environment.GetEnvironmentVariable(SDKRootEnvVar);
+			if (String.IsNullOrEmpty(sdkRoot))
 			{
-				OutPlatformDir = null;
+				outPlatformDir = null;
 				return false;
 			}
 			else
 			{
-				OutPlatformDir = DirectoryReference.Combine(new DirectoryReference(SDKRoot), "Host" + GetAutoSDKHostPlatform());
+				outPlatformDir = DirectoryReference.Combine(new DirectoryReference(sdkRoot), "Host" + GetAutoSDKHostPlatform());
 				return true;
 			}
 		}
@@ -1439,34 +1418,34 @@ namespace EpicGames.Core
 		/// <summary>
 		/// Gets currently installed version
 		/// </summary>
-		/// <param name="PlatformSDKRoot">absolute path to platform SDK root</param>
-		/// <param name="OutInstalledSDKVersionString">version string as currently installed</param>
-		/// <param name="OutInstalledSDKLevel"></param>
+		/// <param name="platformSDKRoot">absolute path to platform SDK root</param>
+		/// <param name="outInstalledSDKVersionString">version string as currently installed</param>
+		/// <param name="outInstalledSDKLevel"></param>
 		/// <returns>true if was able to read it</returns>
-		protected bool GetCurrentlyInstalledSDKString(string PlatformSDKRoot, out string OutInstalledSDKVersionString, out string OutInstalledSDKLevel)
+		protected static bool GetCurrentlyInstalledSDKString(string platformSDKRoot, out string outInstalledSDKVersionString, out string outInstalledSDKLevel)
 		{
-			if (Directory.Exists(PlatformSDKRoot))
+			if (Directory.Exists(platformSDKRoot))
 			{
-				string VersionFilename = Path.Combine(PlatformSDKRoot, CurrentlyInstalledSDKStringManifest);
-				if (File.Exists(VersionFilename))
+				string versionFilename = Path.Combine(platformSDKRoot, CurrentlyInstalledSDKStringManifest);
+				if (File.Exists(versionFilename))
 				{
-					using (StreamReader Reader = new StreamReader(VersionFilename))
+					using (StreamReader reader = new StreamReader(versionFilename))
 					{
-						string? Version = Reader.ReadLine();
-						string? Type = Reader.ReadLine();
-						string? Level = Reader.ReadLine();
-						if (String.IsNullOrEmpty(Level))
+						string? version = reader.ReadLine();
+						string? type = reader.ReadLine();
+						string? level = reader.ReadLine();
+						if (String.IsNullOrEmpty(level))
 						{
-							Level = "FULL";
+							level = "FULL";
 						}
 
 						// don't allow ManualSDK installs to count as an AutoSDK install version.
-						if (Type != null && Type == "AutoSDK")
+						if (type != null && type == "AutoSDK")
 						{
-							if (Version != null)
+							if (version != null)
 							{
-								OutInstalledSDKVersionString = Version;
-								OutInstalledSDKLevel = Level;
+								outInstalledSDKVersionString = version;
+								outInstalledSDKLevel = level;
 								return true;
 							}
 						}
@@ -1474,62 +1453,62 @@ namespace EpicGames.Core
 				}
 			}
 
-			OutInstalledSDKVersionString = "";
-			OutInstalledSDKLevel = "";
+			outInstalledSDKVersionString = "";
+			outInstalledSDKLevel = "";
 			return false;
 		}
 
 		/// <summary>
 		/// Gets the version of the last successfully run setup script.
 		/// </summary>
-		/// <param name="PlatformSDKRoot">absolute path to platform SDK root</param>
-		/// <param name="OutLastRunScriptVersion">version string</param>
+		/// <param name="platformSDKRoot">absolute path to platform SDK root</param>
+		/// <param name="outLastRunScriptVersion">version string</param>
 		/// <returns>true if was able to read it</returns>
-		protected bool GetLastRunScriptVersionString(string PlatformSDKRoot, out string OutLastRunScriptVersion)
+		protected static bool GetLastRunScriptVersionString(string platformSDKRoot, out string outLastRunScriptVersion)
 		{
-			if (Directory.Exists(PlatformSDKRoot))
+			if (Directory.Exists(platformSDKRoot))
 			{
-				string VersionFilename = Path.Combine(PlatformSDKRoot, LastRunScriptVersionManifest);
-				if (File.Exists(VersionFilename))
+				string versionFilename = Path.Combine(platformSDKRoot, LastRunScriptVersionManifest);
+				if (File.Exists(versionFilename))
 				{
-					using (StreamReader Reader = new StreamReader(VersionFilename))
+					using (StreamReader reader = new StreamReader(versionFilename))
 					{
-						string? Version = Reader.ReadLine();
-						if (Version != null)
+						string? version = reader.ReadLine();
+						if (version != null)
 						{
-							OutLastRunScriptVersion = Version;
+							outLastRunScriptVersion = version;
 							return true;
 						}
 					}
 				}
 			}
 
-			OutLastRunScriptVersion = "";
+			outLastRunScriptVersion = "";
 			return false;
 		}
 
 		/// <summary>
 		/// Sets currently installed version
 		/// </summary>
-		/// <param name="InstalledSDKVersionString">SDK version string to set</param>
-		/// <param name="InstalledSDKLevelString"></param>
+		/// <param name="installedSDKVersionString">SDK version string to set</param>
+		/// <param name="installedSDKLevelString"></param>
 		/// <returns>true if was able to set it</returns>
-		protected bool SetCurrentlyInstalledAutoSDKString(string InstalledSDKVersionString, string InstalledSDKLevelString)
+		protected bool SetCurrentlyInstalledAutoSDKString(string installedSDKVersionString, string installedSDKLevelString)
 		{
-			string PlatformSDKRoot = GetPathToPlatformAutoSDKs();
-			if (Directory.Exists(PlatformSDKRoot))
+			string platformSDKRoot = GetPathToPlatformAutoSDKs();
+			if (Directory.Exists(platformSDKRoot))
 			{
-				string VersionFilename = Path.Combine(PlatformSDKRoot, CurrentlyInstalledSDKStringManifest);
-				if (File.Exists(VersionFilename))
+				string versionFilename = Path.Combine(platformSDKRoot, CurrentlyInstalledSDKStringManifest);
+				if (File.Exists(versionFilename))
 				{
-					File.Delete(VersionFilename);
+					File.Delete(versionFilename);
 				}
 
-				using (StreamWriter Writer = File.CreateText(VersionFilename))
+				using (StreamWriter writer = File.CreateText(versionFilename))
 				{
-					Writer.WriteLine(InstalledSDKVersionString);
-					Writer.WriteLine("AutoSDK");
-					Writer.WriteLine(InstalledSDKLevelString);
+					writer.WriteLine(installedSDKVersionString);
+					writer.WriteLine("AutoSDK");
+					writer.WriteLine(installedSDKLevelString);
 					return true;
 				}
 			}
@@ -1541,31 +1520,31 @@ namespace EpicGames.Core
 		{
 			if (PlatformSupportsAutoSDKs() && HasAutoSDKSystemEnabled())
 			{
-				string InstalledSDKVersionString = GetAutoSDKDirectoryForMainVersion();
-				string PlatformSDKRoot = GetPathToPlatformAutoSDKs();
-				if (!Directory.Exists(PlatformSDKRoot))
+				string installedSDKVersionString = GetAutoSDKDirectoryForMainVersion();
+				string platformSDKRoot = GetPathToPlatformAutoSDKs();
+				if (!Directory.Exists(platformSDKRoot))
 				{
-					Directory.CreateDirectory(PlatformSDKRoot);
+					Directory.CreateDirectory(platformSDKRoot);
 				}
 
 				{
-					string VersionFilename = Path.Combine(PlatformSDKRoot, CurrentlyInstalledSDKStringManifest);
-					if (File.Exists(VersionFilename))
+					string versionFilename = Path.Combine(platformSDKRoot, CurrentlyInstalledSDKStringManifest);
+					if (File.Exists(versionFilename))
 					{
-						File.Delete(VersionFilename);
+						File.Delete(versionFilename);
 					}
 
-					string EnvVarFile = Path.Combine(PlatformSDKRoot, SDKEnvironmentVarsFile);
-					if (File.Exists(EnvVarFile))
+					string envVarFile = Path.Combine(platformSDKRoot, SDKEnvironmentVarsFile);
+					if (File.Exists(envVarFile))
 					{
-						File.Delete(EnvVarFile);
+						File.Delete(envVarFile);
 					}
 
-					using (StreamWriter Writer = File.CreateText(VersionFilename))
+					using (StreamWriter writer = File.CreateText(versionFilename))
 					{
-						Writer.WriteLine(InstalledSDKVersionString);
-						Writer.WriteLine("ManualSDK");
-						Writer.WriteLine("FULL");
+						writer.WriteLine(installedSDKVersionString);
+						writer.WriteLine("ManualSDK");
+						writer.WriteLine("FULL");
 					}
 				}
 			}
@@ -1573,89 +1552,89 @@ namespace EpicGames.Core
 
 		public static void ClearManualSDKEnvVarCache()
 		{
-			string ManualSDKEnvironmentVarsPath = Path.Combine(Unreal.EngineDirectory.ToString(), "Intermediate", ManualSDKEnvironmentVarsFile);
-			if (File.Exists(ManualSDKEnvironmentVarsPath))
+			string manualSDKEnvironmentVarsPath = Path.Combine(Unreal.EngineDirectory.ToString(), "Intermediate", ManualSDKEnvironmentVarsFile);
+			if (File.Exists(manualSDKEnvironmentVarsPath))
 			{
-				File.Delete(ManualSDKEnvironmentVarsPath);
+				File.Delete(manualSDKEnvironmentVarsPath);
 			}
 		}
 
-		private bool TrySelectBestManualSDK(out string? SelectedSDKVersion)
+		private bool TrySelectBestManualSDK(out string? selectedSDKVersion)
 		{
 			// find the best valid manual sdk
-			SelectedSDKVersion = FindBestInstalledSDKVersion();
-			if (SelectedSDKVersion == null)
+			selectedSDKVersion = FindBestInstalledSDKVersion();
+			if (selectedSDKVersion == null)
 			{
 				return false;
 			}
 
 			// query the platform for environment changes
-			if (!TryGetEnvironmentForManualSDK(SelectedSDKVersion, out Dictionary<string, string>? EnvVarValues, out List<string>? PathAdds, out List<string>? PathRemoves))
+			if (!TryGetEnvironmentForManualSDK(selectedSDKVersion, out Dictionary<string, string>? envVarValues, out List<string>? pathAdds, out List<string>? pathRemoves))
 			{
 				return false;
 			}
 
 			// apply environment variables
-			if (EnvVarValues != null)
+			if (envVarValues != null)
 			{
-				foreach (KeyValuePair<string,string> EnvVarValue in EnvVarValues)
+				foreach (KeyValuePair<string, string> envVarValue in envVarValues)
 				{
-					Environment.SetEnvironmentVariable(EnvVarValue.Key, EnvVarValue.Value);
+					Environment.SetEnvironmentVariable(envVarValue.Key, envVarValue.Value);
 				}
 			}
 
 			// apply PATH modifications
-			if (PathRemoves != null || PathAdds != null)
+			if (pathRemoves != null || pathAdds != null)
 			{
-				string OrigPathVar = Environment.GetEnvironmentVariable("PATH") ?? String.Empty;
-				IEnumerable<string> PathVars = OrigPathVar.Split(Path.PathSeparator);
-				if (PathRemoves != null)
+				string origPathVar = Environment.GetEnvironmentVariable("PATH") ?? String.Empty;
+				IEnumerable<string> pathVars = origPathVar.Split(Path.PathSeparator);
+				if (pathRemoves != null)
 				{
-					PathVars = PathVars.Except(PathRemoves, FileUtils.PlatformPathComparer);
+					pathVars = pathVars.Except(pathRemoves, FileUtils.PlatformPathComparer);
 				}
-				if (PathAdds != null)
+				if (pathAdds != null)
 				{
-					PathVars = PathVars.Except(PathAdds, FileUtils.PlatformPathComparer); // remove all of the ADDs so that if this function is executed multiple times, the paths will be guaranteed to be in the same order after each run.
-					PathVars = PathVars.Union(PathAdds, FileUtils.PlatformPathComparer);
+					pathVars = pathVars.Except(pathAdds, FileUtils.PlatformPathComparer); // remove all of the ADDs so that if this function is executed multiple times, the paths will be guaranteed to be in the same order after each run.
+					pathVars = pathVars.Union(pathAdds, FileUtils.PlatformPathComparer);
 				}
 
-				string NewPathVar = String.Join( Path.PathSeparator, PathVars );
-				Environment.SetEnvironmentVariable("PATH", NewPathVar);
+				string newPathVar = String.Join(Path.PathSeparator, pathVars);
+				Environment.SetEnvironmentVariable("PATH", newPathVar);
 			}
 
 			// write all environment modifications
-			if ( (EnvVarValues != null && EnvVarValues.Count > 0) || (PathRemoves != null && PathRemoves.Count > 0) || (PathAdds != null && PathAdds.Count > 0) )
+			if ((envVarValues != null && envVarValues.Count > 0) || (pathRemoves != null && pathRemoves.Count > 0) || (pathAdds != null && pathAdds.Count > 0))
 			{
-				string ManualSDKEnvironmentVarsPath = Path.Combine(Unreal.EngineDirectory.ToString(), "Intermediate", ManualSDKEnvironmentVarsFile);
+				string manualSDKEnvironmentVarsPath = Path.Combine(Unreal.EngineDirectory.ToString(), "Intermediate", ManualSDKEnvironmentVarsFile);
 
-				Directory.CreateDirectory(Path.GetDirectoryName(ManualSDKEnvironmentVarsPath)!);
-				using (StreamWriter Writer = File.AppendText(ManualSDKEnvironmentVarsPath))
+				Directory.CreateDirectory(Path.GetDirectoryName(manualSDKEnvironmentVarsPath)!);
+				using (StreamWriter writer = File.AppendText(manualSDKEnvironmentVarsPath))
 				{
 					// write environment variables
-					if (EnvVarValues != null)
+					if (envVarValues != null)
 					{
-						foreach (KeyValuePair<string,string> EnvVarValue in EnvVarValues)
+						foreach (KeyValuePair<string, string> envVarValue in envVarValues)
 						{
-							Writer.WriteLine($"{EnvVarValue.Key}={EnvVarValue.Value}");
-							Logger.LogDebug("Setting variable '{Name}' to '{Value}'", EnvVarValue.Key, EnvVarValue.Value);
+							writer.WriteLine($"{envVarValue.Key}={envVarValue.Value}");
+							Logger.LogDebug("Setting variable '{Name}' to '{Value}'", envVarValue.Key, envVarValue.Value);
 						}
 					}
 
 					// write PATH modifications
-					if (PathRemoves != null)
+					if (pathRemoves != null)
 					{
-						foreach (string PathVar in PathRemoves)
+						foreach (string pathVar in pathRemoves)
 						{
-							Writer.WriteLine($"strippath={PathVar}");
-							Logger.LogDebug("Removing Path: '{Path}'", PathVar);
+							writer.WriteLine($"strippath={pathVar}");
+							Logger.LogDebug("Removing Path: '{Path}'", pathVar);
 						}
 					}
-					if (PathAdds != null)
+					if (pathAdds != null)
 					{
-						foreach (string PathVar in PathAdds)
+						foreach (string pathVar in pathAdds)
 						{
-							Writer.WriteLine($"addpath={PathVar}");
-							Logger.LogDebug("Adding Path: '{Path}'", PathVar);
+							writer.WriteLine($"addpath={pathVar}");
+							Logger.LogDebug("Adding Path: '{Path}'", pathVar);
 						}
 					}
 				}
@@ -1664,20 +1643,20 @@ namespace EpicGames.Core
 			return true;
 		}
 
-		protected bool SetLastRunAutoSDKScriptVersion(string LastRunScriptVersion)
+		protected bool SetLastRunAutoSDKScriptVersion(string lastRunScriptVersion)
 		{
-			string PlatformSDKRoot = GetPathToPlatformAutoSDKs();
-			if (Directory.Exists(PlatformSDKRoot))
+			string platformSDKRoot = GetPathToPlatformAutoSDKs();
+			if (Directory.Exists(platformSDKRoot))
 			{
-				string VersionFilename = Path.Combine(PlatformSDKRoot, LastRunScriptVersionManifest);
-				if (File.Exists(VersionFilename))
+				string versionFilename = Path.Combine(platformSDKRoot, LastRunScriptVersionManifest);
+				if (File.Exists(versionFilename))
 				{
-					File.Delete(VersionFilename);
+					File.Delete(versionFilename);
 				}
 
-				using (StreamWriter Writer = File.CreateText(VersionFilename))
+				using (StreamWriter writer = File.CreateText(versionFilename))
 				{
-					Writer.WriteLine(LastRunScriptVersion);
+					writer.WriteLine(lastRunScriptVersion);
 					return true;
 				}
 			}
@@ -1688,49 +1667,30 @@ namespace EpicGames.Core
 		/// Returns Hook names as needed by the platform
 		/// (e.g. can be overridden with custom executables or scripts)
 		/// </summary>
-		/// <param name="Hook">Hook type</param>
-		protected virtual string GetHookExecutableName(SDKHookType Hook)
+		/// <param name="hook">Hook type</param>
+		protected virtual string GetHookExecutableName(SDKHookType hook)
 		{
-			if (OperatingSystem.IsWindows())
-			{
-				if (Hook == SDKHookType.Uninstall)
-				{
-					return "unsetup.bat";
-				}
-				else
-				{
-					return "setup.bat";
-				}
-			}
-			else
-			{
-				if (Hook == SDKHookType.Uninstall)
-				{
-					return "unsetup.sh";
-				}
-				else
-				{
-					return "setup.sh";
-				}
-			}
+			return OperatingSystem.IsWindows()
+				? hook == SDKHookType.Uninstall ? "unsetup.bat" : "setup.bat"
+				: hook == SDKHookType.Uninstall ? "unsetup.sh" : "setup.sh";
 		}
 
 		/// <summary>
 		/// Whether the hook must be run with administrator privileges.
 		/// </summary>
-		/// <param name="Hook">Hook for which to check the required privileges.</param>
+		/// <param name="hook">Hook for which to check the required privileges.</param>
 		/// <returns>true if the hook must be run with administrator privileges.</returns>
-		protected virtual bool DoesHookRequireAdmin(SDKHookType Hook)
+		protected virtual bool DoesHookRequireAdmin(SDKHookType hook)
 		{
 			return true;
 		}
 
-		private void LogAutoSDKHook(object Sender, DataReceivedEventArgs Args)
+		private void LogAutoSDKHook(object sender, DataReceivedEventArgs args)
 		{
-			if (Args.Data != null)
+			if (args.Data != null)
 			{
-				LogFormatOptions Options = Log.OutputLevel >= LogEventType.Verbose ? LogFormatOptions.None : LogFormatOptions.NoConsoleOutput;
-				Log.WriteLine(LogEventType.Log, Options, Args.Data);
+				LogFormatOptions options = Log.OutputLevel >= LogEventType.Verbose ? LogFormatOptions.None : LogFormatOptions.NoConsoleOutput;
+				Log.WriteLine(LogEventType.Log, options, args.Data);
 			}
 		}
 
@@ -1743,120 +1703,120 @@ namespace EpicGames.Core
 			"FULL",
 		];
 
-		private bool IsSDKLevelAtLeast(string Value, string ComparedTo)
+		private static bool IsSDKLevelAtLeast(string value, string comparedTo)
 		{
-			int ValueIndex = AutoSDKLevels.FindIndex(x => x.Equals(Value, StringComparison.OrdinalIgnoreCase));
-			int ComparedIndex = AutoSDKLevels.FindIndex(x => x.Equals(ComparedTo, StringComparison.OrdinalIgnoreCase));
-			if (ValueIndex < 0 || ComparedIndex < 0)
+			int valueIndex = AutoSDKLevels.FindIndex(x => x.Equals(value, StringComparison.OrdinalIgnoreCase));
+			int comparedIndex = AutoSDKLevels.FindIndex(x => x.Equals(comparedTo, StringComparison.OrdinalIgnoreCase));
+			if (valueIndex < 0 || comparedIndex < 0)
 			{
-				throw new Exception($"Passed in a bad value to IsSDKLevelAtLeast: {Value}, {ComparedTo}");
+				throw new Exception($"Passed in a bad value to IsSDKLevelAtLeast: {value}, {comparedTo}");
 			}
-			return ValueIndex >= ComparedIndex;
+			return valueIndex >= comparedIndex;
 		}
 
-		private string GetAutoSDKLevelForPlatform(string PlatformSDKRoot)
+		private static string GetAutoSDKLevelForPlatform(string platformSDKRoot)
 		{
 			// the last component is the platform name
-			string PlatformName = Path.GetFileName(PlatformSDKRoot).ToUpper();
+			string platformName = Path.GetFileName(platformSDKRoot).ToUpper();
 
 			// parse the envvar
-			Dictionary<string, string> PlatformSpecificLevels = [];
-			string? DetailedSettings = Environment.GetEnvironmentVariable("UE_AUTOSDK_SPECIFIC_LEVELS"); // "Android=PACKAGE;GDk=RuN;PS5=BUILD";
-			if (!String.IsNullOrEmpty(DetailedSettings) && DetailedSettings.Contains(PlatformName, StringComparison.InvariantCultureIgnoreCase))
+			Dictionary<string, string> platformSpecificLevels = [];
+			string? detailedSettings = Environment.GetEnvironmentVariable("UE_AUTOSDK_SPECIFIC_LEVELS"); // "Android=PACKAGE;GDk=RuN;PS5=BUILD";
+			if (!String.IsNullOrEmpty(detailedSettings) && detailedSettings.Contains(platformName, StringComparison.InvariantCultureIgnoreCase))
 			{
-				foreach (string Detail in DetailedSettings.ToUpper().Split(';'))
+				foreach (string detail in detailedSettings.ToUpper().Split(';'))
 				{
-					string[] Tokens = Detail.Split('=');
+					string[] tokens = detail.Split('=');
 					// validate the level string
-					if (AutoSDKLevels.Contains(Tokens[1]))
+					if (AutoSDKLevels.Contains(tokens[1]))
 					{
-						PlatformSpecificLevels.Add(Tokens[0], Tokens[1]);
+						platformSpecificLevels.Add(tokens[0], tokens[1]);
 					}
 				}
 			}
 
-			string FinalAutoSDKLevel = "FULL";
-			if (PlatformSpecificLevels.TryGetValue(PlatformName, out string? value))
+			string finalAutoSDKLevel = "FULL";
+			if (platformSpecificLevels.TryGetValue(platformName, out string? value))
 			{
-				FinalAutoSDKLevel = value;
+				finalAutoSDKLevel = value;
 			}
 			else
 			{
-				string? DefaultLevel = Environment.GetEnvironmentVariable("UE_AUTOSDK_DEFAULT_LEVEL");
-				if (!String.IsNullOrEmpty(DefaultLevel) && AutoSDKLevels.Contains(DefaultLevel, StringComparer.InvariantCultureIgnoreCase))
+				string? defaultLevel = Environment.GetEnvironmentVariable("UE_AUTOSDK_DEFAULT_LEVEL");
+				if (!String.IsNullOrEmpty(defaultLevel) && AutoSDKLevels.Contains(defaultLevel, StringComparer.InvariantCultureIgnoreCase))
 				{
-					FinalAutoSDKLevel = DefaultLevel.ToUpper();
+					finalAutoSDKLevel = defaultLevel.ToUpper();
 				}
 			}
 
-			return FinalAutoSDKLevel;
+			return finalAutoSDKLevel;
 		}
 
 		/// <summary>
 		/// Runs install/uninstall hooks for SDK
 		/// </summary>
-		/// <param name="PlatformSDKRoot">absolute path to platform SDK root</param>
-		/// <param name="SDKVersionString">version string to run for (can be empty!)</param>
-		/// <param name="AutoSDKLevel"></param>
-		/// <param name="Hook">which one of hooks to run</param>
+		/// <param name="platformSDKRoot">absolute path to platform SDK root</param>
+		/// <param name="sdkVersionString">version string to run for (can be empty!)</param>
+		/// <param name="autoSDKLevel"></param>
+		/// <param name="hook">which one of hooks to run</param>
 		/// <param name="bHookCanBeNonExistent">whether a non-existing hook means failure</param>
 		/// <returns>true if succeeded</returns>
-		protected virtual bool RunAutoSDKHooks(string PlatformSDKRoot, string SDKVersionString, string AutoSDKLevel, SDKHookType Hook, bool bHookCanBeNonExistent = true)
+		protected virtual bool RunAutoSDKHooks(string platformSDKRoot, string sdkVersionString, string autoSDKLevel, SDKHookType hook, bool bHookCanBeNonExistent = true)
 		{
 			if (!IsAutoSDKSafe())
 			{
 				Logger.LogDebug("{Platform} attempted to run SDK hook which could have damaged manual SDK install!", GetAutoSDKPlatformName());
 				return false;
 			}
-			if (!String.IsNullOrEmpty(SDKVersionString))
+			if (!String.IsNullOrEmpty(sdkVersionString))
 			{
-				string SDKDirectory = Path.Combine(PlatformSDKRoot, SDKVersionString);
-				string HookExe = Path.Combine(SDKDirectory, GetHookExecutableName(Hook));
+				string sdkDirectory = Path.Combine(platformSDKRoot, sdkVersionString);
+				string hookExe = Path.Combine(sdkDirectory, GetHookExecutableName(hook));
 
-				if (File.Exists(HookExe))
+				if (File.Exists(hookExe))
 				{
-					Logger.LogDebug("Running {Hook} hook {HookExe}", Hook, HookExe);
+					Logger.LogDebug("Running {Hook} hook {HookExe}", hook, hookExe);
 					using ITimelineEvent _ = Timeline.ScopeEvent("RunAutoSDKHooks");
 
 					// run it
-					using Process HookProcess = new();
-					HookProcess.StartInfo.WorkingDirectory = SDKDirectory;
-					HookProcess.StartInfo.FileName = HookExe;
-					HookProcess.StartInfo.Arguments = AutoSDKLevel;
-					HookProcess.StartInfo.WindowStyle = ProcessWindowStyle.Hidden;
+					using Process hookProcess = new();
+					hookProcess.StartInfo.WorkingDirectory = sdkDirectory;
+					hookProcess.StartInfo.FileName = hookExe;
+					hookProcess.StartInfo.Arguments = autoSDKLevel;
+					hookProcess.StartInfo.WindowStyle = ProcessWindowStyle.Hidden;
 
-					bool bHookRequiresAdmin = DoesHookRequireAdmin(Hook);
+					bool bHookRequiresAdmin = DoesHookRequireAdmin(hook);
 					if (bHookRequiresAdmin)
 					{
 						// installers may require administrator access to succeed. so run as an admin.
-						HookProcess.StartInfo.Verb = "runas";
+						hookProcess.StartInfo.Verb = "runas";
 
 						//Forcing the old .Net Framework default to prevent processes from failing
-						HookProcess.StartInfo.UseShellExecute = true;
+						hookProcess.StartInfo.UseShellExecute = true;
 					}
 					else
 					{
-						HookProcess.StartInfo.UseShellExecute = false;
-						HookProcess.StartInfo.RedirectStandardOutput = true;
-						HookProcess.StartInfo.RedirectStandardError = true;
-						HookProcess.OutputDataReceived += LogAutoSDKHook;
-						HookProcess.ErrorDataReceived += LogAutoSDKHook;
+						hookProcess.StartInfo.UseShellExecute = false;
+						hookProcess.StartInfo.RedirectStandardOutput = true;
+						hookProcess.StartInfo.RedirectStandardError = true;
+						hookProcess.OutputDataReceived += LogAutoSDKHook;
+						hookProcess.ErrorDataReceived += LogAutoSDKHook;
 					}
 
 					//using (ScopedTimer HookTimer = new ScopedTimer("Time to run hook: ", LogEventType.Log))
 					{
-						HookProcess.Start();
+						hookProcess.Start();
 						if (!bHookRequiresAdmin)
 						{
-							HookProcess.BeginOutputReadLine();
-							HookProcess.BeginErrorReadLine();
+							hookProcess.BeginOutputReadLine();
+							hookProcess.BeginErrorReadLine();
 						}
-						HookProcess.WaitForExit();
+						hookProcess.WaitForExit();
 					}
 
-					if (HookProcess.ExitCode != 0)
+					if (hookProcess.ExitCode != 0)
 					{
-						Logger.LogDebug("Hook exited uncleanly (returned {ExitCode}), considering it failed.", HookProcess.ExitCode);
+						Logger.LogDebug("Hook exited uncleanly (returned {ExitCode}), considering it failed.", hookProcess.ExitCode);
 						return false;
 					}
 
@@ -1864,12 +1824,12 @@ namespace EpicGames.Core
 				}
 				else
 				{
-					Logger.LogDebug("File {HookExe} does not exist", HookExe);
+					Logger.LogDebug("File {HookExe} does not exist", hookExe);
 				}
 			}
 			else
 			{
-				Logger.LogDebug("Version string is blank for {SdkRoot}. Can't determine {Hook} hook.", PlatformSDKRoot, Hook.ToString());
+				Logger.LogDebug("Version string is blank for {SdkRoot}. Can't determine {Hook} hook.", platformSDKRoot, hook.ToString());
 			}
 
 			return bHookCanBeNonExistent;
@@ -1880,153 +1840,153 @@ namespace EpicGames.Core
 		/// If any commands are added or removed the handling needs to be duplicated in
 		/// TargetPlatformManagerModule.cpp
 		/// </summary>
-		/// <param name="PlatformSDKRoot">absolute path to platform SDK</param>
+		/// <param name="platformSDKRoot">absolute path to platform SDK</param>
 		/// <returns>true if succeeded</returns>
-		protected bool SetupEnvironmentFromAutoSDK(string PlatformSDKRoot)
+		protected bool SetupEnvironmentFromAutoSDK(string platformSDKRoot)
 		{
-			string EnvVarFile = Path.Combine(PlatformSDKRoot, SDKEnvironmentVarsFile);
-			if (File.Exists(EnvVarFile))
+			string envVarFile = Path.Combine(platformSDKRoot, SDKEnvironmentVarsFile);
+			if (File.Exists(envVarFile))
 			{
-				using (StreamReader Reader = new StreamReader(EnvVarFile))
+				using (StreamReader reader = new StreamReader(envVarFile))
 				{
-					List<string> PathAdds = [];
-					List<string> PathRemoves = [];
+					List<string> pathAdds = [];
+					List<string> pathRemoves = [];
 
-					List<string> EnvVarNames = [];
-					List<string> EnvVarValues = [];
+					List<string> envVarNames = [];
+					List<string> envVarValues = [];
 
 					bool bNeedsToWriteAutoSetupEnvVar = true;
-					string PlatformSetupEnvVar = GetPlatformAutoSDKSetupEnvVar();
+					string platformSetupEnvVar = GetPlatformAutoSDKSetupEnvVar();
 					for (; ; )
 					{
-						string? VariableString = Reader.ReadLine();
-						if (VariableString == null)
+						string? variableString = reader.ReadLine();
+						if (variableString == null)
 						{
 							break;
 						}
 
-						string[] Parts = VariableString.Split('=');
-						if (Parts.Length != 2)
+						string[] parts = variableString.Split('=');
+						if (parts.Length != 2)
 						{
 							Logger.LogDebug("Incorrect environment variable declaration:");
-							Logger.LogDebug("{VariableString}", VariableString);
+							Logger.LogDebug("{VariableString}", variableString);
 							return false;
 						}
 
-						if (String.Compare(Parts[0], "strippath", true) == 0)
+						if (String.Equals(parts[0], "strippath", StringComparison.OrdinalIgnoreCase))
 						{
-							PathRemoves.Add(Parts[1]);
+							pathRemoves.Add(parts[1]);
 						}
-						else if (String.Compare(Parts[0], "addpath", true) == 0)
+						else if (String.Equals(parts[0], "addpath", StringComparison.OrdinalIgnoreCase))
 						{
-							PathAdds.Add(Parts[1]);
+							pathAdds.Add(parts[1]);
 						}
 						else
 						{
-							if (String.Compare(Parts[0], PlatformSetupEnvVar) == 0)
+							if (String.Equals(parts[0], platformSetupEnvVar, StringComparison.Ordinal))
 							{
 								bNeedsToWriteAutoSetupEnvVar = false;
 							}
 							// convenience for setup.bat writers.  Trim any accidental whitespace from variable names/values.
-							EnvVarNames.Add(Parts[0].Trim());
-							EnvVarValues.Add(Parts[1].Trim());
+							envVarNames.Add(parts[0].Trim());
+							envVarValues.Add(parts[1].Trim());
 						}
 					}
 
 					// don't actually set anything until we successfully validate and read all values in.
 					// we don't want to set a few vars, return a failure, and then have a platform try to
 					// build against a manually installed SDK with half-set env vars.
-					for (int i = 0; i < EnvVarNames.Count; ++i)
+					for (int i = 0; i < envVarNames.Count; ++i)
 					{
-						string EnvVarName = EnvVarNames[i];
-						string EnvVarValue = EnvVarValues[i];
-						Logger.LogDebug("Setting variable '{Name}' to '{Value}'", EnvVarName, EnvVarValue);
-						Environment.SetEnvironmentVariable(EnvVarName, EnvVarValue);
+						string envVarName = envVarNames[i];
+						string envVarValue = envVarValues[i];
+						Logger.LogDebug("Setting variable '{Name}' to '{Value}'", envVarName, envVarValue);
+						Environment.SetEnvironmentVariable(envVarName, envVarValue);
 					}
 
 					// actually perform the PATH stripping / adding.
-					string? OrigPathVar = Environment.GetEnvironmentVariable("PATH");
-					string PathDelimiter = OperatingSystem.IsWindows() ? ";" : ":";
-					string[] PathVars = [];
-					if (!String.IsNullOrEmpty(OrigPathVar))
+					string? origPathVar = Environment.GetEnvironmentVariable("PATH");
+					string pathDelimiter = OperatingSystem.IsWindows() ? ";" : ":";
+					string[] pathVars = [];
+					if (!String.IsNullOrEmpty(origPathVar))
 					{
-						PathVars = OrigPathVar.Split(PathDelimiter.ToCharArray());
+						pathVars = origPathVar.Split(pathDelimiter.ToCharArray());
 					}
 					else
 					{
 						Logger.LogDebug("Path environment variable is null during AutoSDK");
 					}
 
-					List<string> ModifiedPathVars = [.. PathVars];
+					List<string> modifiedPathVars = [.. pathVars];
 
 					// perform removes first, in case they overlap with any adds.
-					foreach (string PathRemove in PathRemoves)
+					foreach (string pathRemove in pathRemoves)
 					{
-						foreach (string PathVar in PathVars)
+						foreach (string pathVar in pathVars)
 						{
-							if (PathVar.IndexOf(PathRemove, StringComparison.OrdinalIgnoreCase) >= 0)
+							if (pathVar.Contains(pathRemove, StringComparison.OrdinalIgnoreCase))
 							{
-								Logger.LogDebug("Removing Path: '{Path}'", PathVar);
-								ModifiedPathVars.Remove(PathVar);
+								Logger.LogDebug("Removing Path: '{Path}'", pathVar);
+								modifiedPathVars.Remove(pathVar);
 							}
 						}
 					}
 
 					// remove all the of ADDs so that if this function is executed multiple times, the paths will be guaranteed to be in the same order after each run.
 					// If we did not do this, a 'remove' that matched some, but not all, of our 'adds' would cause the order to change.
-					foreach (string PathAdd in PathAdds)
+					foreach (string pathAdd in pathAdds)
 					{
-						foreach (string PathVar in PathVars)
+						foreach (string pathVar in pathVars)
 						{
-							if (String.Compare(PathAdd, PathVar, true) == 0)
+							if (String.Equals(pathAdd, pathVar, StringComparison.OrdinalIgnoreCase))
 							{
-								Logger.LogDebug("Removing Path: '{Path}'", PathVar);
-								ModifiedPathVars.Remove(PathVar);
+								Logger.LogDebug("Removing Path: '{Path}'", pathVar);
+								modifiedPathVars.Remove(pathVar);
 							}
 						}
 					}
 
 					// perform adds, but don't add duplicates
-					foreach (string PathAdd in PathAdds)
+					foreach (string pathAdd in pathAdds)
 					{
-						if (!ModifiedPathVars.Contains(PathAdd))
+						if (!modifiedPathVars.Contains(pathAdd))
 						{
-							Logger.LogDebug("Adding Path: '{Path}'", PathAdd);
-							ModifiedPathVars.Add(PathAdd);
+							Logger.LogDebug("Adding Path: '{Path}'", pathAdd);
+							modifiedPathVars.Add(pathAdd);
 						}
 					}
 
-					string ModifiedPath = String.Join(PathDelimiter, ModifiedPathVars);
-					Environment.SetEnvironmentVariable("PATH", ModifiedPath);
+					string modifiedPath = String.Join(pathDelimiter, modifiedPathVars);
+					Environment.SetEnvironmentVariable("PATH", modifiedPath);
 
-					Reader.Close();
+					reader.Close();
 
 					// write out environment variable command so any process using this commandfile will mark itself as having had autosdks set up.
 					// avoids child processes spuriously detecting manualsdks.
 					if (bNeedsToWriteAutoSetupEnvVar)
 					{
 						// write out the manual sdk version since child processes won't be able to detect manual with AutoSDK messing up env vars
-						using (StreamWriter Writer = File.AppendText(EnvVarFile))
+						using (StreamWriter writer = File.AppendText(envVarFile))
 						{
-							Writer.WriteLine("{0}={1}", PlatformSetupEnvVar, "1");
+							writer.WriteLine("{0}={1}", platformSetupEnvVar, "1");
 						}
 						// set the variable in the local environment in case this process spawns any others.
-						Environment.SetEnvironmentVariable(PlatformSetupEnvVar, "1");
+						Environment.SetEnvironmentVariable(platformSetupEnvVar, "1");
 					}
 
 					// make sure we know that we've modified the local environment, invalidating manual installs for this run.
 					bLocalProcessSetupAutoSDK = true;
 
 					// tell any child processes what our manual versions were before setting up autosdk
-					string ValueToWrite = CachedManualSDKVersions.Count > 0 ? JsonSerializer.Serialize(CachedManualSDKVersions) : "__None";
-					Environment.SetEnvironmentVariable(GetPlatformManualSDKSetupEnvVar(), ValueToWrite);
+					string valueToWrite = CachedManualSDKVersions.Count > 0 ? JsonSerializer.Serialize(CachedManualSDKVersions) : "__None";
+					Environment.SetEnvironmentVariable(GetPlatformManualSDKSetupEnvVar(), valueToWrite);
 
 					return true;
 				}
 			}
 			else
 			{
-				Logger.LogDebug("Cannot set up environment for {SdkRoot} because command file {EnvVarFile} does not exist.", PlatformSDKRoot, EnvVarFile);
+				Logger.LogDebug("Cannot set up environment for {SdkRoot} because command file {EnvVarFile} does not exist.", platformSDKRoot, envVarFile);
 			}
 
 			return false;
@@ -2034,25 +1994,25 @@ namespace EpicGames.Core
 
 		protected void InvalidateCurrentlyInstalledAutoSDK()
 		{
-			string PlatformSDKRoot = GetPathToPlatformAutoSDKs();
-			if (Directory.Exists(PlatformSDKRoot))
+			string platformSDKRoot = GetPathToPlatformAutoSDKs();
+			if (Directory.Exists(platformSDKRoot))
 			{
-				string SDKFilename = Path.Combine(PlatformSDKRoot, CurrentlyInstalledSDKStringManifest);
-				if (File.Exists(SDKFilename))
+				string sdkFilename = Path.Combine(platformSDKRoot, CurrentlyInstalledSDKStringManifest);
+				if (File.Exists(sdkFilename))
 				{
-					File.Delete(SDKFilename);
+					File.Delete(sdkFilename);
 				}
 
-				string VersionFilename = Path.Combine(PlatformSDKRoot, LastRunScriptVersionManifest);
-				if (File.Exists(VersionFilename))
+				string versionFilename = Path.Combine(platformSDKRoot, LastRunScriptVersionManifest);
+				if (File.Exists(versionFilename))
 				{
-					File.Delete(VersionFilename);
+					File.Delete(versionFilename);
 				}
 
-				string EnvVarFile = Path.Combine(PlatformSDKRoot, SDKEnvironmentVarsFile);
-				if (File.Exists(EnvVarFile))
+				string envVarFile = Path.Combine(platformSDKRoot, SDKEnvironmentVarsFile);
+				if (File.Exists(envVarFile))
 				{
-					File.Delete(EnvVarFile);
+					File.Delete(envVarFile);
 				}
 			}
 		}
@@ -2065,32 +2025,31 @@ namespace EpicGames.Core
 		{
 			if (PlatformSupportsAutoSDKs() && HasAutoSDKSystemEnabled())
 			{
-				string AutoSDKRoot = GetPathToPlatformAutoSDKs();
-				if (!String.IsNullOrEmpty(AutoSDKRoot))
+				string autoSDKRoot = GetPathToPlatformAutoSDKs();
+				if (!String.IsNullOrEmpty(autoSDKRoot))
 				{
-					string DesiredSDKLevel = GetAutoSDKLevelForPlatform(AutoSDKRoot);
+					string desiredSDKLevel = GetAutoSDKLevelForPlatform(autoSDKRoot);
 					// if the user doesn't want AutoSDK for this platform, then return that it is not installed, even if it actually is
-					if (String.Compare(DesiredSDKLevel, "NONE", true) == 0)
+					if (String.Equals(desiredSDKLevel, "NONE", StringComparison.OrdinalIgnoreCase))
 					{
 						return SDKStatus.Invalid;
 					}
 
 					// check script version so script fixes can be propagated without touching every build machine's CurrentlyInstalled file manually.
-					string CurrentScriptVersionString;
-					if (GetLastRunScriptVersionString(AutoSDKRoot, out CurrentScriptVersionString) && CurrentScriptVersionString == GetRequiredScriptVersionString())
+					if (GetLastRunScriptVersionString(autoSDKRoot, out string currentScriptVersionString) && currentScriptVersionString == GetRequiredScriptVersionString())
 					{
 						// check to make sure OutputEnvVars doesn't need regenerating
-						string EnvVarFile = Path.Combine(AutoSDKRoot, SDKEnvironmentVarsFile);
-						bool bEnvVarFileExists = File.Exists(EnvVarFile);
+						string envVarFile = Path.Combine(autoSDKRoot, SDKEnvironmentVarsFile);
+						bool bEnvVarFileExists = File.Exists(envVarFile);
 
-						string CurrentSDKString;
-						string CurrentSDKLevel;
-						if (bEnvVarFileExists && GetCurrentlyInstalledSDKString(AutoSDKRoot, out CurrentSDKString, out CurrentSDKLevel))
+						string currentSDKString;
+						string currentSDKLevel;
+						if (bEnvVarFileExists && GetCurrentlyInstalledSDKString(autoSDKRoot, out currentSDKString, out currentSDKLevel))
 						{
 							// match version
-							if (CurrentSDKString == GetAutoSDKDirectoryForMainVersion())
+							if (currentSDKString == GetAutoSDKDirectoryForMainVersion())
 							{
-								if (IsSDKLevelAtLeast(CurrentSDKLevel, DesiredSDKLevel))
+								if (IsSDKLevelAtLeast(currentSDKLevel, desiredSDKLevel))
 								{
 									return SDKStatus.Valid;
 								}
@@ -2113,12 +2072,12 @@ namespace EpicGames.Core
 			return bLocalProcessSetupAutoSDK || HasParentProcessSetupAutoSDK(out _);
 		}
 
-		protected bool HasParentProcessSetupAutoSDK([NotNullWhen(true)] out string? OutAutoSDKSetupValue)
+		protected bool HasParentProcessSetupAutoSDK([NotNullWhen(true)] out string? outAutoSDKSetupValue)
 		{
-			string AutoSDKSetupVarName = GetPlatformAutoSDKSetupEnvVar();
-			OutAutoSDKSetupValue = Environment.GetEnvironmentVariable(AutoSDKSetupVarName);
+			string autoSDKSetupVarName = GetPlatformAutoSDKSetupEnvVar();
+			outAutoSDKSetupValue = Environment.GetEnvironmentVariable(autoSDKSetupVarName);
 
-			if (!String.IsNullOrEmpty(OutAutoSDKSetupValue))
+			if (!String.IsNullOrEmpty(outAutoSDKSetupValue))
 			{
 				return true;
 			}
@@ -2184,11 +2143,11 @@ namespace EpicGames.Core
 				// run installation for autosdk if necessary.
 				if (HasRequiredAutoSDKInstalled() == SDKStatus.Invalid)
 				{
-					string AutoSDKRoot = GetPathToPlatformAutoSDKs();
+					string autoSDKRoot = GetPathToPlatformAutoSDKs();
 
-					string DesiredSDKLevel = GetAutoSDKLevelForPlatform(AutoSDKRoot);
+					string desiredSDKLevel = GetAutoSDKLevelForPlatform(autoSDKRoot);
 					// if the user doesn't want AutoSDK for this platform, then do nothing
-					if (String.Compare(DesiredSDKLevel, "NONE", true) == 0)
+					if (String.Equals(desiredSDKLevel, "NONE", StringComparison.OrdinalIgnoreCase))
 					{
 						Logger.LogDebug("Skipping AutoSDK for {PlatformName} because NONE was specified as the desired AutoSDK level", PlatformName);
 						return;
@@ -2197,36 +2156,34 @@ namespace EpicGames.Core
 					//reset check status so any checking sdk status after the attempted setup will do a real check again.
 					SDKCheckStatus = -1;
 
-					string CurrentSDKString;
-					string CurrentSDKLevel;
-					GetCurrentlyInstalledSDKString(AutoSDKRoot, out CurrentSDKString, out CurrentSDKLevel);
+					GetCurrentlyInstalledSDKString(autoSDKRoot, out string currentSDKString, out string currentSDKLevel);
 
 					// switch over (note that version string can be empty)
-					if (!RunAutoSDKHooks(AutoSDKRoot, CurrentSDKString, CurrentSDKLevel, SDKHookType.Uninstall))
+					if (!RunAutoSDKHooks(autoSDKRoot, currentSDKString, currentSDKLevel, SDKHookType.Uninstall))
 					{
-						Logger.LogDebug("Failed to uninstall currently installed SDK {SdkVersion}", CurrentSDKString);
+						Logger.LogDebug("Failed to uninstall currently installed SDK {SdkVersion}", currentSDKString);
 						InvalidateCurrentlyInstalledAutoSDK();
 						return;
 					}
 					// delete Manifest file to avoid multiple uninstalls
 					InvalidateCurrentlyInstalledAutoSDK();
 
-					if (!RunAutoSDKHooks(AutoSDKRoot, GetAutoSDKDirectoryForMainVersion(), DesiredSDKLevel, SDKHookType.Install, false))
+					if (!RunAutoSDKHooks(autoSDKRoot, GetAutoSDKDirectoryForMainVersion(), desiredSDKLevel, SDKHookType.Install, false))
 					{
 						Logger.LogDebug("Failed to install required SDK {SdkVersion}.  Attempting to uninstall", GetAutoSDKDirectoryForMainVersion());
-						RunAutoSDKHooks(AutoSDKRoot, GetAutoSDKDirectoryForMainVersion(), DesiredSDKLevel, SDKHookType.Uninstall, false);
+						RunAutoSDKHooks(autoSDKRoot, GetAutoSDKDirectoryForMainVersion(), desiredSDKLevel, SDKHookType.Uninstall, false);
 						return;
 					}
 
-					string EnvVarFile = Path.Combine(AutoSDKRoot, SDKEnvironmentVarsFile);
-					if (!File.Exists(EnvVarFile))
+					string envVarFile = Path.Combine(autoSDKRoot, SDKEnvironmentVarsFile);
+					if (!File.Exists(envVarFile))
 					{
-						Logger.LogDebug("Installation of required SDK {SdkVersion}.  Did not generate Environment file {EnvVarFile}", GetAutoSDKDirectoryForMainVersion(), EnvVarFile);
-						RunAutoSDKHooks(AutoSDKRoot, GetAutoSDKDirectoryForMainVersion(), DesiredSDKLevel, SDKHookType.Uninstall, false);
+						Logger.LogDebug("Installation of required SDK {SdkVersion}.  Did not generate Environment file {EnvVarFile}", GetAutoSDKDirectoryForMainVersion(), envVarFile);
+						RunAutoSDKHooks(autoSDKRoot, GetAutoSDKDirectoryForMainVersion(), desiredSDKLevel, SDKHookType.Uninstall, false);
 						return;
 					}
 
-					SetCurrentlyInstalledAutoSDKString(GetAutoSDKDirectoryForMainVersion(), DesiredSDKLevel);
+					SetCurrentlyInstalledAutoSDKString(GetAutoSDKDirectoryForMainVersion(), desiredSDKLevel);
 					SetLastRunAutoSDKScriptVersion(GetRequiredScriptVersionString());
 				}
 
@@ -2262,10 +2219,10 @@ namespace EpicGames.Core
 
 		public SDKStatus SetupEnvironmentFromAutoSDK()
 		{
-			string PlatformSDKRoot = GetPathToPlatformAutoSDKs();
+			string platformSDKRoot = GetPathToPlatformAutoSDKs();
 
 			// load environment variables from current SDK
-			if (!SetupEnvironmentFromAutoSDK(PlatformSDKRoot))
+			if (!SetupEnvironmentFromAutoSDK(platformSDKRoot))
 			{
 				Logger.LogDebug("Failed to load environment from required SDK {SdkRoot}", GetAutoSDKDirectoryForMainVersion());
 				InvalidateCurrentlyInstalledAutoSDK();
@@ -2321,7 +2278,7 @@ namespace EpicGames.Core
 			}
 
 			// print all SDKs to log file (errors will print out later for builds and generateprojectfiles)
-			PrintSDKInfoAndReturnValidity(LogEventType.Log, LogFormatOptions.NoConsoleOutput, LogEventType.Log, LogFormatOptions.NoConsoleOutput, bBriefInvalidSDKWarnings:true);
+			PrintSDKInfoAndReturnValidity(LogEventType.Log, LogFormatOptions.NoConsoleOutput, LogEventType.Log, LogFormatOptions.NoConsoleOutput, bBriefInvalidSDKWarnings: true);
 		}
 		#endregion
 
