[EpicGames.Build]
* ParallelQueue - Removed dispose call of semaphore _available used for workers processing queue. Since tasks could be started after ParallelQueue goes out of scope we need to make sure _available live long enough. This will give the dispose over to the garbage collector instead.

#Synced-CL 49076111

diff --git a/Engine/Source/Programs/Shared/EpicGames.Build/System/Parallel2.cs b/Engine/Source/Programs/Shared/EpicGames.Build/System/Parallel2.cs
--- a/Engine/Source/Programs/Shared/EpicGames.Build/System/Parallel2.cs
+++ b/Engine/Source/Programs/Shared/EpicGames.Build/System/Parallel2.cs
@@ -332,7 +332,9 @@ namespace UnrealBuildBase
 			}
 		}
 
-		public void Dispose() => _available.Dispose();
+		// Note, we can not dispose _available here since it needs to outlive all started tasks which could start after ParallelQueue
+		// goes out of scope.
+		public void Dispose() { } // => _available.Dispose(); 
 
 		private readonly ConcurrentQueue<Action> _queue = new();
 		private readonly SemaphoreSlim _available = new(0);
