[EpicGames.Build]
* Fixed static analysis warning by adding _available.Dispose in Dispose call (but only if there were tasks started

#Synced-CL 49076111

diff --git a/Engine/Source/Programs/Shared/EpicGames.Build/System/Parallel2.cs b/Engine/Source/Programs/Shared/EpicGames.Build/System/Parallel2.cs
--- a/Engine/Source/Programs/Shared/EpicGames.Build/System/Parallel2.cs
+++ b/Engine/Source/Programs/Shared/EpicGames.Build/System/Parallel2.cs
@@ -266,12 +266,8 @@ namespace UnrealBuildBase
 				return;
 			}
 
-			if (helperCount == -1)
-			{
-				helperCount = MaxHelperCount - 1;
-			}
-
-			_workerCount = helperCount + 1;
+			_helperCount = helperCount == -1 ? MaxHelperCount - 1 : helperCount;
+			_workerCount = _helperCount + 1;
 
 			void WorkerLoop()
 			{
@@ -312,7 +308,7 @@ namespace UnrealBuildBase
 				}
 			}
 
-			for (int i = 0; i != helperCount; ++i)
+			for (int i = 0; i != _helperCount; ++i)
 			{
 				if (useTasks)
 				{
@@ -332,15 +328,22 @@ namespace UnrealBuildBase
 			}
 		}
 
-		// Note, we can not dispose _available here since it needs to outlive all started tasks which could start after ParallelQueue
-		// goes out of scope.
-		public void Dispose() { } // => _available.Dispose(); 
+		public void Dispose()
+		{
+			// Note, we can not dispose _available here if there are tasks since it needs to outlive
+			// all started tasks which could start after ParallelQueue goes out of scope.
+			if (_helperCount == 0)
+			{
+				_available.Dispose();
+			}
+		}
 
 		private readonly ConcurrentQueue<Action> _queue = new();
 		private readonly SemaphoreSlim _available = new(0);
 		private int _outstanding;
 		private int _accepting = 1;
 		private int _done;
+		private int _helperCount;
 		private int _workerCount;
 		private ExceptionDispatchInfo? _exception;
 
