[Horde][Analytics][Shared] New Analytics projects

#Synced-CL 48916881

diff --git a/Engine/Source/Programs/Horde/Plugins/Build/HordeServer.Build/HordeServer.Build.csproj b/Engine/Source/Programs/Horde/Plugins/Build/HordeServer.Build/HordeServer.Build.csproj
--- a/Engine/Source/Programs/Horde/Plugins/Build/HordeServer.Build/HordeServer.Build.csproj
+++ b/Engine/Source/Programs/Horde/Plugins/Build/HordeServer.Build/HordeServer.Build.csproj
@@ -31,6 +31,7 @@
   </ItemGroup>
 
   <ItemGroup>
+    <ProjectReference Include="..\..\..\..\Shared\EpicGames.Analytics.Telemetry\EpicGames.Analytics.Telemetry.csproj" />
     <ProjectReference Include="..\..\..\..\Shared\EpicGames.Core\EpicGames.Core.csproj" />
     <ProjectReference Include="..\..\..\..\Shared\EpicGames.Horde\EpicGames.Horde.csproj" />
     <ProjectReference Include="..\..\..\..\Shared\EpicGames.IoHash\EpicGames.IoHash.csproj" />
diff --git a/Engine/Source/Programs/Horde/Plugins/Build/HordeServer.Build/Jobs/JobTelemetry.cs b/Engine/Source/Programs/Horde/Plugins/Build/HordeServer.Build/Jobs/JobTelemetry.cs
--- a/Engine/Source/Programs/Horde/Plugins/Build/HordeServer.Build/Jobs/JobTelemetry.cs
+++ b/Engine/Source/Programs/Horde/Plugins/Build/HordeServer.Build/Jobs/JobTelemetry.cs
@@ -1,5 +1,7 @@
 // Copyright Epic Games, Inc. All Rights Reserved.
 
+using System.ComponentModel.DataAnnotations.Schema;
+using EpicGames.Analytics.Telemetry;
 using EpicGames.Horde.Jobs;
 using EpicGames.Horde.Jobs.Graphs;
 using EpicGames.Horde.Jobs.Templates;
@@ -12,106 +14,128 @@ namespace HordeServer.Jobs
 	/// <summary>
 	/// Record used for Job Summary telemetry. Models a <see cref="IJob"/>.
 	/// </summary>
+	[AnalyticsTableGen]
+	[Table("ingest.horde.state_job_summary")]
 	public record JobSummaryTelemetry
 	{
 		/// <summary>
 		/// The telemetry event name.
 		/// </summary>
+		[Column("event_name")]
 		public string EventName { get; init; }
 
 		/// <summary>
 		/// The schema version of the telemetry.
 		/// </summary>
+		[Column("schema_version")]
 		public int SchemaVersion { get; init; }
 
 		/// <summary>
 		/// The job Id.
 		/// </summary>
+		[Column("job_id")]
 		public JobId JobId { get; init; }
 
 		/// <summary>
 		/// The template of the job.
 		/// </summary>
+		[Column("template_id")]
 		public TemplateId TemplateId { get; init; }
 
 		/// <summary>
 		/// The stream of the job.
 		/// </summary>
+		[Column("stream_id")]
 		public StreamId StreamId { get; init; }
 
 		/// <summary>
 		/// The commit of the job. <see cref="IJob.CommitId"/>.
 		/// </summary>
+		[Column("commit_id")]
 		public string CommitId { get; init; }
 
 		/// <summary>
 		/// The code commit of the job. <see cref="IJob.CodeCommitId"/>.
 		/// </summary>
+		[Column("code_commit_id")]
 		public string CodeCommitId { get; init; }
 
 		/// <summary>
 		/// The created time of the job.
 		/// </summary>
+		[Column("create_time_utc")]
 		public DateTime CreateTimeUtc { get; init; }
 
 		/// <summary>
 		/// The finish time of the job.
 		/// </summary>
+		[Column("finish_time_utc")]
 		public DateTime FinishTimeUtc { get; init; }
 
 		/// <summary>
 		/// The summed duration of all steps.
 		/// </summary>
+		[Column("job_steps_total_time")]
 		public float JobStepsTotalTime { get; init; }
 
 		/// <summary>
 		/// The wall time of the job, in seconds.
 		/// </summary>
+		[Column("job_wall_time")]
 		public float JobWallTime { get; init; }
 
 		/// <summary>
 		/// <see cref="JobStepsCompletionInfo.PassRatio"/>.
 		/// </summary>
+		[Column("pass_ratio")]
 		public float PassRatio { get; init; }
 
 		/// <summary>
 		/// <see cref="JobStepsCompletionInfo.PassWithWarningRatio"/>.
 		/// </summary>
+		[Column("pass_with_warning_ratio")]
 		public float PassWithWarningRatio { get; init; }
 
 		/// <summary>
 		/// <see cref="JobStepsCompletionInfo.WarningRatio"/>.
 		/// </summary>
+		[Column("warning_ratio")]
 		public float WarningRatio { get; init; }
 
 		/// <summary>
 		/// <see cref="JobStepsCompletionInfo.FailureRatio"/>.
 		/// </summary>
+		[Column("failure_ratio")]
 		public float FailureRatio { get; init; }
 
 		/// <summary>
 		/// <see cref="JobStepsCompletionInfo.StepPassCount"/>.
 		/// </summary>
+		[Column("step_pass_count")]
 		public int StepPassCount { get; init; }
 
 		/// <summary>
 		/// <see cref="JobStepsCompletionInfo.StepWarningCount"/>.
 		/// </summary>
+		[Column("step_warning_count")]
 		public int StepWarningCount { get; init; }
 
 		/// <summary>
 		/// <see cref="JobStepsCompletionInfo.StepFailureCount"/>.
 		/// </summary>
+		[Column("step_failure_count")]
 		public int StepFailureCount { get; init; }
 
 		/// <summary>
 		/// <see cref="JobStepsCompletionInfo.StepTotalCount"/>.
 		/// </summary>
+		[Column("step_total_count")]
 		public int StepTotalCount { get; init; }
 
 		/// <summary>
 		/// Whether the job was a preflight or not.
 		/// </summary>
+		[Column("is_preflight")]
 		public bool IsPreflight { get; init; }
 
 		/// <summary>
@@ -155,6 +179,35 @@ namespace HordeServer.Jobs
 			WarningRatio = jobStepsCompletionInfo.WarningRatio;
 			FailureRatio = jobStepsCompletionInfo.FailureRatio;
 		}
+
+		/// <summary>
+		/// Default constructor for an empty telemetry object,
+		/// </summary>
+		/// <remarks>This constructor is effectively reserved for ORM mapping.</remarks>
+		public JobSummaryTelemetry()
+		{
+			EventName = DefaultEventName;
+			SchemaVersion = CurrentSchemaVersion;
+			CodeCommitId = String.Empty;
+			JobId = JobId.Empty;
+			TemplateId = new TemplateId(String.Empty);
+			StreamId = new StreamId(String.Empty);
+			CommitId = String.Empty;
+			CodeCommitId = String.Empty;
+			CreateTimeUtc = DateTime.UtcNow;
+			FinishTimeUtc = DateTime.UtcNow;
+			JobWallTime = 0;
+			IsPreflight = false;
+			StepTotalCount = 0;
+			StepPassCount = 0;
+			StepWarningCount= 0;
+			StepFailureCount = 0;
+			JobStepsTotalTime = 0;
+			PassRatio = 0;
+			PassWithWarningRatio = 0;
+			WarningRatio = 0;
+			FailureRatio = 0;
+		}
 	}
 
 	/// <summary>
diff --git a/Engine/Source/Programs/Shared/EpicGames.Analytics.Generators/AnalyticsSourceGenHelper.cs b/Engine/Source/Programs/Shared/EpicGames.Analytics.Generators/AnalyticsSourceGenHelper.cs
new file mode 100644
--- /dev/null
+++ b/Engine/Source/Programs/Shared/EpicGames.Analytics.Generators/AnalyticsSourceGenHelper.cs
@@ -0,0 +1,89 @@
+// Copyright Epic Games, Inc. All Rights Reserved.
+
+using System;
+using System.Text;
+
+namespace EpicGames.Analytics.Generators
+{
+	/// <summary>
+	/// General source gen helper associated with analytics tables.
+	/// </summary>
+	public static class AnalyticsSourceGenHelper
+	{
+		/// <summary>
+		/// Analytics Generator namespace.
+		/// </summary>
+		public const string AnalyticsTableGenNamespace = "EpicGames.Analytics.Generated";
+		
+		/// <summary>
+		/// Analytics Generator prefix used in generated files.
+		/// </summary>
+		public const string AnalyticsTableGenPrefix = "AnalyticsTableGen";
+
+		/// <summary>
+		/// Genreates the version file for the table gen.
+		/// </summary>
+		/// <returns>string representing the version.</returns>
+		public static string GenerateAnalyticsTableGenVersionFile()
+		{
+			StringBuilder sb = new StringBuilder();
+			sb.Append($$"""
+				namespace {{AnalyticsTableGenNamespace}}
+				{
+					/// <summary>
+					/// Descriptor class for generated metadata.
+					/// </summary>
+					public static class AnalyticsTableGen
+					{
+						/// <summary>
+						/// The version of the table gen.
+						/// </summary>
+						public const int Version=1;
+					}
+				}
+				""");
+
+			return sb.ToString();
+		}
+
+		/// <summary>
+		/// Constructs a analytics table source file from the provided object.
+		/// </summary>
+		/// <param name="analyticstableToGenerate">The analytics table definition.</param>
+		/// <returns>The c-sharp representation of the object.</returns>
+		public static string GenerateAnalyticsTable(AnalyticsTableToGenerate analyticstableToGenerate)
+		{
+			{
+				StringBuilder sb = new StringBuilder();
+				sb.Append($$"""
+				namespace {{AnalyticsTableGenNamespace}}
+				{
+					/// <summary>
+					/// Generated Table Descriptor for {{analyticstableToGenerate.Name}}.
+					/// </summary>
+					public static class {{analyticstableToGenerate.Name}}Gen
+					{
+
+				""");
+
+				sb.AppendLine($"\t\t/// <summary>");
+				sb.AppendLine($"\t\t///Table name.");
+				sb.AppendLine($"\t\t/// </summary>");
+				sb.AppendLine($"\t\tpublic const string TableName= \"{analyticstableToGenerate.TableName}\";");
+
+				foreach (Tuple<string, string> member in analyticstableToGenerate.ColumnsNames)
+				{
+					sb.AppendLine($"\t\t/// <summary>");
+					sb.AppendLine($"\t\t/// Column descriptor for {member.Item1}.");
+					sb.AppendLine($"\t\t/// </summary>");
+					sb.AppendLine($"\t\tpublic const string {member.Item1} = \"{member.Item2}\";");
+				}
+
+				sb.AppendLine("	}");
+				sb.AppendLine("}");
+
+				return sb.ToString();
+			}
+		}
+	}
+}
diff --git a/Engine/Source/Programs/Shared/EpicGames.Analytics.Generators/AnalyticsTableGenerator.cs b/Engine/Source/Programs/Shared/EpicGames.Analytics.Generators/AnalyticsTableGenerator.cs
new file mode 100644
--- /dev/null
+++ b/Engine/Source/Programs/Shared/EpicGames.Analytics.Generators/AnalyticsTableGenerator.cs
@@ -0,0 +1,209 @@
+// Copyright Epic Games, Inc. All Rights Reserved.
+
+using System;
+using System.Collections.Generic;
+using System.Collections.Immutable;
+using System.Linq;
+using System.Text;
+using Microsoft.CodeAnalysis;
+using Microsoft.CodeAnalysis.CSharp.Syntax;
+using Microsoft.CodeAnalysis.Text;
+
+namespace EpicGames.Analytics.Generators
+{
+	/// <summary>
+	/// DTO that describes a analytics table.
+	/// </summary>
+	public readonly record struct AnalyticsTableToGenerate
+	{
+		/// <summary>
+		/// The name of the telemetry record.
+		/// </summary>
+		public readonly string Name;
+
+		/// <summary>
+		/// The table name.
+		/// </summary>
+		public readonly string TableName;
+
+		/// <summary>
+		/// The column names.
+		/// </summary>
+		public readonly List<Tuple<string, string>> ColumnsNames;
+
+		/// <summary>
+		/// Constructor.
+		/// </summary>
+		/// <param name="name">Name of the telemetry record.</param>
+		/// <param name="tableName">Name of the table.</param>
+		/// <param name="propsValues">List of the property and column name pairs.</param>
+		public AnalyticsTableToGenerate(string name, string tableName, List<Tuple<string, string>> propsValues)
+		{
+			Name = name;
+			TableName = tableName;
+			ColumnsNames = propsValues;
+		}
+	}
+
+	/// <summary>
+	/// Generator class.
+	/// </summary>
+	[Generator]
+	public class AnalyticsTableGenerator : ISourceGenerator
+	{
+#pragma warning disable RS2008 // Enable analyzer release tracking
+		private static readonly DiagnosticDescriptor s_descriptor = new DiagnosticDescriptor(id: "HORDEATG01", title: "Unexpected Error", messageFormat: "Error for object: {0}", category: "Design", defaultSeverity: DiagnosticSeverity.Warning, isEnabledByDefault: true);
+#pragma warning restore RS2008 // Enable analyzer release tracking
+
+		/// <inheritdoc/>
+		public void Initialize(GeneratorInitializationContext context)
+		{
+			context.RegisterForSyntaxNotifications(() => new SyntaxReceiver());
+		}
+
+		/// <inheritdoc/>
+		public void Execute(GeneratorExecutionContext context)
+		{
+			if (context.SyntaxReceiver is not SyntaxReceiver)
+			{
+				return;
+			}
+
+			context.AddSource(AnalyticsSourceGenHelper.AnalyticsTableGenPrefix + ".g.cs", SourceText.From(AnalyticsSourceGenHelper.GenerateAnalyticsTableGenVersionFile(), Encoding.UTF8));
+
+			// Search all named types in the compilation; we currently must do this becuase of an indirection in dependencies. Ideally this analyzer is applied to a EpicGames.Analytics.Telemetry project, which contains the neccesssary references.
+			IEnumerable<INamedTypeSymbol> allTypes = GetAllNamedTypes(context.Compilation.GlobalNamespace);
+
+			foreach (INamedTypeSymbol typeSymbol in allTypes)
+			{
+				ImmutableArray<AttributeData> attributes = typeSymbol.GetAttributes();
+				if (attributes.Length == 0)
+				{
+					continue;
+				}
+
+				if (!attributes.Any(x => x.AttributeClass?.ToDisplayString() == "EpicGames.Analytics.Telemetry.AnalyticsTableGenAttribute"))
+				{
+					continue;
+				}
+
+				AnalyticsTableToGenerate? attg = BuildAnalyticsTable(typeSymbol);
+
+				if (attg == null)
+				{
+					context.ReportDiagnostic(Diagnostic.Create(s_descriptor, Location.None, $"Failed to produce table for type: {typeSymbol.ToDisplayString()}"));
+					continue;
+				}
+
+				string result = AnalyticsSourceGenHelper.GenerateAnalyticsTable(attg.Value);
+				context.AddSource(AnalyticsSourceGenHelper.AnalyticsTableGenPrefix + "." + attg.Value.Name + ".g.cs", SourceText.From(result, Encoding.UTF8));
+			}
+		}
+
+		private class SyntaxReceiver : ISyntaxReceiver
+		{
+			public List<TypeDeclarationSyntax> CandidateTypes { get; } = new List<TypeDeclarationSyntax>();
+
+			public void OnVisitSyntaxNode(SyntaxNode syntaxNode)
+			{
+				if (syntaxNode is TypeDeclarationSyntax tds && tds.AttributeLists.Count > 0)
+				{
+					CandidateTypes.Add(tds);
+				}
+			}
+		}
+
+#nullable enable
+		private static string? BuildFullNamespace(INamespaceSymbol? symbol)
+		{
+			if (symbol == null)
+			{
+				return null;
+			}
+			string? retString = BuildFullNamespace(symbol.ContainingNamespace);
+
+			return !String.IsNullOrEmpty(retString) ? $"{retString}_{symbol.Name}" : symbol.Name;
+		}
+#nullable disable
+
+		private static AnalyticsTableToGenerate? BuildAnalyticsTable(INamedTypeSymbol telemetrySymbol)
+		{
+#nullable enable
+			string? fullNamepsace = BuildFullNamespace(telemetrySymbol.ContainingNamespace);
+			string telemetryName = !String.IsNullOrEmpty(fullNamepsace) ? $"{fullNamepsace}_{telemetrySymbol.Name}" : telemetrySymbol.Name;
+			string? tableName = null;
+
+			List<Tuple<string, string>> columns = new List<Tuple<string, string>>();
+
+			foreach (AttributeData attr in telemetrySymbol.GetAttributes())
+			{
+				if (attr.AttributeClass?.ToDisplayString() == "System.ComponentModel.DataAnnotations.Schema.TableAttribute")
+				{
+					if (attr.ConstructorArguments.Length > 0)
+					{
+						tableName = attr.ConstructorArguments[0].Value as string;
+					}
+				}
+			}
+
+			foreach (IPropertySymbol prop in telemetrySymbol.GetMembers().OfType<IPropertySymbol>())
+			{
+				string propertyName = prop.Name;
+				string? columnName = null;
+
+				foreach (AttributeData attr in prop.GetAttributes())
+				{
+					if (attr.AttributeClass?.ToDisplayString() == "System.ComponentModel.DataAnnotations.Schema.ColumnAttribute")
+					{
+						if (attr.ConstructorArguments.Length > 0)
+						{
+							columnName = attr.ConstructorArguments[0].Value as string;
+						}
+					}
+				}
+
+				if (columnName != null)
+				{
+					columns.Add(new Tuple<string, string>(propertyName, columnName));
+				}
+			}
+
+			return String.IsNullOrEmpty(tableName) ? null : new AnalyticsTableToGenerate(telemetryName, tableName, columns);
+#nullable disable
+		}
+
+		private static IEnumerable<INamedTypeSymbol> GetAllNamedTypes(INamespaceSymbol ns)
+		{
+			foreach (INamedTypeSymbol type in ns.GetTypeMembers())
+			{
+				yield return type;
+
+				foreach (INamedTypeSymbol nested in GetAllNestedTypes(type))
+				{
+					yield return nested;
+				}
+			}
+
+			foreach (INamespaceSymbol childNs in ns.GetNamespaceMembers())
+			{
+				foreach (INamedTypeSymbol type in GetAllNamedTypes(childNs))
+				{
+					yield return type;
+				}
+			}
+		}
+
+		private static IEnumerable<INamedTypeSymbol> GetAllNestedTypes(INamedTypeSymbol type)
+		{
+			foreach (INamedTypeSymbol nested in type.GetTypeMembers())
+			{
+				yield return nested;
+
+				foreach (INamedTypeSymbol deeper in GetAllNestedTypes(nested))
+				{
+					yield return deeper;
+				}
+			}
+		}
+	}
+}
\ No newline at end of file
diff --git a/Engine/Source/Programs/Shared/EpicGames.Analytics.Generators/EpicGames.Analytics.Generators.csproj b/Engine/Source/Programs/Shared/EpicGames.Analytics.Generators/EpicGames.Analytics.Generators.csproj
new file mode 100644
--- /dev/null
+++ b/Engine/Source/Programs/Shared/EpicGames.Analytics.Generators/EpicGames.Analytics.Generators.csproj
@@ -0,0 +1,17 @@
+﻿<Project Sdk="Microsoft.NET.Sdk">
+  <Import Project="..\UnrealEngine.csproj.props" />
+
+  <PropertyGroup>
+    <TargetFramework>netstandard2.0</TargetFramework>
+    <LangVersion>latest</LangVersion>
+    <IsRoslynComponent>true</IsRoslynComponent>
+    <EnforceExtendedAnalyzerRules>true</EnforceExtendedAnalyzerRules>
+  </PropertyGroup>
+
+  <ItemGroup>
+    <PackageReference Include="Microsoft.CodeAnalysis.Analyzers" Version="3.3.4" PrivateAssets="all" />
+    <PackageReference Include="Microsoft.CodeAnalysis.Common" Version="4.8.0" PrivateAssets="all" />
+    <PackageReference Include="Microsoft.CodeAnalysis.CSharp" Version="4.8.0" PrivateAssets="all" />
+  </ItemGroup>
+
+</Project>
\ No newline at end of file
diff --git a/Engine/Source/Programs/Shared/EpicGames.Analytics.Telemetry/AnalyticsTableGenAttribute.cs b/Engine/Source/Programs/Shared/EpicGames.Analytics.Telemetry/AnalyticsTableGenAttribute.cs
new file mode 100644
--- /dev/null
+++ b/Engine/Source/Programs/Shared/EpicGames.Analytics.Telemetry/AnalyticsTableGenAttribute.cs
@@ -0,0 +1,10 @@
+// Copyright Epic Games, Inc. All Rights Reserved.
+
+namespace EpicGames.Analytics.Telemetry
+{
+	/// <summary>
+	/// Attribute to annoatate that table generation should run on the type.
+	/// </summary>
+	[AttributeUsage(AttributeTargets.Class | AttributeTargets.Struct)]
+	public sealed class AnalyticsTableGenAttribute : System.Attribute { }
+}
\ No newline at end of file
diff --git a/Engine/Source/Programs/Shared/EpicGames.Analytics.Telemetry/EpicGames.Analytics.Telemetry.csproj b/Engine/Source/Programs/Shared/EpicGames.Analytics.Telemetry/EpicGames.Analytics.Telemetry.csproj
new file mode 100644
--- /dev/null
+++ b/Engine/Source/Programs/Shared/EpicGames.Analytics.Telemetry/EpicGames.Analytics.Telemetry.csproj
@@ -0,0 +1,10 @@
+﻿<Project Sdk="Microsoft.NET.Sdk">
+  <Import Project="..\UnrealEngine.csproj.props" />
+  
+  <PropertyGroup>
+    <TargetFramework>net8.0</TargetFramework>
+    <ImplicitUsings>enable</ImplicitUsings>
+    <Nullable>enable</Nullable>
+  </PropertyGroup>
+
+</Project>
