[Horde][Analytics] Fixing default configs for Analytics Project.

#Synced-CL 49094299

diff --git a/Engine/Source/Programs/Shared/EpicGames.Analytics.Telemetry/AnalyticsTableGenAttribute.cs b/Engine/Source/Programs/Shared/EpicGames.Analytics.Telemetry/AnalyticsTableGenAttribute.cs
--- a/Engine/Source/Programs/Shared/EpicGames.Analytics.Telemetry/AnalyticsTableGenAttribute.cs
+++ b/Engine/Source/Programs/Shared/EpicGames.Analytics.Telemetry/AnalyticsTableGenAttribute.cs
@@ -1,5 +1,7 @@
 // Copyright Epic Games, Inc. All Rights Reserved.
 
+using System;
+
 namespace EpicGames.Analytics.Telemetry
 {
 	/// <summary>
diff --git a/Engine/Source/Programs/Shared/EpicGames.Analytics.Telemetry/EpicGames.Analytics.Telemetry.csproj b/Engine/Source/Programs/Shared/EpicGames.Analytics.Telemetry/EpicGames.Analytics.Telemetry.csproj
--- a/Engine/Source/Programs/Shared/EpicGames.Analytics.Telemetry/EpicGames.Analytics.Telemetry.csproj
+++ b/Engine/Source/Programs/Shared/EpicGames.Analytics.Telemetry/EpicGames.Analytics.Telemetry.csproj
@@ -1,10 +1,34 @@
 ﻿<Project Sdk="Microsoft.NET.Sdk">
   <Import Project="..\UnrealEngine.csproj.props" />
-  
+
   <PropertyGroup>
     <TargetFramework>net8.0</TargetFramework>
-    <ImplicitUsings>enable</ImplicitUsings>
+    <Configuration Condition=" '$(Configuration)' == '' ">Development</Configuration>
+    <OutputType>Library</OutputType>
+    <RootNamespace>EpicGames.Analytics</RootNamespace>
+    <GenerateAssemblyInfo>false</GenerateAssemblyInfo>
+    <GenerateTargetFrameworkAttribute>false</GenerateTargetFrameworkAttribute>
+    <Configurations>Debug;Release;Development;Analyze</Configurations>
     <Nullable>enable</Nullable>
+    <ProduceReferenceAssemblyInOutDir>true</ProduceReferenceAssemblyInOutDir>
+    <AllowUnsafeBlocks>true</AllowUnsafeBlocks>
+    <RunAnalyzersDuringBuild>false</RunAnalyzersDuringBuild>
+  </PropertyGroup>
+  <PropertyGroup Condition=" '$(Configuration)|$(Platform)' == 'Debug|AnyCPU' ">
+    <DocumentationFile></DocumentationFile>
+  </PropertyGroup>
+  <PropertyGroup Condition="'$(Configuration)|$(Platform)' == 'Development|AnyCPU'">
+    <DocumentationFile></DocumentationFile>
+    <Optimize>false</Optimize>
+  </PropertyGroup>
+  <PropertyGroup Condition="'$(Configuration)|$(Platform)' == 'Release|AnyCPU'">
+    <DocumentationFile></DocumentationFile>
+    <Optimize>true</Optimize>
+  </PropertyGroup>
+  <PropertyGroup Condition=" '$(Configuration)|$(Platform)' == 'Analyze|AnyCPU' ">
+    <DocumentationFile></DocumentationFile>
+    <EnforceCodeStyleInBuild>True</EnforceCodeStyleInBuild>
+    <GenerateDocumentationFile>True</GenerateDocumentationFile>
+    <RunAnalyzersDuringBuild>True</RunAnalyzersDuringBuild>
   </PropertyGroup>
-
 </Project>
