[Horde] Add metadata to Analytics

#Synced-CL 49094299

diff --git a/Engine/Source/Programs/Shared/EpicGames.Analytics.Telemetry/EpicGames.Analytics.Telemetry.csproj b/Engine/Source/Programs/Shared/EpicGames.Analytics.Telemetry/EpicGames.Analytics.Telemetry.csproj
--- a/Engine/Source/Programs/Shared/EpicGames.Analytics.Telemetry/EpicGames.Analytics.Telemetry.csproj
+++ b/Engine/Source/Programs/Shared/EpicGames.Analytics.Telemetry/EpicGames.Analytics.Telemetry.csproj
@@ -31,4 +31,9 @@
     <GenerateDocumentationFile>True</GenerateDocumentationFile>
     <RunAnalyzersDuringBuild>True</RunAnalyzersDuringBuild>
   </PropertyGroup>
+  <ItemGroup>
+    <Compile Include="..\..\Shared\MetaData.cs">
+      <Link>Properties\MetaData.cs</Link>
+    </Compile>
+  </ItemGroup>
 </Project>
