[Horde][Analytics] Fixing Analytics gen csproj to include standard configs

#Synced-CL 49094299

diff --git a/Engine/Source/Programs/Shared/EpicGames.Analytics.Generators/EpicGames.Analytics.Generators.csproj b/Engine/Source/Programs/Shared/EpicGames.Analytics.Generators/EpicGames.Analytics.Generators.csproj
--- a/Engine/Source/Programs/Shared/EpicGames.Analytics.Generators/EpicGames.Analytics.Generators.csproj
+++ b/Engine/Source/Programs/Shared/EpicGames.Analytics.Generators/EpicGames.Analytics.Generators.csproj
@@ -6,6 +6,23 @@
     <LangVersion>latest</LangVersion>
     <IsRoslynComponent>true</IsRoslynComponent>
     <EnforceExtendedAnalyzerRules>true</EnforceExtendedAnalyzerRules>
+    <Configuration Condition=" '$(Configuration)' == '' ">Development</Configuration>
+    <Configurations>Debug;Release;Development;Analyze</Configurations>
+  </PropertyGroup>
+
+  <PropertyGroup Condition=" '$(Configuration)|$(Platform)' == 'Debug|AnyCPU' ">
+    <DocumentationFile></DocumentationFile>
+  </PropertyGroup>
+  <PropertyGroup Condition="'$(Configuration)|$(Platform)' == 'Development|AnyCPU'">
+    <DocumentationFile></DocumentationFile>
+    <Optimize>false</Optimize>
+  </PropertyGroup>
+  <PropertyGroup Condition="'$(Configuration)|$(Platform)' == 'Release|AnyCPU'">
+    <DocumentationFile></DocumentationFile>
+    <Optimize>true</Optimize>
+  </PropertyGroup>
+  <PropertyGroup Condition=" '$(Configuration)|$(Platform)' == 'Analyze|AnyCPU' ">
+    <DocumentationFile></DocumentationFile>
   </PropertyGroup>
 
   <ItemGroup>
