[Horde][Analytics][HttpClient] Adding telemetry Post api for clients

#Synced-CL 49094299

diff --git a/Engine/Source/Programs/Shared/EpicGames.Horde/HordeHttpClient.cs b/Engine/Source/Programs/Shared/EpicGames.Horde/HordeHttpClient.cs
--- a/Engine/Source/Programs/Shared/EpicGames.Horde/HordeHttpClient.cs
+++ b/Engine/Source/Programs/Shared/EpicGames.Horde/HordeHttpClient.cs
@@ -9,6 +9,7 @@ using System.Net.Http;
 using System.Net.Http.Headers;
 using System.Net.Http.Json;
 using System.Text.Json;
+using System.Text.Json.Nodes;
 using System.Text.RegularExpressions;
 using System.Threading;
 using System.Threading.Tasks;
@@ -26,6 +27,7 @@ using EpicGames.Horde.Secrets;
 using EpicGames.Horde.Server;
 using EpicGames.Horde.Storage;
 using EpicGames.Horde.Streams;
+using EpicGames.Horde.Telemetry;
 using EpicGames.Horde.Tools;
 using EpicGames.Horde.Ugs;
 
@@ -61,12 +63,12 @@ namespace EpicGames.Horde
 		/// Name of clients used for anonymous requests.
 		/// </summary>
 		public const string AnonymousHttpClientName = "HordeAnonymousHttpClient";
-		
+
 		/// <summary>
 		/// Name of clients created for storage operations
 		/// </summary>
 		public const string StorageHttpClientName = "HordeStorageHttpClient";
-		
+
 		/// <summary>
 		/// Name of clients created from the http client factory for handling upload redirects. Should not contain Horde auth headers.
 		/// </summary>
@@ -427,6 +429,57 @@ namespace EpicGames.Horde
 
 		#region Telemetry
 
+		/// <summary>
+		/// Posts telemetry to Horde.
+		/// </summary>
+		/// <param name="appId">The id of the App.</param>
+		/// <param name="appVersion">The version of the App.</param>
+		/// <param name="appEnvironment">The environment.</param>
+		/// <param name="payloads">A list of payloads to write as telemetry events.</param>
+		/// <param name="cancellationToken">The cancellation token to use throughout.</param>
+		/// <returns>The HttpResponseMessage.</returns>
+		/// <remarks>BadRequest can be returned under circusmtances where the payload is not serializable (primitive or null payload, <see cref="NotSupportedException"/>, and <see cref="JsonException"/>).</remarks>
+		public async Task<HttpResponseMessage> PostTelemetryAsync(string appId, int appVersion, string appEnvironment, object[] payloads, CancellationToken cancellationToken = default)
+		{
+			QueryStringBuilder queryParams = new QueryStringBuilder();
+			queryParams.Add("appId", appId);
+			queryParams.Add("appVersion", appVersion.ToString());
+			queryParams.Add("appEnvironment", appEnvironment);
+			queryParams.Add("uploadType", "EtEventStream");
+
+			PostTelemetryEventStreamRequest request = new PostTelemetryEventStreamRequest();
+
+			try
+			{
+				for (int i = 0; i < payloads.Length; ++i)
+				{
+					object payload = payloads[i];
+					JsonNode? node = JsonSerializer.SerializeToNode(payload, JsonSerializerOptions);
+
+					if (node != null && node is JsonObject obj)
+					{
+						request.Events.Add(obj);
+					}
+					else
+					{
+						return new HttpResponseMessage(HttpStatusCode.BadRequest);
+					}
+				}
+			}
+			catch(NotSupportedException) // The scenario where a type doesn't contain a JsonConverter.
+			{
+				return new HttpResponseMessage(HttpStatusCode.BadRequest);
+			}
+			catch (JsonException) // The scenario where a type encounters a json exception upon serialization.
+			{
+				return new HttpResponseMessage(HttpStatusCode.BadRequest);
+			}
+
+			HttpResponseMessage response = await PostAsync(_httpClient, $"/api/v1/telemetry?{queryParams}", request, cancellationToken);
+
+			return response;
+		}
+
 		/// <summary>
 		/// Gets telemetry for Horde within a given range
 		/// </summary>
@@ -510,7 +563,7 @@ namespace EpicGames.Horde
 		/// <summary>
 		/// Gets job information for given job ID. Fail response if jobID does not exist.
 		/// </summary>
-		/// <param name="id">Id of the job to get infomation for</param>
+		/// <param name="id">Id of the job to get information for</param>
 		/// <param name="cancellationToken">Cancellation token for the operation</param>
 		/// <returns></returns>
 		public Task<GetJobResponse> GetJobAsync(JobId id, CancellationToken cancellationToken = default)
@@ -528,8 +581,8 @@ namespace EpicGames.Horde
 		/// <returns></returns>
 		public Task PutJobMetadataAsync(JobId id, IEnumerable<string>? jobMetaData = null, Dictionary<string, List<string>>? stepMetaData = null, CancellationToken cancellationToken = default)
 		{
-			return PutAsync(_httpClient, $"api/v1/jobs/{id}/metadata", new PutJobMetadataRequest() {JobMetaData = jobMetaData?.ToList(), StepMetaData = stepMetaData}, cancellationToken);
-		}		
+			return PutAsync(_httpClient, $"api/v1/jobs/{id}/metadata", new PutJobMetadataRequest() { JobMetaData = jobMetaData?.ToList(), StepMetaData = stepMetaData }, cancellationToken);
+		}
 
 		#endregion
 
