EpicGames.UHT: Fix file delete path in CullOutputDirectory
Remove some blank lines

#rnx
#Synced-CL 49130599

diff --git a/Engine/Source/Programs/Shared/EpicGames.UHT/Utils/UhtSession.cs b/Engine/Source/Programs/Shared/EpicGames.UHT/Utils/UhtSession.cs
--- a/Engine/Source/Programs/Shared/EpicGames.UHT/Utils/UhtSession.cs
+++ b/Engine/Source/Programs/Shared/EpicGames.UHT/Utils/UhtSession.cs
@@ -584,7 +584,7 @@ namespace EpicGames.UHT.Utils
 				{
 					try
 					{
-						File.Delete(Path.Combine(outputDirectory, filePath));
+						File.Delete(filePath);
 					}
 					catch (Exception)
 					{
@@ -1048,7 +1048,6 @@ namespace EpicGames.UHT.Utils
 		}
 		#endregion
 
-
 		#region Statistics
 		/// <summary>
 		/// The number of errors
