EpicGames.Horde: CA1513: Use ObjectDisposedException throw helper

#rnx

#Synced-CL 49130599

diff --git a/Engine/Source/Programs/Shared/EpicGames.Horde/Storage/Bundles/V1/BundleWriter.cs b/Engine/Source/Programs/Shared/EpicGames.Horde/Storage/Bundles/V1/BundleWriter.cs
--- a/Engine/Source/Programs/Shared/EpicGames.Horde/Storage/Bundles/V1/BundleWriter.cs
+++ b/Engine/Source/Programs/Shared/EpicGames.Horde/Storage/Bundles/V1/BundleWriter.cs
@@ -815,10 +815,7 @@ namespace EpicGames.Horde.Storage.Bundles.V1
 		/// <returns></returns>
 		public override async Task FlushAsync(CancellationToken cancellationToken = default)
 		{
-			if (_disposed)
-			{
-				throw new ObjectDisposedException(GetType().Name);
-			}
+			ObjectDisposedException.ThrowIf(_disposed, this);
 
 			await CompleteAsync(cancellationToken);
 			await _writeQueue.FlushAsync(cancellationToken);
