Shared: CA1859: Use concrete types when possible for improved performance
CA1860: Avoid using 'Enumerable.Any()' extension method
CA1864: Prefer the 'IDictionary.TryAdd(TKey, TValue)' method
CA1865-CA1867: Use 'string.Method(char)' instead of 'string.Method(string)' for string with single char
IDE2000: Avoid multiple blank lines
Other misc formatting fixes

#rnx
#Synced-CL 49186938

diff --git a/Engine/Source/Programs/Shared/EpicGames.AspNet/ServerTimings.cs b/Engine/Source/Programs/Shared/EpicGames.AspNet/ServerTimings.cs
--- a/Engine/Source/Programs/Shared/EpicGames.AspNet/ServerTimings.cs
+++ b/Engine/Source/Programs/Shared/EpicGames.AspNet/ServerTimings.cs
@@ -135,7 +135,7 @@ namespace EpicGames.AspNet
 			await _next(context);
 
 			// we limit the server timing header to 10 metrics because otherwise we risk generating very large response headers for operations that do a lot of work
-			string serverTimingValue = serverTiming.Metrics.Any() ? String.Join(",", serverTiming.Metrics.Take(10)) : "";
+			string serverTimingValue = serverTiming.Metrics.Count != 0 ? String.Join(",", serverTiming.Metrics.Take(10)) : "";
 			context.Response.AppendTrailer(
 				"Server-Timing",
 				serverTimingValue);
@@ -145,7 +145,7 @@ namespace EpicGames.AspNet
 		{
 			context.Response.OnStarting(() =>
 			{
-				if (serverTiming.Metrics.Any())
+				if (serverTiming.Metrics.Count != 0)
 				{
 					string serverTimingValue = String.Join(",", serverTiming.Metrics.Take(10));
 					context.Response.Headers.Append("Server-Timing", serverTimingValue);
diff --git a/Engine/Source/Programs/Shared/EpicGames.Build/CsProjectInfo.cs b/Engine/Source/Programs/Shared/EpicGames.Build/CsProjectInfo.cs
--- a/Engine/Source/Programs/Shared/EpicGames.Build/CsProjectInfo.cs
+++ b/Engine/Source/Programs/Shared/EpicGames.Build/CsProjectInfo.cs
@@ -126,7 +126,7 @@ namespace EpicGames.Core
 		public static CsProjectInfo Read(FileReference file, Dictionary<string, string> properties)
 		{
 			CsProjectInfo? project;
-			if(!TryRead(file, properties, out project))
+			if (!TryRead(file, properties, out project))
 			{
 				throw new Exception(String.Format("Unable to read '{0}'", file));
 			}
@@ -243,14 +243,14 @@ namespace EpicGames.Core
 							continue;
 						}
 					case "IsEngineProject":
-					{
-						projectInfo.Properties[element.Name] = "False";
-						if (projectInfo.Properties.TryGetValue("EngineDirectory", out string? engineDirectory) && !String.IsNullOrEmpty(engineDirectory))
 						{
-							projectInfo.Properties[element.Name] = projectInfo.ProjectPath.IsUnderDirectory(new DirectoryReference(engineDirectory)).ToString();
+							projectInfo.Properties[element.Name] = "False";
+							if (projectInfo.Properties.TryGetValue("EngineDirectory", out string? engineDirectory) && !String.IsNullOrEmpty(engineDirectory))
+							{
+								projectInfo.Properties[element.Name] = projectInfo.ProjectPath.IsUnderDirectory(new DirectoryReference(engineDirectory)).ToString();
+							}
+							continue;
 						}
-						continue;
-					}
 					case "IsAutomationProject": projectInfo.Properties[element.Name] = projectInfo.ProjectPath.FullName.EndsWith(".automation.csproj", StringComparison.OrdinalIgnoreCase).ToString(); continue;
 					case "HasAssemblyInfo": projectInfo.Properties[element.Name] = FileReference.Exists(FileReference.Combine(projectInfo.ProjectPath.Directory, "Properties", "AssemblyInfo.cs")).ToString(); continue;
 					default: break;
@@ -322,7 +322,7 @@ namespace EpicGames.Core
 			{
 				// Don't include embedded assemblies; they aren't referenced externally by the compiled executable
 				bool bEmbedInteropTypes = GetChildElementBoolean(parentElement, "EmbedInteropTypes", false);
-				if(!bEmbedInteropTypes)
+				if (!bEmbedInteropTypes)
 				{
 					if (!OperatingSystem.IsWindows() && hintPath.Contains('\\', StringComparison.Ordinal))
 					{
@@ -509,7 +509,7 @@ namespace EpicGames.Core
 		/// <returns></returns>
 		static IEnumerable<FileReference> FindMatchingFiles(FileReference inPath)
 		{
-			List<FileReference> foundFiles = [];		
+			List<FileReference> foundFiles = [];
 
 			// split off the drive root
 			string driveRoot = Path.GetPathRoot(inPath.FullName)!;
@@ -779,17 +779,17 @@ namespace EpicGames.Core
 			{
 				// Find the end of the variable name, accounting for changes in scope
 				int endIdx = idx + 2;
-				for(int depth = 1; depth > 0; endIdx++)
+				for (int depth = 1; depth > 0; endIdx++)
 				{
-					if(endIdx == newText.Length)
+					if (endIdx == newText.Length)
 					{
 						throw new Exception("Encountered end of string while expanding properties");
 					}
-					else if(newText[endIdx] == '(')
+					else if (newText[endIdx] == '(')
 					{
 						depth++;
 					}
-					else if(newText[endIdx] == ')')
+					else if (newText[endIdx] == ')')
 					{
 						depth--;
 					}
@@ -799,7 +799,7 @@ namespace EpicGames.Core
 				string[] tokens = Tokenize(newText.Substring(idx + 2, (endIdx - 1) - (idx + 2)));
 
 				// Make sure the first token is a valid property name
-				if(tokens.Length == 0 || !(Char.IsLetter(tokens[0][0]) || tokens[0][0] == '_' || tokens[0][0] == '[' ))
+				if (tokens.Length == 0 || !(Char.IsLetter(tokens[0][0]) || tokens[0][0] == '_' || tokens[0][0] == '['))
 				{
 					throw new Exception(String.Format("Invalid property name '{0}' in .csproj file", tokens[0]));
 				}
@@ -817,7 +817,7 @@ namespace EpicGames.Core
 
 				// Evaluate any functions within it
 				int tokenIdx = 1;
-				while(tokenIdx + 3 < tokens.Length && tokens[tokenIdx] == "." && tokens[tokenIdx + 2] == "(")
+				while (tokenIdx + 3 < tokens.Length && tokens[tokenIdx] == "." && tokens[tokenIdx + 2] == "(")
 				{
 					// Read the method name
 					string methodName = tokens[tokenIdx + 1];
@@ -827,18 +827,18 @@ namespace EpicGames.Core
 
 					// Parse any arguments
 					List<object> arguments = [];
-					if(tokens[tokenIdx] != ")")
+					if (tokens[tokenIdx] != ")")
 					{
 						arguments.Add(ParseArgument(tokens[tokenIdx]));
 						tokenIdx++;
 
-						while(tokenIdx + 1 < tokens.Length && tokens[tokenIdx] == ",")
+						while (tokenIdx + 1 < tokens.Length && tokens[tokenIdx] == ",")
 						{
 							arguments.Add(ParseArgument(tokens[tokenIdx + 2]));
 							tokenIdx += 2;
 						}
 
-						if(tokens[tokenIdx] != ")")
+						if (tokens[tokenIdx] != ")")
 						{
 							throw new Exception("Missing closing parenthesis in condition");
 						}
@@ -852,7 +852,7 @@ namespace EpicGames.Core
 					{
 						value = typeof(string).InvokeMember(methodName, System.Reflection.BindingFlags.Instance | System.Reflection.BindingFlags.Public | System.Reflection.BindingFlags.InvokeMethod, Type.DefaultBinder, value, [.. arguments])!.ToString()!;
 					}
-					catch(Exception ex)
+					catch (Exception ex)
 					{
 						throw new Exception(String.Format("Unable to evaluate condition '{0}'", text), ex);
 					}
@@ -864,7 +864,7 @@ namespace EpicGames.Core
 				}
 
 				// Make sure there's nothing left over
-				if(tokenIdx != tokens.Length)
+				if (tokenIdx != tokens.Length)
 				{
 					throw new Exception(String.Format("Unable to parse token '{0}'", newText));
 				}
@@ -886,14 +886,14 @@ namespace EpicGames.Core
 		static object ParseArgument(string token)
 		{
 			// Try to parse a string
-			if(token.Length > 2 && token[0] == '\'' && token[^1] == '\'')
+			if (token.Length > 2 && token[0] == '\'' && token[^1] == '\'')
 			{
 				return token.Substring(1, token.Length - 2);
 			}
 
 			// Try to parse an integer
 			int value;
-			if(Int32.TryParse(token, out value))
+			if (Int32.TryParse(token, out value))
 			{
 				return value;
 			}
@@ -910,9 +910,9 @@ namespace EpicGames.Core
 		static string[] Tokenize(string condition)
 		{
 			List<string> tokens = [];
-			for (int idx = 0; idx < condition.Length; )
+			for (int idx = 0; idx < condition.Length;)
 			{
-				if(Char.IsWhiteSpace(condition[idx]))
+				if (Char.IsWhiteSpace(condition[idx]))
 				{
 					// Whitespace
 					idx++;
@@ -933,13 +933,13 @@ namespace EpicGames.Core
 				{
 					// Quoted string
 					int startIdx = idx++;
-					for(;;idx++)
+					for (; ; idx++)
 					{
-						if(idx == condition.Length)
+						if (idx == condition.Length)
 						{
 							throw new Exception(String.Format("Missing end quote in condition string ('{0}')", condition));
 						}
-						if(condition[idx] == '\'')
+						if (condition[idx] == '\'')
 						{
 							break;
 						}
@@ -1018,11 +1018,11 @@ namespace EpicGames.Core
 
 					tokens.Add(condition.Substring(startIdx, idx - startIdx));
 				}
-				else if(Char.IsLetterOrDigit(condition[idx]) || condition[idx] == '_')
+				else if (Char.IsLetterOrDigit(condition[idx]) || condition[idx] == '_')
 				{
 					// Identifier or number
 					int startIdx = idx++;
-					while(idx < condition.Length && (Char.IsLetterOrDigit(condition[idx]) || condition[idx] == '_'))
+					while (idx < condition.Length && (Char.IsLetterOrDigit(condition[idx]) || condition[idx] == '_'))
 					{
 						idx++;
 					}
@@ -1048,17 +1048,17 @@ namespace EpicGames.Core
 			const string HexChars = "0123456789abcdef";
 
 			string? newText = text;
-			if(newText != null)
+			if (newText != null)
 			{
-				for(int idx = 0; idx + 2 < newText.Length; idx++)
+				for (int idx = 0; idx + 2 < newText.Length; idx++)
 				{
-					if(newText[idx] == '%')
+					if (newText[idx] == '%')
 					{
 						int upperDigitIdx = HexChars.IndexOf(Char.ToLowerInvariant(newText[idx + 1]), StringComparison.Ordinal);
-						if(upperDigitIdx != -1)
+						if (upperDigitIdx != -1)
 						{
 							int lowerDigitIdx = HexChars.IndexOf(Char.ToLowerInvariant(newText[idx + 2]), StringComparison.Ordinal);
-							if(lowerDigitIdx != -1)
+							if (lowerDigitIdx != -1)
 							{
 								char newChar = (char)((upperDigitIdx << 4) | lowerDigitIdx);
 								newText = newText.Substring(0, idx) + newChar + newText.Substring(idx + 3);
diff --git a/Engine/Source/Programs/Shared/EpicGames.Build/Platform/UEBuildPlatformSDK.cs b/Engine/Source/Programs/Shared/EpicGames.Build/Platform/UEBuildPlatformSDK.cs
--- a/Engine/Source/Programs/Shared/EpicGames.Build/Platform/UEBuildPlatformSDK.cs
+++ b/Engine/Source/Programs/Shared/EpicGames.Build/Platform/UEBuildPlatformSDK.cs
@@ -846,7 +846,7 @@ namespace EpicGames.Core
 							Logger.LogWarning("Project {Project} is overriding {Platform} SDK to {OverrideMainVersion}", projectFile.GetFileNameWithoutAnyExtensions(), platform, overrideMainVersion);
 
 							ConfigSDKVersions[MainVersionKey] = overrideMainVersion;
-							UEBuildPlatformSDK.bHasAnySDKOverride = true;
+							bHasAnySDKOverride = true;
 							bHasSDKOverride = true;
 							ProjectsThatOverrodeSDK.Add(projectFile);
 						}
diff --git a/Engine/Source/Programs/Shared/EpicGames.Build/System/CompileScriptModules.cs b/Engine/Source/Programs/Shared/EpicGames.Build/System/CompileScriptModules.cs
--- a/Engine/Source/Programs/Shared/EpicGames.Build/System/CompileScriptModules.cs
+++ b/Engine/Source/Programs/Shared/EpicGames.Build/System/CompileScriptModules.cs
@@ -94,7 +94,7 @@ namespace UnrealBuildBase
 
 			public void ValidateRecursively(Dictionary<FileReference, CsProjBuildRecordEntry> BuildRecords, FileReference ProjectPath)
 			{
-				CompileScriptModule.ValidateBuildRecordRecursively(BuildRecords, ProjectPath, this, Logger);
+				ValidateBuildRecordRecursively(BuildRecords, ProjectPath, this, Logger);
 			}
 
 			public void SetValidBuildRecords(IReadOnlyDictionary<FileReference, CsProjBuildRecordEntry> Records)
@@ -460,8 +460,8 @@ namespace UnrealBuildBase
 
 		private static bool ValidateGlobbedFiles(DirectoryReference ProjectDirectory, List<CsProjBuildRecord.Glob> Globs, HashSet<string> Dependencies, out string ValidationFailureMessage)
 		{
-			string[] AllFiles = [..Directory.EnumerateFiles(ProjectDirectory.FullName, "*", SearchOption.AllDirectories)];
-			
+			string[] AllFiles = [.. Directory.EnumerateFiles(ProjectDirectory.FullName, "*", SearchOption.AllDirectories)];
+
 			// First, evaluate globs
 
 			// Files are grouped by ItemType (e.g. Compile, EmbeddedResource) to ensure that Exclude and
@@ -485,7 +485,7 @@ namespace UnrealBuildBase
 				{
 					matcher.AddExcludePatterns(Glob.Exclude);
 				}
-				
+
 				// Not using GetResultsInFullPath, as internally it's super slow
 				// Match returns filenames separated with '/'
 				if (OperatingSystem.IsWindows())
@@ -510,7 +510,7 @@ namespace UnrealBuildBase
 			// Then, validation that our evaluation matches what we're comparing against
 
 			// Look for extra files that were found.
-			HashSet<string> newFiles = [..Files.Values.SelectMany(x => x)];
+			HashSet<string> newFiles = [.. Files.Values.SelectMany(x => x)];
 			newFiles.ExceptWith(Dependencies);
 			ValidationFailureMessage = newFiles.Count > 0 ? $"Found additional files: {String.Join(Environment.NewLine, newFiles.Order())}" : String.Empty;
 
@@ -557,7 +557,7 @@ namespace UnrealBuildBase
 					return false;
 				}
 			}
-			
+
 			if (!ValidateGlobbedFiles(ProjectDirectory, BuildRecord.Globs, BuildRecord.Dependencies, out ValidationFailureMessage))
 			{
 				return false;
@@ -603,7 +603,7 @@ namespace UnrealBuildBase
 			}
 
 			// Are all referenced build records valid?
-			foreach (CsProjBuildRecordRef ReferencedProject in Entry.BuildRecord. ProjectReferencesAndTimes)
+			foreach (CsProjBuildRecordRef ReferencedProject in Entry.BuildRecord.ProjectReferencesAndTimes)
 			{
 				FileReference FullProjectPath = FileReference.FromString(Path.GetFullPath(ReferencedProject.ProjectPath, ProjectPath.Directory.FullName));
 				ValidateBuildRecordRecursively(BuildRecords, FullProjectPath, Hook, Logger);
diff --git a/Engine/Source/Programs/Shared/EpicGames.Build/System/DirectoryItem.cs b/Engine/Source/Programs/Shared/EpicGames.Build/System/DirectoryItem.cs
--- a/Engine/Source/Programs/Shared/EpicGames.Build/System/DirectoryItem.cs
+++ b/Engine/Source/Programs/Shared/EpicGames.Build/System/DirectoryItem.cs
@@ -62,7 +62,7 @@ namespace UnrealBuildBase
 		public DirectoryItem? GetParentDirectoryItem()
 		{
 			DirectoryReference? parent = Location.ParentDirectory;
-			return parent != null ? DirectoryItem.GetItemByDirectoryReference(parent) : null;
+			return parent != null ? GetItemByDirectoryReference(parent) : null;
 		}
 
 		/// <summary>
@@ -130,7 +130,7 @@ namespace UnrealBuildBase
 			{
 				FileReference fileRef = FileReference.Combine(Location, file);
 				FileItem fileItem = FileItem.GetItemByFileReference(fileRef);
-				fileItem.Update(new() { Attributes = FileAttributes.Normal, Size = long.MaxValue }, true); 
+				fileItem.Update(new() { Attributes = FileAttributes.Normal, Size = long.MaxValue }, true);
 				cacheFiles.Add(file, fileItem);
 			}
 
@@ -233,7 +233,7 @@ namespace UnrealBuildBase
 		/// Scans the directory for directories and files, used for lazy initialization.
 		/// </summary>
 		/// <returns></returns>
-		(Dictionary<string, DirectoryItem> , Dictionary<string, FileItem>) Scan()
+		(Dictionary<string, DirectoryItem>, Dictionary<string, FileItem>) Scan()
 		{
 			if (_info != null && !_info.Exists)
 			{
@@ -258,8 +258,8 @@ namespace UnrealBuildBase
 					}
 				}
 
-				Dictionary<string, DirectoryItem>? newDirs = new(dirCount, FileReference.Comparer);
-				Dictionary<string, FileItem>? newFiles = new(entries.Count - dirCount, FileReference.Comparer);
+				Dictionary<string, DirectoryItem>? newDirs = new(dirCount, FileSystemReference.Comparer);
+				Dictionary<string, FileItem>? newFiles = new(entries.Count - dirCount, FileSystemReference.Comparer);
 
 				foreach (DirectoryEnumerator.Entry entry in entries)
 				{
@@ -355,7 +355,7 @@ namespace UnrealBuildBase
 		public static void WriteDebugFileWithAllEnumeratedFiles(string OutFile)
 		{
 			SortedSet<string> AllFiles = [];
-			foreach (DirectoryItem Item in DirectoryItem.LocationToItem.Values)
+			foreach (DirectoryItem Item in LocationToItem.Values)
 			{
 				if (Item.Files != null)
 				{
diff --git a/Engine/Source/Programs/Shared/EpicGames.Build/System/EnumeratePlugins.cs b/Engine/Source/Programs/Shared/EpicGames.Build/System/EnumeratePlugins.cs
--- a/Engine/Source/Programs/Shared/EpicGames.Build/System/EnumeratePlugins.cs
+++ b/Engine/Source/Programs/Shared/EpicGames.Build/System/EnumeratePlugins.cs
@@ -40,7 +40,7 @@ namespace UnrealBuildBase
 		public static List<FileReference> EnumeratePlugins(FileReference? ProjectFile)
 		{
 			List<DirectoryReference> BaseDirs = [.. Unreal.GetExtensionDirs(Unreal.EngineDirectory, "Plugins")];
-			if(ProjectFile != null)
+			if (ProjectFile != null)
 			{
 				BaseDirs.AddRange(Unreal.GetExtensionDirs(ProjectFile.Directory, "Plugins"));
 				BaseDirs.AddRange(Unreal.GetExtensionDirs(ProjectFile.Directory, "Mods"));
@@ -85,9 +85,9 @@ namespace UnrealBuildBase
 				bool bSearchSubDirectories = true;
 				foreach (FileItem PluginFile in ChildDirectory.EnumerateFiles())
 				{
-					if(PluginFile.HasExtension(".uplugin"))
+					if (PluginFile.HasExtension(".uplugin"))
 					{
-						lock(FileNames)
+						lock (FileNames)
 						{
 							FileNames.Add(PluginFile.Location);
 						}
@@ -127,7 +127,7 @@ namespace UnrealBuildBase
 				PluginDirectories.AddRange(Unreal.GetExtensionDirs(ProjectFile.Directory, "Mods"));
 			}
 			Rules.Populate(PluginDirectories);
-			List<FileReference> Plugins = Rules.Get(PluginDirectories, UnrealBuildBase.Rules.RulesFileType.UbtPlugin);
+			List<FileReference> Plugins = Rules.Get(PluginDirectories, Rules.RulesFileType.UbtPlugin);
 			return Plugins;
 		}
 
@@ -182,7 +182,7 @@ namespace UnrealBuildBase
 			Both,
 		}
 
-		private static bool BuildUbtPluginsInternal(PluginSet PluginSet, FileReference? ProjectFile, 
+		private static bool BuildUbtPluginsInternal(PluginSet PluginSet, FileReference? ProjectFile,
 			List<FileReference> Built, IEnumerable<FileReference> NotBuilt, IEnumerable<string>? DefineConstants, ILogger Logger)
 		{
 			// Collect the work to be done
@@ -222,7 +222,7 @@ namespace UnrealBuildBase
 			// Build the things needing to be built
 			bool bBuildSuccess = true;
 			Dictionary<FileReference, CsProjBuildRecordEntry>? BuiltPlugins =
-				CompileScriptModule.Build(UnrealBuildBase.Rules.RulesFileType.UbtPlugin, [.. NotBuilt],
+				CompileScriptModule.Build(Rules.RulesFileType.UbtPlugin, [.. NotBuilt],
 				BaseDirectories, DefineConstants, CompileScriptModule.BuildFlags.UseBuildRecords, out bBuildSuccess, Logger);
 
 			// Add any built plugins back into the cache
diff --git a/Engine/Source/Programs/Shared/EpicGames.Build/System/Rules.cs b/Engine/Source/Programs/Shared/EpicGames.Build/System/Rules.cs
--- a/Engine/Source/Programs/Shared/EpicGames.Build/System/Rules.cs
+++ b/Engine/Source/Programs/Shared/EpicGames.Build/System/Rules.cs
@@ -125,8 +125,8 @@ namespace UnrealBuildBase
 		/// Invalidate the cache for the givcen directory
 		/// </summary>
 		/// <param name="DirectoryPath">Directory to invalidate</param>
-        public static void InvalidateRulesFileCache(string DirectoryPath)
-        {
+		public static void InvalidateRulesFileCache(string DirectoryPath)
+		{
 			DirectoryReference dirRef = new DirectoryReference(DirectoryPath);
 			lock (Root)
 			{
@@ -165,7 +165,7 @@ namespace UnrealBuildBase
 			Directory parentDir = Root;
 			for (int i = 0; i != parts.Length; ++i)
 			{
-				parentDir.SubDirectories ??= new(DirectoryReference.Comparer);
+				parentDir.SubDirectories ??= new(FileSystemReference.Comparer);
 				if (!parentDir.SubDirectories.TryGetValue(parts[i], out Directory? dir))
 				{
 					if (!createIfNotFound)
@@ -375,7 +375,7 @@ namespace UnrealBuildBase
 				else
 				{
 					subDir = new();
-					dir.SubDirectories ??= new(DirectoryReference.Comparer);
+					dir.SubDirectories ??= new(FileSystemReference.Comparer);
 					dir.SubDirectories.Add(subDirItem.Name, subDir);
 				}
 				subDir.IsScanned = true;
diff --git a/Engine/Source/Programs/Shared/EpicGames.BuildGraph/BgGraphDef.cs b/Engine/Source/Programs/Shared/EpicGames.BuildGraph/BgGraphDef.cs
--- a/Engine/Source/Programs/Shared/EpicGames.BuildGraph/BgGraphDef.cs
+++ b/Engine/Source/Programs/Shared/EpicGames.BuildGraph/BgGraphDef.cs
@@ -144,7 +144,7 @@ namespace EpicGames.BuildGraph
 		public bool TryResolveReference(string name, [NotNullWhen(true)] out BgNodeDef[]? outNodes)
 		{
 			// Check if it's a tag reference or node reference
-			if (name.StartsWith("#", StringComparison.Ordinal))
+			if (name.StartsWith('#'))
 			{
 				// Check if it's a regular node or output name
 				BgNodeOutput? output;
@@ -195,7 +195,7 @@ namespace EpicGames.BuildGraph
 		public bool TryResolveInputReference(string name, [NotNullWhen(true)] out BgNodeOutput[]? outOutputs)
 		{
 			// Check if it's a tag reference or node reference
-			if (name.StartsWith("#", StringComparison.Ordinal))
+			if (name.StartsWith('#'))
 			{
 				// Check if it's a regular node or output name
 				BgNodeOutput? output;
diff --git a/Engine/Source/Programs/Shared/EpicGames.BuildGraph/Expressions/BgList.cs b/Engine/Source/Programs/Shared/EpicGames.BuildGraph/Expressions/BgList.cs
--- a/Engine/Source/Programs/Shared/EpicGames.BuildGraph/Expressions/BgList.cs
+++ b/Engine/Source/Programs/Shared/EpicGames.BuildGraph/Expressions/BgList.cs
@@ -135,7 +135,7 @@ namespace EpicGames.BuildGraph.Expressions
 		/// <returns></returns>
 		public static BgList<T> Concat(params BgList<T>[] others)
 		{
-			BgList<T> result = BgList<T>.Empty;
+			BgList<T> result = Empty;
 			for (int idx = 0; idx < others.Length; idx++)
 			{
 				result = Concat(result, others[idx]);
diff --git a/Engine/Source/Programs/Shared/EpicGames.BuildGraph/Expressions/BgType.cs b/Engine/Source/Programs/Shared/EpicGames.BuildGraph/Expressions/BgType.cs
--- a/Engine/Source/Programs/Shared/EpicGames.BuildGraph/Expressions/BgType.cs
+++ b/Engine/Source/Programs/Shared/EpicGames.BuildGraph/Expressions/BgType.cs
@@ -66,7 +66,7 @@ namespace EpicGames.BuildGraph.Expressions
 	{
 		class BgTypeCache<T> where T : BgExpr
 		{
-			public static BgType<T> Instance { get; } = (BgType<T>)BgType.CreateInstance(typeof(T));
+			public static BgType<T> Instance { get; } = (BgType<T>)CreateInstance(typeof(T));
 		}
 
 		static readonly ConcurrentDictionary<Type, IBgType> s_typeCache = new ConcurrentDictionary<Type, IBgType>();
diff --git a/Engine/Source/Programs/Shared/EpicGames.Core.Tests/JsonObjectTests.cs b/Engine/Source/Programs/Shared/EpicGames.Core.Tests/JsonObjectTests.cs
--- a/Engine/Source/Programs/Shared/EpicGames.Core.Tests/JsonObjectTests.cs
+++ b/Engine/Source/Programs/Shared/EpicGames.Core.Tests/JsonObjectTests.cs
@@ -32,9 +32,9 @@ namespace EpicGames.Core.Tests
 		[TestCleanup]
 		public void RemoveTempDir()
 		{
-			if (Directory.Exists(JsonObjectTests.s_tempDirectory.FullName))
+			if (Directory.Exists(s_tempDirectory.FullName))
 			{
-				Directory.Delete(JsonObjectTests.s_tempDirectory.FullName, true);
+				Directory.Delete(s_tempDirectory.FullName, true);
 			}
 		}
 
@@ -438,9 +438,9 @@ namespace EpicGames.Core.Tests
 		public void SetObjectArrayField()
 		{
 			string fieldName = "Attacks";
-			EnemyAttack[] attacks = 
-			[ 
-				new EnemyAttack("Attack1", 40, StatusEffect.Sleep), 
+			EnemyAttack[] attacks =
+			[
+				new EnemyAttack("Attack1", 40, StatusEffect.Sleep),
 				new EnemyAttack("Attack2", 100, StatusEffect.Burn)
 			];
 			Assert.IsTrue(_validTestObject.ContainsField(fieldName));
@@ -532,7 +532,7 @@ namespace EpicGames.Core.Tests
 		[TestMethod]
 		public void Constructor()
 		{
-			JsonObject obj= new JsonObject();
+			JsonObject obj = new JsonObject();
 			Assert.AreEqual(0, obj.KeyNames.Count());
 		}
 
@@ -570,7 +570,7 @@ namespace EpicGames.Core.Tests
 			JsonObject jsonObject = new JsonObject();
 			jsonObject.AddOrSetFieldValue(escapeCharacterFieldName, escapeString);
 			string jsonString = jsonObject.ToJsonString();
-			Assert.AreEqual(correctString , jsonString);
+			Assert.AreEqual(correctString, jsonString);
 		}
 		[TestMethod]
 		public void KeyNameCaseSensitivity()
diff --git a/Engine/Source/Programs/Shared/EpicGames.Core/AssemblyUtils.cs b/Engine/Source/Programs/Shared/EpicGames.Core/AssemblyUtils.cs
--- a/Engine/Source/Programs/Shared/EpicGames.Core/AssemblyUtils.cs
+++ b/Engine/Source/Programs/Shared/EpicGames.Core/AssemblyUtils.cs
@@ -36,13 +36,14 @@ namespace EpicGames.Core
 		/// </summary>
 		public static void InstallAssemblyResolver(string pathToBinariesDotNET)
 		{
+			string[] knownAssemblyNames = [ "SwarmAgent.exe", "../ThirdParty/Ionic/Ionic.Zip.Reduced.dll", "../ThirdParty/Newtonsoft/NewtonSoft.Json.dll" ];
 			AppDomain.CurrentDomain.AssemblyResolve += (sender, args) =>
 			{
 				// Name is fully qualified assembly definition - e.g. "p4dn, Version=1.0.0.0, Culture=neutral, PublicKeyToken=ff968dc1933aba6f"
 				string assemblyName = args.Name!.Split(',')[0];
 
 				return (
-					from knownAssemblyName in new[] { "SwarmAgent.exe", "../ThirdParty/Ionic/Ionic.Zip.Reduced.dll", "../ThirdParty/Newtonsoft/NewtonSoft.Json.dll" }
+					from knownAssemblyName in knownAssemblyNames
 					where assemblyName.Equals(Path.GetFileNameWithoutExtension(knownAssemblyName), StringComparison.OrdinalIgnoreCase)
 					let resolvedAssemblyFilename = Path.Combine(pathToBinariesDotNET, knownAssemblyName)
 					// check if the file exists first. If we just try to load it, we correctly throw an exception, but it's a generic
@@ -103,9 +104,9 @@ namespace EpicGames.Core
 
 			string assemblyName = Path.GetFileNameWithoutExtension(assemblyPath);
 			DateTime assemblyLastWriteTime = File.GetLastWriteTimeUtc(assemblyPath);
-			lock(s_assemblyLocationCache)
+			lock (s_assemblyLocationCache)
 			{
-				if (s_assemblyLocationCache.ContainsKey(assemblyName))
+				if (!s_assemblyLocationCache.TryAdd(assemblyName, assemblyPath))
 				{
 					// We already have this assembly in our cache. Only replace it if the discovered file is newer (to avoid stale assemblies breaking stuff).
 					if (assemblyLastWriteTime > s_assemblyWriteTimes[assemblyName])
@@ -116,8 +117,6 @@ namespace EpicGames.Core
 				}
 				else
 				{
-					// This is the first copy of this assembly ... add it to our cache.
-					s_assemblyLocationCache.Add(assemblyName, assemblyPath);
 					s_assemblyWriteTimes.Add(assemblyName, assemblyLastWriteTime);
 				}
 			}
@@ -132,7 +131,7 @@ namespace EpicGames.Core
 					{
 						if (s_assemblyLocationCache.TryGetValue(assemblyName, out string? assemblyLocation))
 						{
-							// We have this assembly in our folder.					
+							// We have this assembly in our folder.
 							if (File.Exists(assemblyLocation))
 							{
 								// The assembly still exists, so load it.
diff --git a/Engine/Source/Programs/Shared/EpicGames.Core/BinaryArchiveReader.cs b/Engine/Source/Programs/Shared/EpicGames.Core/BinaryArchiveReader.cs
--- a/Engine/Source/Programs/Shared/EpicGames.Core/BinaryArchiveReader.cs
+++ b/Engine/Source/Programs/Shared/EpicGames.Core/BinaryArchiveReader.cs
@@ -294,7 +294,7 @@ namespace EpicGames.Core
 		/// <param name="size">Size of data to read</param>
 		private void ReadBulkData(Array data, int size)
 		{
-			System.Buffer.BlockCopy(_buffer!, _bufferPos, data, 0, size);
+			Buffer.BlockCopy(_buffer!, _bufferPos, data, 0, size);
 			_bufferPos += size;
 		}
 
diff --git a/Engine/Source/Programs/Shared/EpicGames.Core/BinaryArchiveWriter.cs b/Engine/Source/Programs/Shared/EpicGames.Core/BinaryArchiveWriter.cs
--- a/Engine/Source/Programs/Shared/EpicGames.Core/BinaryArchiveWriter.cs
+++ b/Engine/Source/Programs/Shared/EpicGames.Core/BinaryArchiveWriter.cs
@@ -345,7 +345,7 @@ namespace EpicGames.Core
 				{
 					int copySize = Math.Min(size - pos, _buffer.Length - _bufferPos);
 
-					System.Buffer.BlockCopy(data, pos, _buffer, _bufferPos, copySize);
+					Buffer.BlockCopy(data, pos, _buffer, _bufferPos, copySize);
 					_bufferPos += copySize;
 					pos += copySize;
 
diff --git a/Engine/Source/Programs/Shared/EpicGames.Core/BuzHash.cs b/Engine/Source/Programs/Shared/EpicGames.Core/BuzHash.cs
--- a/Engine/Source/Programs/Shared/EpicGames.Core/BuzHash.cs
+++ b/Engine/Source/Programs/Shared/EpicGames.Core/BuzHash.cs
@@ -218,7 +218,7 @@ namespace EpicGames.Core
 			int offset = 0;
 			for (; ; )
 			{
-				int length = BuzHash.FindChunkLength(span[offset..], minSize, maxSize, targetSize, out uint rollingHash);
+				int length = FindChunkLength(span[offset..], minSize, maxSize, targetSize, out uint rollingHash);
 
 				offset += length;
 				if (offset >= span.Length)
@@ -327,7 +327,7 @@ namespace EpicGames.Core
 
 			// Fast path for appending data to the buffer up to the chunk window size
 			int length = Math.Min(windowSize, inputData.Length);
-			rollingHash = BuzHash.Add(0, inputData.Slice(0, length));
+			rollingHash = Add(0, inputData.Slice(0, length));
 
 			// Get the threshold for the rolling hash to split the output
 			uint rollingHashThreshold = (uint)((1L << 32) / (targetSize - minSize));
@@ -340,7 +340,7 @@ namespace EpicGames.Core
 				ReadOnlySpan<byte> tailSpan = inputData.Slice(length - windowSize, inputData.Length - windowSize);
 				ReadOnlySpan<byte> headSpan = inputData.Slice(length);
 
-				int count = BuzHash.Update(tailSpan, headSpan, minSize, rollingHashThreshold, ref rollingHash);
+				int count = Update(tailSpan, headSpan, minSize, rollingHashThreshold, ref rollingHash);
 				if (count != -1)
 				{
 					length += count;
diff --git a/Engine/Source/Programs/Shared/EpicGames.Core/CommandLineArguments.cs b/Engine/Source/Programs/Shared/EpicGames.Core/CommandLineArguments.cs
--- a/Engine/Source/Programs/Shared/EpicGames.Core/CommandLineArguments.cs
+++ b/Engine/Source/Programs/Shared/EpicGames.Core/CommandLineArguments.cs
@@ -149,7 +149,7 @@ namespace EpicGames.Core
 
 			// Clear the linked list of identical arguments
 			_nextArgumentIndex = new int[arguments.Length];
-			for(int idx = 0; idx < arguments.Length; idx++)
+			for (int idx = 0; idx < arguments.Length; idx++)
 			{
 				_nextArgumentIndex[idx] = -1;
 			}
@@ -159,7 +159,7 @@ namespace EpicGames.Core
 
 			// Parse the argument array and build a lookup
 			_argumentToFirstIndex = new Dictionary<string, int>(arguments.Length, StringComparer.OrdinalIgnoreCase);
-			for(int idx = 0; idx < arguments.Length; idx++)
+			for (int idx = 0; idx < arguments.Length; idx++)
 			{
 				if (arguments[idx].Equals("--", StringComparison.Ordinal))
 				{
@@ -178,14 +178,10 @@ namespace EpicGames.Core
 					if (separatorIdx == -1)
 					{
 						// Ignore duplicate -Option flags; they are harmless.
-						if (_argumentToFirstIndex.ContainsKey(arguments[idx]))
+						if (!_argumentToFirstIndex.TryAdd(arguments[idx], idx))
 						{
 							_usedArguments.Set(idx, true);
 						}
-						else
-						{
-							_argumentToFirstIndex.Add(arguments[idx], idx);
-						}
 
 						// Mark this argument as a flag
 						_flagArguments.Set(idx, true);
@@ -265,7 +261,7 @@ namespace EpicGames.Core
 		public bool HasOption(string option)
 		{
 			int index;
-			if(_argumentToFirstIndex.TryGetValue(option, out index))
+			if (_argumentToFirstIndex.TryGetValue(option, out index))
 			{
 				_usedArguments.Set(index, true);
 				return true;
@@ -358,7 +354,7 @@ namespace EpicGames.Core
 		public string GetString(string prefix)
 		{
 			string? value;
-			if(!TryGetValue(prefix, out value))
+			if (!TryGetValue(prefix, out value))
 			{
 				throw new CommandLineArgumentException(String.Format("Missing '{0}...' argument", prefix));
 			}
@@ -373,7 +369,7 @@ namespace EpicGames.Core
 		public int GetInteger(string prefix)
 		{
 			int value;
-			if(!TryGetValue(prefix, out value))
+			if (!TryGetValue(prefix, out value))
 			{
 				throw new CommandLineArgumentException(String.Format("Missing '{0}...' argument", prefix));
 			}
@@ -388,7 +384,7 @@ namespace EpicGames.Core
 		public FileReference GetFileReference(string prefix)
 		{
 			FileReference? value;
-			if(!TryGetValue(prefix, out value))
+			if (!TryGetValue(prefix, out value))
 			{
 				throw new CommandLineArgumentException(String.Format("Missing '{0}...' argument", prefix));
 			}
@@ -403,7 +399,7 @@ namespace EpicGames.Core
 		public DirectoryReference GetDirectoryReference(string prefix)
 		{
 			DirectoryReference? value;
-			if(!TryGetValue(prefix, out value))
+			if (!TryGetValue(prefix, out value))
 			{
 				throw new CommandLineArgumentException(String.Format("Missing '{0}...' argument", prefix));
 			}
@@ -418,7 +414,7 @@ namespace EpicGames.Core
 		public T GetEnum<T>(string prefix) where T : struct
 		{
 			T value;
-			if(!TryGetValue(prefix, out value))
+			if (!TryGetValue(prefix, out value))
 			{
 				throw new CommandLineArgumentException(String.Format("Missing '{0}...' argument", prefix));
 			}
@@ -435,7 +431,7 @@ namespace EpicGames.Core
 		public string? GetStringOrDefault(string prefix, string? defaultValue)
 		{
 			string? value;
-			if(!TryGetValue(prefix, out value))
+			if (!TryGetValue(prefix, out value))
 			{
 				value = defaultValue;
 			}
@@ -451,7 +447,7 @@ namespace EpicGames.Core
 		public int GetIntegerOrDefault(string prefix, int defaultValue)
 		{
 			int value;
-			if(!TryGetValue(prefix, out value))
+			if (!TryGetValue(prefix, out value))
 			{
 				value = defaultValue;
 			}
@@ -468,7 +464,7 @@ namespace EpicGames.Core
 		public FileReference? GetFileReferenceOrDefault(string prefix, FileReference? defaultValue)
 		{
 			FileReference? value;
-			if(!TryGetValue(prefix, out value))
+			if (!TryGetValue(prefix, out value))
 			{
 				value = defaultValue;
 			}
@@ -485,7 +481,7 @@ namespace EpicGames.Core
 		public DirectoryReference? GetDirectoryReferenceOrDefault(string prefix, DirectoryReference? defaultValue)
 		{
 			DirectoryReference? value;
-			if(!TryGetValue(prefix, out value))
+			if (!TryGetValue(prefix, out value))
 			{
 				value = defaultValue;
 			}
@@ -501,7 +497,7 @@ namespace EpicGames.Core
 		public T GetEnumOrDefault<T>(string prefix, T defaultValue) where T : struct
 		{
 			T value;
-			if(!TryGetValue(prefix, out value))
+			if (!TryGetValue(prefix, out value))
 			{
 				value = defaultValue;
 			}
@@ -519,13 +515,13 @@ namespace EpicGames.Core
 			CheckValidPrefix(prefix);
 
 			int index;
-			if(!_argumentToFirstIndex.TryGetValue(prefix, out index))
+			if (!_argumentToFirstIndex.TryGetValue(prefix, out index))
 			{
 				value = null;
 				return false;
 			}
 
-			if(_nextArgumentIndex[index] != -1)
+			if (_nextArgumentIndex[index] != -1)
 			{
 				throw new CommandLineArgumentException(String.Format("Multiple {0}... arguments are specified", prefix));
 			}
@@ -545,7 +541,7 @@ namespace EpicGames.Core
 		{
 			// Try to get the string value of this argument
 			string? stringValue;
-			if(!TryGetValue(prefix, out stringValue))
+			if (!TryGetValue(prefix, out stringValue))
 			{
 				value = 0;
 				return false;
@@ -557,7 +553,7 @@ namespace EpicGames.Core
 				value = Int32.Parse(stringValue);
 				return true;
 			}
-			catch(Exception ex)
+			catch (Exception ex)
 			{
 				throw new CommandLineArgumentException(String.Format("The argument '{0}{1}' does not specify a valid integer", prefix, stringValue), ex);
 			}
@@ -573,7 +569,7 @@ namespace EpicGames.Core
 		{
 			// Try to get the string value of this argument
 			string? stringValue;
-			if(!TryGetValue(prefix, out stringValue))
+			if (!TryGetValue(prefix, out stringValue))
 			{
 				value = null;
 				return false;
@@ -585,7 +581,7 @@ namespace EpicGames.Core
 				value = new FileReference(stringValue);
 				return true;
 			}
-			catch(Exception ex)
+			catch (Exception ex)
 			{
 				throw new CommandLineArgumentException(String.Format("The argument '{0}{1}' does not specify a valid file name", prefix, stringValue), ex);
 			}
@@ -601,7 +597,7 @@ namespace EpicGames.Core
 		{
 			// Try to get the string value of this argument
 			string? stringValue;
-			if(!TryGetValue(prefix, out stringValue))
+			if (!TryGetValue(prefix, out stringValue))
 			{
 				value = null;
 				return false;
@@ -613,7 +609,7 @@ namespace EpicGames.Core
 				value = new DirectoryReference(stringValue);
 				return true;
 			}
-			catch(Exception ex)
+			catch (Exception ex)
 			{
 				throw new CommandLineArgumentException(String.Format("The argument '{0}{1}' does not specify a valid directory name", prefix, stringValue), ex);
 			}
@@ -629,7 +625,7 @@ namespace EpicGames.Core
 		{
 			// Try to get the string value of this argument
 			string? stringValue;
-			if(!TryGetValue(prefix, out stringValue))
+			if (!TryGetValue(prefix, out stringValue))
 			{
 				value = new T();
 				return false;
@@ -641,7 +637,7 @@ namespace EpicGames.Core
 				value = (T)Enum.Parse(typeof(T), stringValue, true);
 				return true;
 			}
-			catch(Exception ex)
+			catch (Exception ex)
 			{
 				throw new CommandLineArgumentException(String.Format("The argument '{0}{1}' does not specify a valid {2}", prefix, stringValue, typeof(T).Name), ex);
 			}
@@ -657,9 +653,9 @@ namespace EpicGames.Core
 			CheckValidPrefix(prefix);
 
 			int index;
-			if(_argumentToFirstIndex.TryGetValue(prefix, out index))
+			if (_argumentToFirstIndex.TryGetValue(prefix, out index))
 			{
-				for(; index != -1; index = _nextArgumentIndex[index])
+				for (; index != -1; index = _nextArgumentIndex[index])
 				{
 					_usedArguments.Set(index, true);
 					yield return _arguments[index].Substring(prefix.Length);
@@ -675,9 +671,9 @@ namespace EpicGames.Core
 		/// <returns>Sequence of values for the given prefix.</returns>
 		public IEnumerable<string> GetValues(string prefix, char separator)
 		{
-			foreach(string value in GetValues(prefix))
+			foreach (string value in GetValues(prefix))
 			{
-				foreach(string splitValue in value.Split(separator))
+				foreach (string splitValue in value.Split(separator))
 				{
 					yield return splitValue;
 				}
@@ -750,7 +746,7 @@ namespace EpicGames.Core
 				string? assignedArgument = null;
 
 				// Loop through all the attributes for different command line options that can modify it
-				foreach(CommandLineAttribute attribute in target.Attributes)
+				foreach (CommandLineAttribute attribute in target.Attributes)
 				{
 					// Add in any positional arguments
 					if (attribute.Positional)
@@ -771,20 +767,20 @@ namespace EpicGames.Core
 
 					// Get the value with the correct prefix
 					int firstIndex;
-					if(_argumentToFirstIndex.TryGetValue(prefix, out firstIndex))
+					if (_argumentToFirstIndex.TryGetValue(prefix, out firstIndex))
 					{
-						for(int index = firstIndex; index != -1; index = _nextArgumentIndex[index])
+						for (int index = firstIndex; index != -1; index = _nextArgumentIndex[index])
 						{
 							// Get the argument text
 							string argument = _arguments[index];
 
 							// Get the text for this value
 							string valueText;
-							if(attribute.Value != null)
+							if (attribute.Value != null)
 							{
 								valueText = attribute.Value;
 							}
-							else if(_flagArguments.Get(index))
+							else if (_flagArguments.Get(index))
 							{
 								valueText = "true";
 							}
@@ -794,18 +790,18 @@ namespace EpicGames.Core
 							}
 
 							// Apply the value to the field
-							if(attribute.ListSeparator == 0)
+							if (attribute.ListSeparator == 0)
 							{
-								if(ApplyArgument(targetObject, target, argument, valueText, assignedArgument, logger))
+								if (ApplyArgument(targetObject, target, argument, valueText, assignedArgument, logger))
 								{
 									assignedArgument = argument;
 								}
 							}
 							else
 							{
-								foreach(string itemValueText in valueText.Split(attribute.ListSeparator))
+								foreach (string itemValueText in valueText.Split(attribute.ListSeparator))
 								{
-									if(ApplyArgument(targetObject, target, argument, itemValueText, assignedArgument, logger))
+									if (ApplyArgument(targetObject, target, argument, itemValueText, assignedArgument, logger))
 									{
 										assignedArgument = argument;
 									}
@@ -821,23 +817,23 @@ namespace EpicGames.Core
 					}
 
 					// If this attribute is marked as required, keep track of it so we can warn if the field is not assigned to
-					if(attribute.Required && requiredPrefix == null)
+					if (attribute.Required && requiredPrefix == null)
 					{
 						requiredPrefix = prefix;
 					}
 				}
 
 				// Make sure that this field has been assigned to
-				if(assignedArgument == null && requiredPrefix != null)
+				if (assignedArgument == null && requiredPrefix != null)
 				{
 					missingArguments.Add(requiredPrefix);
 				}
 			}
 
 			// If any arguments were missing, print an error about them
-			if(missingArguments.Count > 0)
+			if (missingArguments.Count > 0)
 			{
-				if(missingArguments.Count == 1)
+				if (missingArguments.Count == 1)
 				{
 					throw new CommandLineArgumentException(String.Format("Missing {0} argument", missingArguments[0].Replace("=", "=...", StringComparison.Ordinal)));
 				}
@@ -918,7 +914,7 @@ namespace EpicGames.Core
 				StringBuilder descriptionBuilder = new StringBuilder();
 				foreach (DescriptionAttribute attribute in target.Member.GetCustomAttributes<DescriptionAttribute>())
 				{
-					if(descriptionBuilder.Length > 0)
+					if (descriptionBuilder.Length > 0)
 					{
 						descriptionBuilder.Append('\n');
 					}
@@ -942,7 +938,7 @@ namespace EpicGames.Core
 				foreach (CommandLineAttribute attribute in target.Attributes)
 				{
 					string prefix = GetArgumentPrefix(target, attribute);
-					if(prefix.EndsWith('='))
+					if (prefix.EndsWith('='))
 					{
 						prefix += "...";
 					}
@@ -961,7 +957,7 @@ namespace EpicGames.Core
 		{
 			// See if the entire string is quoted correctly
 			bool bInQuotes = false;
-			for (int idx = 0;;idx++)
+			for (int idx = 0; ; idx++)
 			{
 				if (idx == argument.Length)
 				{
@@ -980,7 +976,7 @@ namespace EpicGames.Core
 			// Try to insert a quote after the argument string
 			if (argument[0] == '-')
 			{
-				for(int idx = 1; idx < argument.Length && argument[idx] != ' '; idx++)
+				for (int idx = 1; idx < argument.Length && argument[idx] != ' '; idx++)
 				{
 					if (argument[idx] == '=')
 					{
@@ -1003,7 +999,7 @@ namespace EpicGames.Core
 			StringBuilder result = new StringBuilder();
 			foreach (string argument in arguments)
 			{
-				if(result.Length > 0)
+				if (result.Length > 0)
 				{
 					result.Append(' ');
 				}
@@ -1023,18 +1019,18 @@ namespace EpicGames.Core
 
 			List<string> arguments = [];
 			// First do a pass leaving all quotes in the arguments, they will be removed later
-			for(int idx = 0; idx < commandLine.Length; idx++)
+			for (int idx = 0; idx < commandLine.Length; idx++)
 			{
-				if(!Char.IsWhiteSpace(commandLine[idx]))
+				if (!Char.IsWhiteSpace(commandLine[idx]))
 				{
 					argument.Clear();
-					for(bool bInQuotes = false; idx < commandLine.Length; idx++)
+					for (bool bInQuotes = false; idx < commandLine.Length; idx++)
 					{
 						if (commandLine[idx] == '\"')
 						{
 							bInQuotes ^= true;
 						}
-						else if(!bInQuotes && Char.IsWhiteSpace(commandLine[idx]))
+						else if (!bInQuotes && Char.IsWhiteSpace(commandLine[idx]))
 						{
 							break;
 						}
@@ -1065,9 +1061,9 @@ namespace EpicGames.Core
 		public CommandLineArguments Append(IEnumerable<string> appendArguments)
 		{
 			CommandLineArguments newArguments = new CommandLineArguments(Enumerable.Concat(_arguments, appendArguments).ToArray());
-			for(int idx = 0; idx < _arguments.Length; idx++)
+			for (int idx = 0; idx < _arguments.Length; idx++)
 			{
-				if(HasBeenUsed(idx))
+				if (HasBeenUsed(idx))
 				{
 					newArguments.MarkAsUsed(idx);
 				}
@@ -1088,10 +1084,10 @@ namespace EpicGames.Core
 			// Split the arguments into the values array and an array of new arguments
 			int[] newArgumentIndex = new int[_arguments.Length];
 			List<string> newArgumentList = new List<string>(_arguments.Length);
-			for(int idx = 0; idx < _arguments.Length; idx++)
+			for (int idx = 0; idx < _arguments.Length; idx++)
 			{
 				string argument = _arguments[idx];
-				if(argument.StartsWith(prefix, StringComparison.OrdinalIgnoreCase))
+				if (argument.StartsWith(prefix, StringComparison.OrdinalIgnoreCase))
 				{
 					newArgumentIndex[idx] = -1;
 					values.Add(argument.Substring(prefix.Length));
@@ -1105,9 +1101,9 @@ namespace EpicGames.Core
 
 			// Create the new argument list, and mark the same arguments as used
 			CommandLineArguments newArguments = new CommandLineArguments([.. newArgumentList]);
-			for(int idx = 0; idx < _arguments.Length; idx++)
+			for (int idx = 0; idx < _arguments.Length; idx++)
 			{
-				if(HasBeenUsed(idx) && newArgumentIndex[idx] != -1)
+				if (HasBeenUsed(idx) && newArgumentIndex[idx] != -1)
 				{
 					newArguments.MarkAsUsed(newArgumentIndex[idx]);
 				}
@@ -1130,18 +1126,18 @@ namespace EpicGames.Core
 		{
 			// Find all the unused arguments
 			List<string> remainingArguments = [];
-			for(int idx = 0; idx < _arguments.Length; idx++)
+			for (int idx = 0; idx < _arguments.Length; idx++)
 			{
-				if(!_usedArguments[idx])
+				if (!_usedArguments[idx])
 				{
 					remainingArguments.Add(_arguments[idx]);
 				}
 			}
 
 			// Output a warning
-			if(remainingArguments.Count > 0)
+			if (remainingArguments.Count > 0)
 			{
-				if(remainingArguments.Count == 1)
+				if (remainingArguments.Count == 1)
 				{
 					logger.LogWarning("Invalid argument: {Argument}", remainingArguments[0]);
 				}
@@ -1167,7 +1163,7 @@ namespace EpicGames.Core
 		/// <returns></returns>
 		public IEnumerable<string> GetUnusedArguments()
 		{
-			for(int idx = 0; idx < _arguments.Length; idx++)
+			for (int idx = 0; idx < _arguments.Length; idx++)
 			{
 				if (!_usedArguments[idx])
 				{
@@ -1191,15 +1187,15 @@ namespace EpicGames.Core
 		/// <param name="prefix">The prefix to check</param>
 		private static void CheckValidPrefix(string prefix)
 		{
-			if(prefix.Length == 0)
+			if (prefix.Length == 0)
 			{
 				throw new ArgumentException("Argument prefix cannot be empty.");
 			}
-			else if(prefix[0] != '-')
+			else if (prefix[0] != '-')
 			{
 				throw new ArgumentException("Argument prefix must begin with a hyphen.");
 			}
-			else if(!s_valueSeparators.Contains(prefix[^1]))
+			else if (!s_valueSeparators.Contains(prefix[^1]))
 			{
 				throw new ArgumentException(String.Format("Argument prefix must end with '{0}'", String.Join("' or '", s_valueSeparators)));
 			}
@@ -1233,7 +1229,7 @@ namespace EpicGames.Core
 
 			// Try to parse the value
 			object? value;
-			if(!TryParseValue(valueType, valueText, out value))
+			if (!TryParseValue(valueType, valueText, out value))
 			{
 				logger.LogWarning("Unable to parse value for argument '{Argument}'.", argumentText);
 				return false;
@@ -1243,10 +1239,10 @@ namespace EpicGames.Core
 			if (collectionType == null)
 			{
 				// Check if this field has already been assigned to. Output a warning if the previous value is in conflict with the new one.
-				if(previousArgumentText != null)
+				if (previousArgumentText != null)
 				{
 					object? previousValue = target.GetValue(targetObject);
-					if(!Object.Equals(previousValue, value))
+					if (!Equals(previousValue, value))
 					{
 						logger.LogWarning("Argument '{Argument}' conflicts with '{PrevArgument}'; ignoring.", argumentText, previousArgumentText);
 					}
@@ -1412,20 +1408,20 @@ namespace EpicGames.Core
 		/// <returns>The command line argument with quotes inserted to escape it if necessary</returns>
 		public static void Append(StringBuilder commandLine, string argument)
 		{
-			if(commandLine.Length > 0)
+			if (commandLine.Length > 0)
 			{
 				commandLine.Append(' ');
 			}
 
 			int spaceIdx = argument.IndexOf(' ', StringComparison.Ordinal);
-			if(spaceIdx == -1)
+			if (spaceIdx == -1)
 			{
 				commandLine.Append(argument);
 			}
 			else
 			{
 				int equalsIdx = argument.IndexOf('=', StringComparison.Ordinal);
-				if(equalsIdx == -1 || equalsIdx > spaceIdx)
+				if (equalsIdx == -1 || equalsIdx > spaceIdx)
 				{
 					commandLine.AppendFormat("\"{0}\"", argument);
 				}
@@ -1443,7 +1439,7 @@ namespace EpicGames.Core
 		public override string ToString()
 		{
 			StringBuilder result = new StringBuilder();
-			foreach(string argument in _arguments)
+			foreach (string argument in _arguments)
 			{
 				Append(result, argument);
 			}
diff --git a/Engine/Source/Programs/Shared/EpicGames.Core/DirectoryUtils.cs b/Engine/Source/Programs/Shared/EpicGames.Core/DirectoryUtils.cs
--- a/Engine/Source/Programs/Shared/EpicGames.Core/DirectoryUtils.cs
+++ b/Engine/Source/Programs/Shared/EpicGames.Core/DirectoryUtils.cs
@@ -39,7 +39,7 @@ namespace EpicGames.Core
 					{
 						foreach (DirectoryInfo childInfo in parentInfo.EnumerateDirectories())
 						{
-							if (String.Equals(childInfo.Name, info.Name, DirectoryReference.Comparison))
+							if (String.Equals(childInfo.Name, info.Name, FileSystemReference.Comparison))
 							{
 								return childInfo;
 							}
diff --git a/Engine/Source/Programs/Shared/EpicGames.Core/FileFilter.cs b/Engine/Source/Programs/Shared/EpicGames.Core/FileFilter.cs
--- a/Engine/Source/Programs/Shared/EpicGames.Core/FileFilter.cs
+++ b/Engine/Source/Programs/Shared/EpicGames.Core/FileFilter.cs
@@ -111,7 +111,7 @@ namespace EpicGames.Core
 			FileFilter? filter = null;
 			if (rules != null && rules.Count > 0)
 			{
-				FileFilterType defaultType = rules[0].StartsWith("-", StringComparison.Ordinal) ? FileFilterType.Include : FileFilterType.Exclude;
+				FileFilterType defaultType = rules[0].StartsWith('-') ? FileFilterType.Include : FileFilterType.Exclude;
 				filter = new FileFilter(defaultType);
 				filter.AddRules(rules);
 			}
diff --git a/Engine/Source/Programs/Shared/EpicGames.Core/FileSet.cs b/Engine/Source/Programs/Shared/EpicGames.Core/FileSet.cs
--- a/Engine/Source/Programs/Shared/EpicGames.Core/FileSet.cs
+++ b/Engine/Source/Programs/Shared/EpicGames.Core/FileSet.cs
@@ -18,7 +18,7 @@ namespace EpicGames.Core
 		/// An empty fileset
 		/// </summary>
 		public static FileSet Empty { get; } = new FileSetFromFiles([]);
-		
+
 		/// <summary>
 		/// Path of this tree
 		/// </summary>
@@ -209,7 +209,7 @@ namespace EpicGames.Core
 			{
 				pathToSourceFile[pathPrefix + path] = file;
 			}
-			foreach((string path, FileSet fileSet) in EnumerateDirectories())
+			foreach ((string path, FileSet fileSet) in EnumerateDirectories())
 			{
 				fileSet.FlattenInternal(pathPrefix + path + "/", pathToSourceFile);
 			}
@@ -244,6 +244,7 @@ namespace EpicGames.Core
 	{
 		readonly Dictionary<string, FileReference> _files = [];
 		readonly Dictionary<string, FileSetFromFiles> _subTrees = [];
+		internal static readonly char[] Separator = ['/', '\\'];
 
 		/// <summary>
 		/// Private constructor
@@ -263,7 +264,7 @@ namespace EpicGames.Core
 		{
 			foreach ((string path, FileReference file) in inputFiles)
 			{
-				string[] fragments = path.Split(new[] { '/', '\\' }, StringSplitOptions.RemoveEmptyEntries);
+				string[] fragments = path.Split(Separator, StringSplitOptions.RemoveEmptyEntries);
 
 				FileSetFromFiles current = this;
 				for (int idx = 0; idx < fragments.Length - 1; idx++)
diff --git a/Engine/Source/Programs/Shared/EpicGames.Core/FileUtils.cs b/Engine/Source/Programs/Shared/EpicGames.Core/FileUtils.cs
--- a/Engine/Source/Programs/Shared/EpicGames.Core/FileUtils.cs
+++ b/Engine/Source/Programs/Shared/EpicGames.Core/FileUtils.cs
@@ -110,7 +110,7 @@ namespace EpicGames.Core
 			{
 				foreach (FileInfo childInfo in parentInfo.EnumerateFiles())
 				{
-					if (String.Equals(childInfo.Name, info.Name, FileReference.Comparison))
+					if (String.Equals(childInfo.Name, info.Name, FileSystemReference.Comparison))
 					{
 						return childInfo;
 					}
diff --git a/Engine/Source/Programs/Shared/EpicGames.Core/HttpClientExtensions.cs b/Engine/Source/Programs/Shared/EpicGames.Core/HttpClientExtensions.cs
--- a/Engine/Source/Programs/Shared/EpicGames.Core/HttpClientExtensions.cs
+++ b/Engine/Source/Programs/Shared/EpicGames.Core/HttpClientExtensions.cs
@@ -158,7 +158,7 @@ namespace EpicGames.Core
 		/// <typeparam name="T">Type of the object to parse</typeparam>
 		/// <param name="obj">The object instance</param>
 		/// <returns>Http content object</returns>
-		private static HttpContent ToJsonContent<T>(T obj)
+		private static StringContent ToJsonContent<T>(T obj)
 		{
 			return new StringContent(JsonSerializer.Serialize<T>(obj), Encoding.UTF8, "application/json");
 		}
@@ -176,7 +176,7 @@ namespace EpicGames.Core
 			return JsonSerializer.Deserialize<T>(bytes, new JsonSerializerOptions { PropertyNameCaseInsensitive = true })!;
 		}
 	}
-	
+
 	/// <summary>
 	/// Extension methods for HttpRequestMessage
 	/// </summary>
@@ -188,7 +188,7 @@ namespace EpicGames.Core
 			"Signature",
 			"AWSAccessKeyId"
 		};
-		
+
 		/// <summary>
 		/// Returns a new URI with sensitive query parameters redacted from the request URI.
 		/// </summary>
@@ -199,14 +199,14 @@ namespace EpicGames.Core
 			{
 				return uri;
 			}
-			
+
 			NameValueCollection queryParams = HttpUtility.ParseQueryString(uri.Query);
 			foreach (string key in queryParams.AllKeys.OfType<string>())
 			{
 				string value = s_sensitiveKeys.Contains(key) ? "redacted" : queryParams[key] ?? String.Empty;
 				queryParams[key] = value;
 			}
-			
+
 			return new UriBuilder(uri) { Query = String.Join("&", queryParams) }.Uri;
 		}
 	}
@@ -306,9 +306,9 @@ namespace EpicGames.Core
 		{
 			public NonOwningHandler(HttpMessageHandler inner) => InnerHandler = inner;
 
-			#pragma warning disable CA2215 // Intentionally skip calling base.Dispose to prevent cascading dispose
+#pragma warning disable CA2215 // Intentionally skip calling base.Dispose to prevent cascading dispose
 			protected override void Dispose(bool disposing) { }
-			#pragma warning restore CA2215 // Dispose objects before losing scope
+#pragma warning restore CA2215 // Dispose objects before losing scope
 		}
 
 		private static readonly Lazy<SocketsHttpHandler> s_defaultHandler = new(InternalCreateDefaultHandler);
diff --git a/Engine/Source/Programs/Shared/EpicGames.Core/JsonLogEvent.cs b/Engine/Source/Programs/Shared/EpicGames.Core/JsonLogEvent.cs
--- a/Engine/Source/Programs/Shared/EpicGames.Core/JsonLogEvent.cs
+++ b/Engine/Source/Programs/Shared/EpicGames.Core/JsonLogEvent.cs
@@ -60,7 +60,7 @@ namespace EpicGames.Core
 		}
 
 		/// <summary>
-		/// Creates a json log event from the given logger paramters
+		/// Creates a json log event from the given logger parameters
 		/// </summary>
 		/// <inheritdoc cref="ILogger.Log{TState}(LogLevel, EventId, TState, Exception, Func{TState, Exception, String})"/>
 		public static JsonLogEvent FromLoggerState<TState>(LogLevel logLevel, EventId eventId, TState state, Exception? exception, Func<TState, Exception?, string> formatter)
@@ -316,7 +316,7 @@ namespace EpicGames.Core
 			}
 			return 1;
 		}
-		
+
 		/// <summary>
 		/// Counts the number of newlines in an escaped JSON string
 		/// </summary>
diff --git a/Engine/Source/Programs/Shared/EpicGames.Core/JsonTracer.cs b/Engine/Source/Programs/Shared/EpicGames.Core/JsonTracer.cs
--- a/Engine/Source/Programs/Shared/EpicGames.Core/JsonTracer.cs
+++ b/Engine/Source/Programs/Shared/EpicGames.Core/JsonTracer.cs
@@ -59,19 +59,19 @@ namespace EpicGames.Core
 		public JsonTracerSpanContext Context => _jsonTracerContext;
 
 		ISpanContext ISpan.Context => Context;
-		
+
 		public string OperationName { get; private set; }
-		
+
 		private readonly JsonTracer _tracer;
 		private readonly JsonTracerSpanContext _jsonTracerContext;
 		private DateTimeOffset _finishTimestamp;
 		private bool _finished;
 		private readonly Dictionary<string, object> _tags;
 		private readonly List<JsonTracerSpan.Reference> _references;
-		
+
 		public DateTimeOffset StartTimestamp { get; }
 		public Dictionary<string, object> Tags => new Dictionary<string, object>(_tags);
-		
+
 		/// <summary>
 		/// The finish time of the span; only valid after a call to <see cref="Finish()"/>.
 		/// </summary>
@@ -103,7 +103,7 @@ namespace EpicGames.Core
 			_tracer = jsonTracer;
 			OperationName = operationName;
 			StartTimestamp = startTimestamp;
-			
+
 			_tags = initialTags == null
 				? []
 				: new Dictionary<string, object>(initialTags);
@@ -111,7 +111,7 @@ namespace EpicGames.Core
 			_references = references == null
 				? []
 				: [.. references];
-			
+
 			JsonTracerSpanContext? parentContext = FindPreferredParentRef(_references);
 			if (parentContext == null)
 			{
@@ -127,36 +127,36 @@ namespace EpicGames.Core
 			}
 		}
 
-		public ISpan SetTag(string key, string value) => SetObjectTag(key, value); 
+		public ISpan SetTag(string key, string value) => SetObjectTag(key, value);
 		public ISpan SetTag(string key, bool value) => SetObjectTag(key, value);
 		public ISpan SetTag(string key, int value) => SetObjectTag(key, value);
 		public ISpan SetTag(string key, double value) => SetObjectTag(key, value);
 
-		public ISpan SetTag(BooleanTag tag, bool value) 
-		{ 
-			SetObjectTag(tag.Key, value); 
-			return this; 
+		public ISpan SetTag(BooleanTag tag, bool value)
+		{
+			SetObjectTag(tag.Key, value);
+			return this;
 		}
 
-		public ISpan SetTag(IntOrStringTag tag, string value) 
-		{ 
-			SetObjectTag(tag.Key, value); 
-			return this; 
+		public ISpan SetTag(IntOrStringTag tag, string value)
+		{
+			SetObjectTag(tag.Key, value);
+			return this;
 		}
 
-		public ISpan SetTag(IntTag tag, int value) 
-		{ 
-			SetObjectTag(tag.Key, value); 
-			return this; 
+		public ISpan SetTag(IntTag tag, int value)
+		{
+			SetObjectTag(tag.Key, value);
+			return this;
 		}
 
-		public ISpan SetTag(StringTag tag, string value) 
-		{ 
-			SetObjectTag(tag.Key, value); 
-			return this; 
+		public ISpan SetTag(StringTag tag, string value)
+		{
+			SetObjectTag(tag.Key, value);
+			return this;
 		}
 
-		private ISpan SetObjectTag(string key, object value)
+		private JsonTracerSpan SetObjectTag(string key, object value)
 		{
 			lock (_lock)
 			{
@@ -196,7 +196,7 @@ namespace EpicGames.Core
 				_finished = true;
 			}
 		}
-		
+
 		private static JsonTracerSpanContext? FindPreferredParentRef(IList<Reference> references)
 		{
 			if (!references.Any())
@@ -207,7 +207,7 @@ namespace EpicGames.Core
 			// return the context of the parent, if applicable
 			foreach (Reference reference in references)
 			{
-				if (OpenTracing.References.ChildOf.Equals(reference.ReferenceType, StringComparison.Ordinal))
+				if (References.ChildOf.Equals(reference.ReferenceType, StringComparison.Ordinal))
 				{
 					return reference.Context;
 				}
@@ -216,7 +216,7 @@ namespace EpicGames.Core
 			// otherwise, return the context of the first reference
 			return references.First().Context;
 		}
-		
+
 		private void CheckForFinished(string format, params object[] args)
 		{
 			if (_finished)
@@ -246,12 +246,12 @@ namespace EpicGames.Core
 			{
 				return Equals(obj as Reference);
 			}
-			
+
 			public bool Equals(Reference? other)
 			{
 				return other != null &&
-				       EqualityComparer<JsonTracerSpanContext>.Default.Equals(Context, other.Context) &&
-				       ReferenceType == other.ReferenceType;
+					   EqualityComparer<JsonTracerSpanContext>.Default.Equals(Context, other.Context) &&
+					   ReferenceType == other.ReferenceType;
 			}
 
 #pragma warning disable IDE0070 // Use 'System.HashCode'
@@ -274,13 +274,13 @@ namespace EpicGames.Core
 		private readonly List<JsonTracerSpan.Reference> _references = [];
 		private readonly Dictionary<string, object> _tags = [];
 		private bool _ignoreActiveSpan;
-		
+
 		public JsonTracerSpanBuilder(JsonTracer tracer, string operationName)
 		{
 			_tracer = tracer;
 			_operationName = operationName;
 		}
-		
+
 		public ISpanBuilder AsChildOf(ISpanContext? parent)
 		{
 			if (parent == null)
@@ -288,7 +288,7 @@ namespace EpicGames.Core
 				return this;
 			}
 
-			return AddReference(OpenTracing.References.ChildOf, parent);
+			return AddReference(References.ChildOf, parent);
 		}
 
 		public ISpanBuilder AsChildOf(ISpan? parent)
@@ -298,7 +298,7 @@ namespace EpicGames.Core
 				return this;
 			}
 
-			return AddReference(OpenTracing.References.ChildOf, parent.Context);
+			return AddReference(References.ChildOf, parent.Context);
 		}
 
 		public ISpanBuilder AddReference(string referenceType, ISpanContext? referencedContext)
@@ -318,51 +318,51 @@ namespace EpicGames.Core
 		}
 
 		public ISpanBuilder WithTag(string key, string value)
-		{ 
-			_tags[key] = value; 
-			return this; 
+		{
+			_tags[key] = value;
+			return this;
 		}
-		
-		public ISpanBuilder WithTag(string key, bool value) 
-		{ 
-			_tags[key] = value; 
-			return this; 
+
+		public ISpanBuilder WithTag(string key, bool value)
+		{
+			_tags[key] = value;
+			return this;
 		}
 
 		public ISpanBuilder WithTag(string key, int value)
-		{ 
-			_tags[key] = value; 
-			return this; 
+		{
+			_tags[key] = value;
+			return this;
 		}
 
-		public ISpanBuilder WithTag(string key, double value) 
-		{ 
-			_tags[key] = value; 
-			return this; 
+		public ISpanBuilder WithTag(string key, double value)
+		{
+			_tags[key] = value;
+			return this;
 		}
 
-		public ISpanBuilder WithTag(BooleanTag tag, bool value) 
-		{ 
-			_tags[tag.Key] = value; 
-			return this; 
+		public ISpanBuilder WithTag(BooleanTag tag, bool value)
+		{
+			_tags[tag.Key] = value;
+			return this;
 		}
 
-		public ISpanBuilder WithTag(IntOrStringTag tag, string value) 
-		{ 
-			_tags[tag.Key] = value; 
-			return this; 
+		public ISpanBuilder WithTag(IntOrStringTag tag, string value)
+		{
+			_tags[tag.Key] = value;
+			return this;
 		}
-		
-		public ISpanBuilder WithTag(IntTag tag, int value) 
-		{ 
-			_tags[tag.Key] = value; 
-			return this; 
+
+		public ISpanBuilder WithTag(IntTag tag, int value)
+		{
+			_tags[tag.Key] = value;
+			return this;
 		}
 
 		public ISpanBuilder WithTag(StringTag tag, string value)
-		{ 
-			_tags[tag.Key] = value; 
-			return this; 
+		{
+			_tags[tag.Key] = value;
+			return this;
 		}
 
 		public ISpanBuilder WithStartTimestamp(DateTimeOffset timestamp)
@@ -390,15 +390,15 @@ namespace EpicGames.Core
 			}
 
 			ISpanContext? activeSpanContext = _tracer.ActiveSpan?.Context;
-			if (!_references.Any() && !_ignoreActiveSpan && activeSpanContext != null)
+			if (_references.Count == 0 && !_ignoreActiveSpan && activeSpanContext != null)
 			{
-				_references.Add(new JsonTracerSpan.Reference((JsonTracerSpanContext)activeSpanContext, OpenTracing.References.ChildOf));
+				_references.Add(new JsonTracerSpan.Reference((JsonTracerSpanContext)activeSpanContext, References.ChildOf));
 			}
 
 			return new JsonTracerSpan(_tracer, _operationName, _startTimestamp, _tags, _references);
 		}
 	}
-	
+
 	public class JsonTracer : ITracer
 	{
 		public IScopeManager ScopeManager { get; }
@@ -407,7 +407,7 @@ namespace EpicGames.Core
 		private readonly object _lock = new object();
 		private readonly List<JsonTracerSpan> _finishedSpans = [];
 		private readonly DirectoryReference? _telemetryDir;
-		
+
 		public JsonTracer(DirectoryReference? telemetryDir = null)
 		{
 			ScopeManager = new AsyncLocalScopeManager();
@@ -446,7 +446,7 @@ namespace EpicGames.Core
 			if (_telemetryDir != null)
 			{
 				string telemetryScopeId = Environment.GetEnvironmentVariable("UE_TELEMETRY_SCOPE_ID") ?? "noscope";
-				
+
 				FileReference file;
 				using (System.Diagnostics.Process process = System.Diagnostics.Process.GetCurrentProcess())
 				{
@@ -462,7 +462,7 @@ namespace EpicGames.Core
 				}
 			}
 		}
-		
+
 		public List<JsonTracerSpan> GetFinishedSpans()
 		{
 			lock (_lock)
@@ -503,7 +503,7 @@ namespace EpicGames.Core
 			writer.WriteArrayEnd();
 			writer.WriteObjectEnd();
 		}
-		
+
 		internal void AppendFinishedSpan(JsonTracerSpan jsonTracerSpan)
 		{
 			lock (_lock)
diff --git a/Engine/Source/Programs/Shared/EpicGames.Core/Log.cs b/Engine/Source/Programs/Shared/EpicGames.Core/Log.cs
--- a/Engine/Source/Programs/Shared/EpicGames.Core/Log.cs
+++ b/Engine/Source/Programs/Shared/EpicGames.Core/Log.cs
@@ -170,13 +170,13 @@ namespace EpicGames.Core
 		/// Log-backup-2021.10.29-19.53.17.txt
 		/// </summary>
 		public static bool BackupLogFiles { get; set; } = true;
-		
+
 		/// <summary>
 		/// The number of backups to be preserved - when there are more than this, the oldest backups will be deleted.
 		/// Backups will not be deleted if BackupLogFiles is false.
 		/// </summary>
 		public static int LogFileBackupCount { get; set; } = 10;
-		
+
 		/// <summary>
 		/// Path to the log file being written to. May be null.
 		/// </summary>
@@ -227,7 +227,7 @@ namespace EpicGames.Core
 		/// <param name="outputFile">The file to back up</param>
 		public static void BackupLogFile(FileReference outputFile)
 		{
-			if (!Log.BackupLogFiles || !FileReference.Exists(outputFile))
+			if (!BackupLogFiles || !FileReference.Exists(outputFile))
 			{
 				return;
 			}
@@ -246,7 +246,7 @@ namespace EpicGames.Core
 				// sort them from newest to oldest
 				.OrderByDescending(x => x.GetFileName())
 				// skip the newest ones that are to be kept; -1 because we're about to create another backup.
-				.Skip(Log.LogFileBackupCount - 1))
+				.Skip(LogFileBackupCount - 1))
 			{
 				Logger.LogDebug("Deleting old log file: {File}", oldBackup);
 				FileReference.Delete(oldBackup);
@@ -725,7 +725,7 @@ namespace EpicGames.Core
 		[StringFormatMethod("Format")]
 		public static void TraceWarningOnce(FileReference file, string format, params object?[] args)
 		{
-			WriteLinePrivate( true, LogEventType.Warning, LogFormatOptions.NoSeverityPrefix, "{0}: warning: {1}", file, String.Format(format, args));
+			WriteLinePrivate(true, LogEventType.Warning, LogFormatOptions.NoSeverityPrefix, "{0}: warning: {1}", file, String.Format(format, args));
 		}
 
 		/// <summary>
@@ -890,7 +890,7 @@ namespace EpicGames.Core
 		/// Renders all the captured events to another logger
 		/// </summary>
 		/// <param name="logger"></param>
-		public void RenderTo(ILogger logger) => Events.ForEach( x => logger.Log(x.Level, x.Id, x, null, (s, e) => s.ToString()) );
+		public void RenderTo(ILogger logger) => Events.ForEach(x => logger.Log(x.Level, x.Id, x, null, (s, e) => s.ToString()));
 
 		/// <inheritdoc/>
 		public IDisposable? BeginScope<TState>(TState state) where TState : notnull => _scopeCollection.BeginScope(state);
@@ -1052,7 +1052,7 @@ namespace EpicGames.Core
 
 			bool bJsonStdOut = Environment.CommandLine.Contains("JsonStdOut", StringComparison.InvariantCulture);
 			string? envVar = Environment.GetEnvironmentVariable("UE_LOG_JSON_TO_STDOUT");
-			if(bJsonStdOut || envVar != null && Int32.TryParse(envVar, out int value) && value != 0)
+			if (bJsonStdOut || envVar != null && Int32.TryParse(envVar, out int value) && value != 0)
 			{
 				WriteJsonToStdOut = true;
 			}
@@ -1192,8 +1192,8 @@ namespace EpicGames.Core
 		{
 			lock (_syncObject)
 			{
-				IEnumerable<StartupTraceListener> startupListeners = Trace.Listeners.OfType<StartupTraceListener>();
-				if (startupListeners.Any())
+				List<StartupTraceListener> startupListeners = [.. Trace.Listeners.OfType<StartupTraceListener>()];
+				if (startupListeners.Count != 0)
 				{
 					Trace.Listeners.Remove(startupListeners.First());
 				}
@@ -1214,8 +1214,8 @@ namespace EpicGames.Core
 			lock (_syncObject)
 			{
 				// find the StartupTraceListener in the listeners that was added early on
-				IEnumerable<StartupTraceListener> startupListeners = Trace.Listeners.OfType<StartupTraceListener>();
-				if (startupListeners.Any())
+				List<StartupTraceListener> startupListeners = [.. Trace.Listeners.OfType<StartupTraceListener>()];
+				if (startupListeners.Count != 0)
 				{
 					StartupTraceListener startupListener = startupListeners.First();
 					startupListener.CopyTo(firstTextWriter);
diff --git a/Engine/Source/Programs/Shared/EpicGames.Core/LogEvent.cs b/Engine/Source/Programs/Shared/EpicGames.Core/LogEvent.cs
--- a/Engine/Source/Programs/Shared/EpicGames.Core/LogEvent.cs
+++ b/Engine/Source/Programs/Shared/EpicGames.Core/LogEvent.cs
@@ -368,7 +368,7 @@ namespace EpicGames.Core
 			}
 		}
 
-		static object ReadArrayPropertyValue(ref Utf8JsonReader reader)
+		static List<object?> ReadArrayPropertyValue(ref Utf8JsonReader reader)
 		{
 			List<object?> result = [];
 			while (reader.Read() && reader.TokenType != JsonTokenType.EndArray)
@@ -544,7 +544,7 @@ namespace EpicGames.Core
 			// Try to log the event
 			IEnumerable<KeyValuePair<string, object?>>? values = state as IEnumerable<KeyValuePair<string, object?>>;
 			string? format = values?.FirstOrDefault(x => x.Key.Equals(MessageTemplate.FormatPropertyName, StringComparison.Ordinal)).Value?.ToString();
-			string message = (format == null)? formatter(state, exception) : MessageTemplate.Render(format, values);
+			string message = (format == null) ? formatter(state, exception) : MessageTemplate.Render(format, values);
 			return new LogEvent(time, level, eventId, message, format, values, LogException.FromException(exception));
 		}
 
diff --git a/Engine/Source/Programs/Shared/EpicGames.Core/MessageTemplate.cs b/Engine/Source/Programs/Shared/EpicGames.Core/MessageTemplate.cs
--- a/Engine/Source/Programs/Shared/EpicGames.Core/MessageTemplate.cs
+++ b/Engine/Source/Programs/Shared/EpicGames.Core/MessageTemplate.cs
@@ -44,7 +44,7 @@ namespace EpicGames.Core
 			List<(int, int)>? names = ParsePropertyNames(format);
 			if (names != null && properties != null)
 			{
-				foreach((int offset, int length) in names)
+				foreach ((int offset, int length) in names)
 				{
 					ReadOnlySpan<char> argument = format.AsSpan(offset, length);
 
@@ -66,7 +66,7 @@ namespace EpicGames.Core
 					{
 						argument = argument.Slice(0, alignmentIdx);
 					}
-					
+
 					// Try to get the property value
 					object? value;
 					if (TryGetPropertyValue(argument, properties, out value))
@@ -139,7 +139,7 @@ namespace EpicGames.Core
 		/// <param name="result">Buffer to receive the escaped string</param>
 		public static void Escape(ReadOnlySpan<char> text, StringBuilder result)
 		{
-			foreach(char character in text)
+			foreach (char character in text)
 			{
 				result.Append(character);
 				if (character == '{' || character == '}')
@@ -265,7 +265,7 @@ namespace EpicGames.Core
 		public static bool TryGetPropertyValue(ReadOnlySpan<char> name, IEnumerable<KeyValuePair<string, object?>> properties, out object? value)
 		{
 			int number;
-			if (Int32.TryParse(name, System.Globalization.NumberStyles.Integer, null, out number))
+			if (Int32.TryParse(name, NumberStyles.Integer, null, out number))
 			{
 				foreach (KeyValuePair<string, object?> property in properties)
 				{
diff --git a/Engine/Source/Programs/Shared/EpicGames.Core/ProgressTextReader.cs b/Engine/Source/Programs/Shared/EpicGames.Core/ProgressTextReader.cs
--- a/Engine/Source/Programs/Shared/EpicGames.Core/ProgressTextReader.cs
+++ b/Engine/Source/Programs/Shared/EpicGames.Core/ProgressTextReader.cs
@@ -144,7 +144,7 @@ namespace EpicGames.Core
 				string remainingLine = line;
 
 				// Look for a progress directive at the end of a line, in square brackets
-				if (trimLine.EndsWith("]", StringComparison.Ordinal))
+				if (trimLine.EndsWith(']'))
 				{
 					for (int lastIdx = trimLine.Length - 2; lastIdx >= 0 && trimLine[lastIdx] != ']'; lastIdx--)
 					{
diff --git a/Engine/Source/Programs/Shared/EpicGames.Core/UHTTypes.cs b/Engine/Source/Programs/Shared/EpicGames.Core/UHTTypes.cs
--- a/Engine/Source/Programs/Shared/EpicGames.Core/UHTTypes.cs
+++ b/Engine/Source/Programs/Shared/EpicGames.Core/UHTTypes.cs
@@ -127,8 +127,8 @@ namespace EpicGames.Core
 		/// Default setting to use
 		/// </summary>
 		Default = Params,
-	
-	 	/// <summary>
+
+		/// <summary>
 		/// Generate all possible types at once for comparison
 		/// </summary>
 		All = Params | ConstInit,
@@ -182,7 +182,7 @@ namespace EpicGames.Core
 		/// dependencies which are needed because of generated code (e.g. FStructBaseChain references)
 		/// </summary>
 		public bool ValidateModuleDependencies { get; set; } = false;
-		
+
 		/// <summary>
 		/// If true, enable the input cache read
 		/// </summary>
@@ -243,7 +243,7 @@ namespace EpicGames.Core
 				int length = name.Length;
 				foreach (string define in PublicDefines)
 				{
-					if (!define.StartsWith(name, System.StringComparison.Ordinal))
+					if (!define.StartsWith(name, StringComparison.Ordinal))
 					{
 						continue;
 					}
diff --git a/Engine/Source/Programs/Shared/EpicGames.Horde.Tests/Logs/LogTests.cs b/Engine/Source/Programs/Shared/EpicGames.Horde.Tests/Logs/LogTests.cs
--- a/Engine/Source/Programs/Shared/EpicGames.Horde.Tests/Logs/LogTests.cs
+++ b/Engine/Source/Programs/Shared/EpicGames.Horde.Tests/Logs/LogTests.cs
@@ -101,7 +101,7 @@ namespace EpicGames.Horde.Tests
 		[TestMethod]
 		public async Task PartialTokenTestsAsync()
 		{
-			IStorageNamespace store = KeyValueStorageNamespace.CreateInMemory();
+			KeyValueStorageNamespace store = KeyValueStorageNamespace.CreateInMemory();
 
 			// Generate the test data
 			string[] lines =
diff --git a/Engine/Source/Programs/Shared/EpicGames.Horde/Commits/CommitId.cs b/Engine/Source/Programs/Shared/EpicGames.Horde/Commits/CommitId.cs
--- a/Engine/Source/Programs/Shared/EpicGames.Horde/Commits/CommitId.cs
+++ b/Engine/Source/Programs/Shared/EpicGames.Horde/Commits/CommitId.cs
@@ -163,7 +163,7 @@ namespace EpicGames.Horde.Commits
 		/// Creates a commit it from a Perforce changelist number. Temporary helper method for migration purposes.
 		/// </summary>
 		public static new CommitIdWithOrder? FromPerforceChange(int? change)
-			=> (change == null) ? null : CommitIdWithOrder.FromPerforceChange(change.Value);
+			=> (change == null) ? null : FromPerforceChange(change.Value);
 
 		/// <inheritdoc/>
 		public override bool Equals(object? obj)
diff --git a/Engine/Source/Programs/Shared/EpicGames.Horde/HordeHttpAuthHandler.cs b/Engine/Source/Programs/Shared/EpicGames.Horde/HordeHttpAuthHandler.cs
--- a/Engine/Source/Programs/Shared/EpicGames.Horde/HordeHttpAuthHandler.cs
+++ b/Engine/Source/Programs/Shared/EpicGames.Horde/HordeHttpAuthHandler.cs
@@ -240,7 +240,7 @@ namespace EpicGames.Horde
 
 			lock (_lockObject)
 			{
-				if (TryGetCurrentAuthState(out AuthState? authState) && Object.Equals(authState.TokenInfo?.AccessToken, accessToken))
+				if (TryGetCurrentAuthState(out AuthState? authState) && Equals(authState.TokenInfo?.AccessToken, accessToken))
 				{
 					_currentAuthResult = null;
 				}
diff --git a/Engine/Source/Programs/Shared/EpicGames.Horde/HordeOptions.cs b/Engine/Source/Programs/Shared/EpicGames.Horde/HordeOptions.cs
--- a/Engine/Source/Programs/Shared/EpicGames.Horde/HordeOptions.cs
+++ b/Engine/Source/Programs/Shared/EpicGames.Horde/HordeOptions.cs
@@ -128,7 +128,7 @@ namespace EpicGames.Horde
 		/// <param name="serverUrl">Horde server URL to use</param>
 		public static void SetDefaultServerUrl(Uri serverUrl)
 		{
-			if (!serverUrl.OriginalString.EndsWith("/", StringComparison.Ordinal))
+			if (!serverUrl.OriginalString.EndsWith('/'))
 			{
 				serverUrl = new Uri(serverUrl.OriginalString + "/");
 			}
diff --git a/Engine/Source/Programs/Shared/EpicGames.Horde/Issues/Handlers/ExternalIssueHandler.cs b/Engine/Source/Programs/Shared/EpicGames.Horde/Issues/Handlers/ExternalIssueHandler.cs
--- a/Engine/Source/Programs/Shared/EpicGames.Horde/Issues/Handlers/ExternalIssueHandler.cs
+++ b/Engine/Source/Programs/Shared/EpicGames.Horde/Issues/Handlers/ExternalIssueHandler.cs
@@ -57,7 +57,7 @@ namespace EpicGames.Horde.Issues.Handlers
 
 		static IssueFingerprint? ParseIssueFingerprint(JsonLogEvent line)
 		{
-			JsonNode? node = JsonObject.Parse(line.Data.Span, s_jsonNodeOptions);
+			JsonNode? node = JsonNode.Parse(line.Data.Span, s_jsonNodeOptions);
 			if (node != null)
 			{
 				JsonObject? properties = node["properties"] as JsonObject;
diff --git a/Engine/Source/Programs/Shared/EpicGames.Horde/Issues/Handlers/GauntletIssueHandler.cs b/Engine/Source/Programs/Shared/EpicGames.Horde/Issues/Handlers/GauntletIssueHandler.cs
--- a/Engine/Source/Programs/Shared/EpicGames.Horde/Issues/Handlers/GauntletIssueHandler.cs
+++ b/Engine/Source/Programs/Shared/EpicGames.Horde/Issues/Handlers/GauntletIssueHandler.cs
@@ -165,7 +165,7 @@ namespace EpicGames.Horde.Issues.Handlers
 		{
 			if (TryGetHash(issueEvent, out Md5Hash hash))
 			{
-				string key = $"hash:{hash}";				
+				string key = $"hash:{hash}";
 				if (!hasCallstack)
 				{
 					// add job step salt if no Callstack property was found
@@ -257,7 +257,7 @@ namespace EpicGames.Horde.Issues.Handlers
 			{
 				foreach (JsonProperty property in properties.EnumerateObject())
 				{
-					if (property.Name.StartsWith(searchType, System.StringComparison.OrdinalIgnoreCase))
+					if (property.Name.StartsWith(searchType, StringComparison.OrdinalIgnoreCase))
 					{
 						// if name is longer, check if it is a split property pattern: {name}${index}
 						if (property.Name.Length > searchType.Length && property.Name.Substring(searchType.Length, 1) != "$")
@@ -318,7 +318,7 @@ namespace EpicGames.Horde.Issues.Handlers
 				string summaryTemplate = hasCallstack ? "Automation {Meta:GauntletType} " + callstack + " {Severity} in {Meta:Node}" : "Automation {Meta:GauntletType} {Severity} in {Meta:Node}";
 
 				string gauntletType = GetEventPrefix(issueEvent.EventId!.Value);
-				IssueEventGroup issue = new IssueEventGroup($"Gauntlet:{gauntletType}",summaryTemplate, IssueChangeFilter.All);
+				IssueEventGroup issue = new IssueEventGroup($"Gauntlet:{gauntletType}", summaryTemplate, IssueChangeFilter.All);
 				issue.Events.Add(issueEvent);
 				GetHash(issueEvent, issue.Keys, issue.Metadata, hasCallstack);
 				issue.Metadata.Add("GauntletType", gauntletType);
diff --git a/Engine/Source/Programs/Shared/EpicGames.Horde/Jobs/JobMessages.cs b/Engine/Source/Programs/Shared/EpicGames.Horde/Jobs/JobMessages.cs
--- a/Engine/Source/Programs/Shared/EpicGames.Horde/Jobs/JobMessages.cs
+++ b/Engine/Source/Programs/Shared/EpicGames.Horde/Jobs/JobMessages.cs
@@ -585,7 +585,7 @@ namespace EpicGames.Horde.Jobs
 		/// </summary>
 		public CommitId? PreflightCommitId
 		{
-			get => _preflightCommitId ?? Horde.Commits.CommitId.FromPerforceChange(_preflightChange);
+			get => _preflightCommitId ?? Commits.CommitId.FromPerforceChange(_preflightChange);
 			set => _preflightCommitId = value;
 		}
 		CommitId? _preflightCommitId;
diff --git a/Engine/Source/Programs/Shared/EpicGames.Horde/Storage/Bundles/BundleStorageNamespace.cs b/Engine/Source/Programs/Shared/EpicGames.Horde/Storage/Bundles/BundleStorageNamespace.cs
--- a/Engine/Source/Programs/Shared/EpicGames.Horde/Storage/Bundles/BundleStorageNamespace.cs
+++ b/Engine/Source/Programs/Shared/EpicGames.Horde/Storage/Bundles/BundleStorageNamespace.cs
@@ -147,7 +147,7 @@ namespace EpicGames.Horde.Storage.Bundles
 			{
 				BundleSignature signature = BundleSignature.Read(data.Span);
 
-				using IRefCountedHandle<Bundles.V2.Packet> packet = Bundles.V2.Packet.Decode(data, Cache.Allocator, nameof(ReadImportsFromDataV2));
+				using IRefCountedHandle<Bundles.V2.Packet> packet = Packet.Decode(data, Cache.Allocator, nameof(ReadImportsFromDataV2));
 				for (int idx = 0; idx < packet.Target.GetImportCount(); idx++)
 				{
 					PacketImport import = packet.Target.GetImport(idx);
diff --git a/Engine/Source/Programs/Shared/EpicGames.Horde/Storage/Bundles/V1/BundleHeader.cs b/Engine/Source/Programs/Shared/EpicGames.Horde/Storage/Bundles/V1/BundleHeader.cs
--- a/Engine/Source/Programs/Shared/EpicGames.Horde/Storage/Bundles/V1/BundleHeader.cs
+++ b/Engine/Source/Programs/Shared/EpicGames.Horde/Storage/Bundles/V1/BundleHeader.cs
@@ -69,7 +69,7 @@ namespace EpicGames.Horde.Storage.Bundles.V1
 		/// <summary>
 		/// Maximum number of export refs from a single bundle
 		/// </summary>
-		public const int MaxExportRefs = MaxSectionSize / V1.BundleExportRef.NumBytes;
+		public const int MaxExportRefs = MaxSectionSize / BundleExportRef.NumBytes;
 
 		/// <summary>
 		/// Signature bytes
@@ -654,7 +654,7 @@ namespace EpicGames.Horde.Storage.Bundles.V1
 		/// Constructor
 		/// </summary>
 		public BundlePacketCollection() : this(Array.Empty<BundlePacket>()) { }
-	
+
 		/// <summary>
 		/// Constructor
 		/// </summary>
@@ -683,7 +683,7 @@ namespace EpicGames.Horde.Storage.Bundles.V1
 		{
 			byte[] data = new byte[length];
 			await stream.ReadFixedLengthBytesAsync(data, cancellationToken);
-			return BundlePacketCollection.Read(data);
+			return Read(data);
 		}
 
 		/// <inheritdoc/>
diff --git a/Engine/Source/Programs/Shared/EpicGames.Horde/Storage/Bundles/V2/BundleWriter.cs b/Engine/Source/Programs/Shared/EpicGames.Horde/Storage/Bundles/V2/BundleWriter.cs
--- a/Engine/Source/Programs/Shared/EpicGames.Horde/Storage/Bundles/V2/BundleWriter.cs
+++ b/Engine/Source/Programs/Shared/EpicGames.Horde/Storage/Bundles/V2/BundleWriter.cs
@@ -311,7 +311,7 @@ namespace EpicGames.Horde.Storage.Bundles.V2
 			{
 				_locatorPrefix = instanceSuffix;
 			}
-			else if (basePath.EndsWith("/", StringComparison.Ordinal))
+			else if (basePath.EndsWith('/'))
 			{
 				_locatorPrefix = $"{basePath}{instanceSuffix}";
 			}
diff --git a/Engine/Source/Programs/Shared/EpicGames.Horde/Storage/Git/GitObject.cs b/Engine/Source/Programs/Shared/EpicGames.Horde/Storage/Git/GitObject.cs
--- a/Engine/Source/Programs/Shared/EpicGames.Horde/Storage/Git/GitObject.cs
+++ b/Engine/Source/Programs/Shared/EpicGames.Horde/Storage/Git/GitObject.cs
@@ -46,7 +46,7 @@ namespace EpicGames.Horde.Storage.Git
 		public static async Task WriteHeaderAsync(Stream stream, Utf8String type, long size, CancellationToken cancellationToken = default)
 		{
 			byte[] header = new byte[32];
-			int length = GitObject.WriteHeader(type, size, header).Length;
+			int length = WriteHeader(type, size, header).Length;
 			await stream.WriteAsync(header.AsMemory(0, length), cancellationToken);
 		}
 
diff --git a/Engine/Source/Programs/Shared/EpicGames.Horde/Storage/KeyValueStorageNamespace.cs b/Engine/Source/Programs/Shared/EpicGames.Horde/Storage/KeyValueStorageNamespace.cs
--- a/Engine/Source/Programs/Shared/EpicGames.Horde/Storage/KeyValueStorageNamespace.cs
+++ b/Engine/Source/Programs/Shared/EpicGames.Horde/Storage/KeyValueStorageNamespace.cs
@@ -84,7 +84,7 @@ namespace EpicGames.Horde.Storage.Clients
 				_outer = outer;
 				_basePath = basePath ?? String.Empty;
 
-				if (!_basePath.EndsWith("/", StringComparison.Ordinal))
+				if (!_basePath.EndsWith('/'))
 				{
 					_basePath += "/";
 				}
@@ -158,7 +158,7 @@ namespace EpicGames.Horde.Storage.Clients
 			=> new Handle(this, locator);
 
 		/// <inheritdoc/>
-		public IBlobWriter CreateBlobWriter(string? basePath = null, BlobSerializerOptions? options = null, CancellationToken cancellationToken = default) 
+		public IBlobWriter CreateBlobWriter(string? basePath = null, BlobSerializerOptions? options = null, CancellationToken cancellationToken = default)
 			=> new BlobWriterImpl(this, basePath, options);
 
 		/// <inheritdoc/>
diff --git a/Engine/Source/Programs/Shared/EpicGames.Horde/Storage/Nodes/DirectoryNode.Update.cs b/Engine/Source/Programs/Shared/EpicGames.Horde/Storage/Nodes/DirectoryNode.Update.cs
--- a/Engine/Source/Programs/Shared/EpicGames.Horde/Storage/Nodes/DirectoryNode.Update.cs
+++ b/Engine/Source/Programs/Shared/EpicGames.Horde/Storage/Nodes/DirectoryNode.Update.cs
@@ -285,7 +285,7 @@ namespace EpicGames.Horde.Storage.Nodes
 			for (; ; )
 			{
 				FileUpdate fileUpdate = files.Current;
-				if (fileUpdate.Path.Length < prefix.Length || !fileUpdate.Path.AsSpan().StartsWith(prefix, FileReference.Comparison))
+				if (fileUpdate.Path.Length < prefix.Length || !fileUpdate.Path.AsSpan().StartsWith(prefix, FileSystemReference.Comparison))
 				{
 					return true;
 				}
diff --git a/Engine/Source/Programs/Shared/EpicGames.Horde/Storage/Workspace.cs b/Engine/Source/Programs/Shared/EpicGames.Horde/Storage/Workspace.cs
--- a/Engine/Source/Programs/Shared/EpicGames.Horde/Storage/Workspace.cs
+++ b/Engine/Source/Programs/Shared/EpicGames.Horde/Storage/Workspace.cs
@@ -162,8 +162,8 @@ namespace EpicGames.Horde.Storage
 				Parent = parent;
 				Name = name;
 
-				_directories = new Dictionary<string, DirectoryState>(DirectoryReference.Comparer);
-				_files = new Dictionary<string, FileState>(FileReference.Comparer);
+				_directories = new Dictionary<string, DirectoryState>(FileSystemReference.Comparer);
+				_files = new Dictionary<string, FileState>(FileSystemReference.Comparer);
 			}
 
 			public DirectoryState(DirectoryState? parent, IMemoryReader reader)
@@ -172,7 +172,7 @@ namespace EpicGames.Horde.Storage
 				Name = reader.ReadString();
 
 				int numDirectories = (int)reader.ReadUnsignedVarInt();
-				_directories = new Dictionary<string, DirectoryState>(numDirectories, DirectoryReference.Comparer);
+				_directories = new Dictionary<string, DirectoryState>(numDirectories, FileSystemReference.Comparer);
 
 				for (int idx = 0; idx < numDirectories; idx++)
 				{
@@ -181,7 +181,7 @@ namespace EpicGames.Horde.Storage
 				}
 
 				int numFiles = (int)reader.ReadUnsignedVarInt();
-				_files = new Dictionary<string, FileState>(numFiles, FileReference.Comparer);
+				_files = new Dictionary<string, FileState>(numFiles, FileSystemReference.Comparer);
 
 				for (int idx = 0; idx < numFiles; idx++)
 				{
@@ -961,7 +961,7 @@ namespace EpicGames.Horde.Storage
 							removeDirectories.Add(nextDirectory);
 						}
 					}
-					else if(filter.PossiblyMatches(nextNamePrefix))
+					else if (filter.PossiblyMatches(nextNamePrefix))
 					{
 						ForgetMatchingFiles(layer, nextDirectory, nextNamePrefix + "/", filter);
 					}
@@ -1067,8 +1067,8 @@ namespace EpicGames.Horde.Storage
 		{
 			public string Name { get; } = name;
 			public DirectoryState? State { get; set; }
-			public Dictionary<string, DirectoryToSync> Directories { get; } = new Dictionary<string, DirectoryToSync>(DirectoryReference.Comparer);
-			public Dictionary<string, FileToSync> Files { get; } = new Dictionary<string, FileToSync>(FileReference.Comparer);
+			public Dictionary<string, DirectoryToSync> Directories { get; } = new Dictionary<string, DirectoryToSync>(FileSystemReference.Comparer);
+			public Dictionary<string, FileToSync> Files { get; } = new Dictionary<string, FileToSync>(FileSystemReference.Comparer);
 
 			public DirectoryToSync? FindDirectoryExactCase(string name)
 				=> (Directories.TryGetValue(name, out DirectoryToSync? dirToSync) && Name.Equals(name, StringComparison.Ordinal)) ? dirToSync : null;
diff --git a/Engine/Source/Programs/Shared/EpicGames.Horde/Users/UserId.cs b/Engine/Source/Programs/Shared/EpicGames.Horde/Users/UserId.cs
--- a/Engine/Source/Programs/Shared/EpicGames.Horde/Users/UserId.cs
+++ b/Engine/Source/Programs/Shared/EpicGames.Horde/Users/UserId.cs
@@ -23,7 +23,7 @@ namespace EpicGames.Horde.Users
 		/// <summary>
 		/// Special user id for an anonymous administrator
 		/// </summary>
-		public static UserId Anonymous { get; } = UserId.Parse("63f7d3525119b9aa4c0f035a");
+		public static UserId Anonymous { get; } = Parse("63f7d3525119b9aa4c0f035a");
 
 		/// <inheritdoc cref="BinaryId.Parse(System.String)"/>
 		public static UserId Parse(string text) => new UserId(BinaryId.Parse(text));
diff --git a/Engine/Source/Programs/Shared/EpicGames.LogMatchers/CompileEventMatcher.cs b/Engine/Source/Programs/Shared/EpicGames.LogMatchers/CompileEventMatcher.cs
--- a/Engine/Source/Programs/Shared/EpicGames.LogMatchers/CompileEventMatcher.cs
+++ b/Engine/Source/Programs/Shared/EpicGames.LogMatchers/CompileEventMatcher.cs
@@ -145,7 +145,7 @@ namespace EpicGames.LogMatchers
 				// Try to match a Clang or LLVM diagnostic
 				Match? match;
 				if ((builder.Current.TryMatch(s_clangDiagnosticPattern, out match) ||
-				     builder.Current.TryMatch(s_llvmDiagnosticPattern, out match)) && IsSourceFile(match))
+					 builder.Current.TryMatch(s_llvmDiagnosticPattern, out match)) && IsSourceFile(match))
 				{
 					LogLevel level = GetLogLevelFromSeverity(match);
 
@@ -159,7 +159,7 @@ namespace EpicGames.LogMatchers
 						SkipClangMarker(builder);
 
 						if (!(builder.Next.TryMatch(s_clangNotePattern, out match) ||
-						      builder.Next.TryMatch(s_llvmNotePattern, out match)))
+							  builder.Next.TryMatch(s_llvmNotePattern, out match)))
 						{
 							break;
 						}
@@ -226,7 +226,6 @@ namespace EpicGames.LogMatchers
 				outEvent = builder.ToMatch(LogEventPriority.Highest, LogLevel.Information, KnownLogEvents.Systemic_XCode);
 				return true;
 			}
-			
 
 			outEvent = null;
 			return false;
@@ -280,7 +279,7 @@ namespace EpicGames.LogMatchers
 						return true;
 					}
 				}
-				else if (codeGroup.Value.StartsWith("C", StringComparison.Ordinal))
+				else if (codeGroup.Value.StartsWith('C'))
 				{
 					if (codeGroup.Value.Equals("C1060", StringComparison.Ordinal))
 					{
diff --git a/Engine/Source/Programs/Shared/EpicGames.OIDC/WindowsTokenStore.cs b/Engine/Source/Programs/Shared/EpicGames.OIDC/WindowsTokenStore.cs
--- a/Engine/Source/Programs/Shared/EpicGames.OIDC/WindowsTokenStore.cs
+++ b/Engine/Source/Programs/Shared/EpicGames.OIDC/WindowsTokenStore.cs
@@ -233,15 +233,17 @@ namespace EpicGames.OIDC
 #pragma warning disable IDE1006 // Pinvoke code doesnt use the same naming conventions as C#
 	static class CryptProtectDataHelper
 	{
-		
-		[StructLayout(LayoutKind.Sequential, CharSet=CharSet.Unicode)]
-		private struct DataBlob {
+
+		[StructLayout(LayoutKind.Sequential, CharSet = CharSet.Unicode)]
+		private struct DataBlob
+		{
 			public int cbData;
 			public IntPtr pbData;
 		}
 
 		[Flags]
-		private enum CryptProtectFlags {
+		private enum CryptProtectFlags
+		{
 			// for remote-access situations where ui is not an option
 			// if UI was specified on protect or unprotect operation, the call
 			// will fail and GetLastError() will indicate ERROR_PASSWORD_RESTRICTION
@@ -266,7 +268,8 @@ namespace EpicGames.OIDC
 		}
 
 		[Flags]
-		private enum CryptProtectPromptFlags {
+		private enum CryptProtectPromptFlags
+		{
 			// prompt on unprotect
 			CryptprotectPromptOnUnprotect = 0x1,
 
@@ -274,8 +277,9 @@ namespace EpicGames.OIDC
 			CryptprotectPromptOnProtect = 0x2
 		}
 
-		[StructLayout(LayoutKind.Sequential, CharSet=CharSet.Unicode)]
-		private struct CryptprotectPromptstruct {
+		[StructLayout(LayoutKind.Sequential, CharSet = CharSet.Unicode)]
+		private struct CryptprotectPromptstruct
+		{
 			public int cbSize;
 			public CryptProtectPromptFlags dwPromptFlags;
 			public IntPtr hwndApp;
@@ -284,8 +288,8 @@ namespace EpicGames.OIDC
 
 		[
 			DllImport("Crypt32.dll",
-				SetLastError=true,
-				CharSet=System.Runtime.InteropServices.CharSet.Auto)
+				SetLastError = true,
+				CharSet = CharSet.Auto)
 		]
 		[return: MarshalAs(UnmanagedType.Bool)]
 		private static extern bool CryptProtectData(
@@ -300,8 +304,8 @@ namespace EpicGames.OIDC
 
 		[
 			DllImport("Crypt32.dll",
-				SetLastError=true,
-				CharSet=System.Runtime.InteropServices.CharSet.Auto)
+				SetLastError = true,
+				CharSet = CharSet.Auto)
 		]
 		[return: MarshalAs(UnmanagedType.Bool)]
 		private static extern bool CryptUnprotectData(
@@ -349,8 +353,7 @@ namespace EpicGames.OIDC
 				dataHandle.Free();
 				entropyHandle.Free();
 			}
-		
-		
+
 			byte[] buf = new byte[dataOut.cbData];
 			Marshal.Copy(dataOut.pbData, buf, 0, dataOut.cbData);
 			return buf;
@@ -391,8 +394,7 @@ namespace EpicGames.OIDC
 				dataHandle.Free();
 				entropyHandle.Free();
 			}
-		
-		
+
 			byte[] buf = new byte[dataOut.cbData];
 			Marshal.Copy(dataOut.pbData, buf, 0, dataOut.cbData);
 			return buf;
diff --git a/Engine/Source/Programs/Shared/EpicGames.Perforce.Managed/ManagedWorkspace.cs b/Engine/Source/Programs/Shared/EpicGames.Perforce.Managed/ManagedWorkspace.cs
--- a/Engine/Source/Programs/Shared/EpicGames.Perforce.Managed/ManagedWorkspace.cs
+++ b/Engine/Source/Programs/Shared/EpicGames.Perforce.Managed/ManagedWorkspace.cs
@@ -169,7 +169,7 @@ namespace EpicGames.Perforce.Managed
 		/// The current revision number for cache archives.
 		/// </summary>
 		static int CurrentVersion { get; } = Enum.GetValues(typeof(ManagedWorkspaceVersion)).Cast<int>().Max();
-		
+
 		/// <summary>
 		/// File format version for untracked state file
 		/// </summary>
@@ -194,7 +194,7 @@ namespace EpicGames.Perforce.Managed
 		/// Name of the main data file for a repository
 		/// </summary>
 		const string DataFileName = "Repository.dat";
-		
+
 		/// <summary>
 		/// Name of the untracked state file for a repository
 		/// </summary>
@@ -214,12 +214,12 @@ namespace EpicGames.Perforce.Managed
 		/// Whether a repair operation should be run on this workspace. Set whenever the state may be inconsistent.
 		/// </summary>
 		bool _requiresRepair;
-		
+
 		/// <summary>
 		/// Tracer
 		/// </summary>
 		readonly Tracer _tracer;
-		
+
 		/// <summary>
 		/// The log output device
 		/// </summary>
@@ -259,12 +259,12 @@ namespace EpicGames.Perforce.Managed
 		/// All the files which are currently being tracked
 		/// </summary>
 		Dictionary<FileContentId, CachedFileInfo> _contentIdToTrackedFile = new Dictionary<FileContentId, CachedFileInfo>();
-		
+
 		/// <summary>
 		/// Notification that a clean has been performed
 		/// </summary>
 		public delegate void OnCleanDelegate(int numFilesDeleted, int numDirsDeleted);
-		
+
 		/// <summary>
 		/// Occurs when a clean operation has been performed.
 		/// </summary>
@@ -428,7 +428,7 @@ namespace EpicGames.Perforce.Managed
 		public static async Task<ManagedWorkspace> LoadAsync(string hostName, DirectoryReference baseDir, ManagedWorkspaceOptions options, Tracer tracer, ILogger logger, CancellationToken cancellationToken)
 		{
 			using TelemetrySpan span = tracer.StartActiveSpan($"{nameof(ManagedWorkspace)}.{nameof(LoadAsync)}");
-			
+
 			if (!Exists(baseDir))
 			{
 				throw new FatalErrorException("No valid repository found at {0}", baseDir);
@@ -472,7 +472,7 @@ namespace EpicGames.Perforce.Managed
 		private async Task SaveAsync(TransactionState state, CancellationToken cancellationToken)
 		{
 			using TelemetrySpan span = _tracer.StartActiveSpan($"{nameof(ManagedWorkspace)}.{nameof(SaveAsync)}");
-			
+
 			// Allocate the buffer for writing
 			int serializedSize = sizeof(int) + sizeof(byte) + sizeof(int) + sizeof(int) + _contentIdToTrackedFile.Values.Sum(x => x.GetSerializedSize()) + _workspace.GetSerializedSize();
 			byte[] buffer = new byte[serializedSize];
@@ -496,7 +496,7 @@ namespace EpicGames.Perforce.Managed
 			await FileReference.WriteAllBytesAsync(dataFile, buffer, cancellationToken);
 			CompleteTransaction(dataFile);
 		}
-		
+
 		/// <summary>
 		/// Mark whether current repository in combination with P4 client contains untracked files
 		/// Stores data in a file separate to the repository metadata as that can grow large (hundreds of megabytes)
@@ -507,20 +507,20 @@ namespace EpicGames.Perforce.Managed
 			{
 				throw new ArgumentException("Unable to save untracked state as P4 client name is not set");
 			}
-			
+
 			using MemoryStream ms = new(300);
 			await using BinaryWriter bw = new(ms);
 			bw.Write(UntrackedFileVersion);
 			bw.Write(connection.Settings.ClientName);
 			bw.Write(containsUntrackedFiles);
-			
+
 			// Write it to disk
 			FileReference dataFile = FileReference.Combine(_baseDir, UntrackedStateFileName);
 			BeginTransaction(dataFile);
 			await FileReference.WriteAllBytesAsync(dataFile, ms.ToArray(), cancellationToken);
 			CompleteTransaction(dataFile);
 		}
-		
+
 		/// <summary>
 		/// Check if current repository in combination with P4 client contains untracked files
 		/// If deserialization is not possible, assume the repo contains untracked files 
@@ -535,17 +535,17 @@ namespace EpicGames.Perforce.Managed
 				{
 					return true;
 				}
-				
+
 				byte[] data = await File.ReadAllBytesAsync(dataFile.FullName, cancellationToken);
-				using MemoryStream ms = new (data);
-				using BinaryReader br = new (ms);
-				
+				using MemoryStream ms = new(data);
+				using BinaryReader br = new(ms);
+
 				int version = br.ReadInt32();
 				if (version != UntrackedFileVersion)
 				{
 					return true;
 				}
-				
+
 				string clientName = br.ReadString();
 				bool containsUntrackedFiles = br.ReadBoolean();
 				return clientName != connection.Settings.ClientName || containsUntrackedFiles;
@@ -555,7 +555,7 @@ namespace EpicGames.Perforce.Managed
 				return true;
 			}
 		}
-		
+
 		#region Commands
 
 		/// <summary>
@@ -573,7 +573,7 @@ namespace EpicGames.Perforce.Managed
 			{
 				await CleanInternalAsync(removeUntracked, cancellationToken);
 			}
-			
+
 			if (perforce != null && removeUntracked)
 			{
 				await SaveUntrackedStateAsync(perforce, false, cancellationToken);
@@ -590,7 +590,7 @@ namespace EpicGames.Perforce.Managed
 		private async Task CleanInternalAsync(bool removeUntracked, CancellationToken cancellationToken)
 		{
 			using TelemetrySpan span = _tracer.StartActiveSpan($"{nameof(ManagedWorkspace)}.{nameof(CleanInternalAsync)}");
-			
+
 			FileInfo[] filesToDelete;
 			DirectoryInfo[] directoriesToDelete;
 			using (Trace("FindFilesToClean"))
@@ -639,7 +639,7 @@ namespace EpicGames.Perforce.Managed
 							return ValueTask.CompletedTask;
 						});
 					}
-					
+
 					{
 						using TelemetrySpan cleanDirsSpan = _tracer.StartActiveSpan($"{nameof(ManagedWorkspace)}.CleaningDirs");
 						await Parallel.ForEachAsync(directoriesToDelete, options, (directoryToDelete, ct) =>
@@ -654,7 +654,7 @@ namespace EpicGames.Perforce.Managed
 
 				await SaveAsync(TransactionState.Clean, cancellationToken);
 			}
-			
+
 			OnClean?.Invoke(filesToDelete.Length, directoriesToDelete.Length);
 		}
 
@@ -714,7 +714,7 @@ namespace EpicGames.Perforce.Managed
 		public async Task RepairCacheAsync(CancellationToken cancellationToken)
 		{
 			using TelemetrySpan span = _tracer.StartActiveSpan($"{nameof(ManagedWorkspace)}.{nameof(RepairCacheAsync)}");
-			
+
 			using (Trace("Repair"))
 			using (ILoggerProgress status = _logger.BeginProgressScope("Checking cache..."))
 			{
@@ -821,7 +821,7 @@ namespace EpicGames.Perforce.Managed
 		public async Task StatsAsync(IPerforceConnection perforceClient, List<string> streamNames, List<string> view, CancellationToken cancellationToken)
 		{
 			using TelemetrySpan span = _tracer.StartActiveSpan($"{nameof(ManagedWorkspace)}.{nameof(StatsAsync)}");
-			
+
 			_logger.LogInformation("Finding stats for {NumStreams} streams", streamNames.Count);
 			using (_logger.BeginIndentScope("  "))
 			{
@@ -929,7 +929,7 @@ namespace EpicGames.Perforce.Managed
 		public async Task StatusAsync()
 		{
 			using TelemetrySpan span = _tracer.StartActiveSpan($"{nameof(ManagedWorkspace)}.{nameof(StatusAsync)}");
-			
+
 			// Print size stats
 			_logger.LogInformation("Cache contains {NumFiles:n0} files, {TotalSize:n1}mb", _contentIdToTrackedFile.Count, _contentIdToTrackedFile.Values.Sum(x => x.Length) / (1024.0 * 1024.0));
 			_logger.LogInformation("Stage contains {NumFiles:n0} files, {TotalSize:n1}mb", _workspace.GetFiles().Count, _workspace.GetFiles().Sum(x => x._length) / (1024.0 * 1024.0));
@@ -941,15 +941,15 @@ namespace EpicGames.Perforce.Managed
 				_logger.LogInformation("Local changes:");
 				foreach (string difference in differences)
 				{
-					if (difference.StartsWith("+", StringComparison.Ordinal))
+					if (difference.StartsWith('+'))
 					{
 						Console.ForegroundColor = ConsoleColor.Green;
 					}
-					else if (difference.StartsWith("-", StringComparison.Ordinal))
+					else if (difference.StartsWith('-'))
 					{
 						Console.ForegroundColor = ConsoleColor.Red;
 					}
-					else if (difference.StartsWith("!", StringComparison.Ordinal))
+					else if (difference.StartsWith('!'))
 					{
 						Console.ForegroundColor = ConsoleColor.Yellow;
 					}
@@ -978,7 +978,7 @@ namespace EpicGames.Perforce.Managed
 		{
 			using TelemetrySpan span = _tracer.StartActiveSpan($"{nameof(ManagedWorkspace)}.{nameof(SyncAsync)}");
 			span.SetAttribute("horde.mw.use_have_table", _options.UseHaveTable);
-			
+
 			Stopwatch timer = Stopwatch.StartNew();
 			if (changeNumber == -1)
 			{
@@ -1015,7 +1015,7 @@ namespace EpicGames.Perforce.Managed
 					{
 						await CleanInternalAsync(removeUntracked, cancellationToken);
 					}
-					
+
 					// Immediately after a potential clean, mark repository having untracked files.
 					// This ensure a clean is run next time even if the sync below is aborted or fails
 					await SaveUntrackedStateAsync(perforce, true, cancellationToken);
@@ -1047,7 +1047,7 @@ namespace EpicGames.Perforce.Managed
 				await RemoveFilesFromWorkspaceAsync(contents, cancellationToken);
 				await AddFilesToWorkspaceAsync(perforce, contents, fakeSync, cancellationToken);
 			}
-			
+
 			_logger.LogInformation("Completed in {ElapsedTime}s", $"{timer.Elapsed.TotalSeconds:0.0}");
 		}
 
@@ -1058,7 +1058,7 @@ namespace EpicGames.Perforce.Managed
 		public async Task UnshelveAsync(IPerforceConnection perforce, int unshelveChangelist, CancellationToken cancellationToken)
 		{
 			using TelemetrySpan span = _tracer.StartActiveSpan($"{nameof(ManagedWorkspace)}.{nameof(UnshelveAsync)}");
-			
+
 			// Need to mark those files as dirty - update the workspace with those files 
 			// Delete is fine, but need to flag anything added
 
@@ -1169,7 +1169,7 @@ namespace EpicGames.Perforce.Managed
 		public async Task<Tuple<int, StreamSnapshot>[]> PopulateCleanAsync(List<PopulateRequest> requests, CancellationToken cancellationToken)
 		{
 			using TelemetrySpan span = _tracer.StartActiveSpan($"{nameof(ManagedWorkspace)}.{nameof(PopulateCleanAsync)}");
-			
+
 			// Revert all changes in each of the unique clients
 			foreach (PopulateRequest request in requests)
 			{
@@ -1274,7 +1274,7 @@ namespace EpicGames.Perforce.Managed
 		public async Task PopulateSyncAsync(List<PopulateRequest> requests, Tuple<int, StreamSnapshot>[] streamState, bool fakeSync, CancellationToken cancellationToken)
 		{
 			using TelemetrySpan span = _tracer.StartActiveSpan($"{nameof(ManagedWorkspace)}.{nameof(PopulateSyncAsync)}");
-			
+
 			// Sync all the new files
 			for (int idx = 0; idx < requests.Count; idx++)
 			{
@@ -1349,7 +1349,7 @@ namespace EpicGames.Perforce.Managed
 
 			if (_options.UseModTime)
 			{
-				client.Options = ClientOptions.ModTime;	
+				client.Options = ClientOptions.ModTime;
 			}
 
 			if (_options.Partitioned)
@@ -1467,7 +1467,7 @@ namespace EpicGames.Perforce.Managed
 				List<OpenedRecord> openedFilesResponse = await perforceClient.OpenedAsync(OpenedOptions.ShortOutput, -1, perforceClient.Settings.ClientName!, null, 1, FileSpecList.Any, cancellationToken).ToListAsync(cancellationToken);
 
 				// If there are any files, revert them
-				if (openedFilesResponse.Any())
+				if (openedFilesResponse.Count != 0)
 				{
 					await perforceClient.RevertAsync(-1, null, RevertOptions.KeepWorkspaceFiles, new[] { "//..." }, cancellationToken);
 				}
@@ -1529,7 +1529,7 @@ namespace EpicGames.Perforce.Managed
 				Stopwatch timer = Stopwatch.StartNew();
 
 				// Sync an initial set of files. Either start with a full workspace and remove files, or start with nothing and add files.
-				if (view.Count == 0 || view[0].StartsWith("-", StringComparison.Ordinal))
+				if (view.Count == 0 || view[0].StartsWith('-'))
 				{
 					await UpdateHaveTablePathAsync(perforceClient, $"//{perforceClient.Settings.ClientName}/...@{changeNumber}", _tracer, cancellationToken);
 				}
@@ -1542,7 +1542,7 @@ namespace EpicGames.Perforce.Managed
 				foreach (string filter in view)
 				{
 					string syncPath;
-					if (filter.StartsWith("-", StringComparison.Ordinal))
+					if (filter.StartsWith('-'))
 					{
 						syncPath = String.Format("//{0}/{1}#0", perforceClient.Settings.ClientName, RemoveLeadingSlash(filter.Substring(1)));
 					}
@@ -1569,8 +1569,8 @@ namespace EpicGames.Perforce.Managed
 		{
 			using TelemetrySpan span = tracer.StartActiveSpan($"{nameof(ManagedWorkspace)}.{nameof(UpdateHaveTablePathAsync)}");
 			span.SetAttribute("horde.sync_path", syncPath);
-			
-            PerforceResponseList<SyncSummaryRecord> responseList = await perforceClient.TrySyncQuietAsync(SyncOptions.KeepWorkspaceFiles, -1, new[] { syncPath }, cancellationToken);
+
+			PerforceResponseList<SyncSummaryRecord> responseList = await perforceClient.TrySyncQuietAsync(SyncOptions.KeepWorkspaceFiles, -1, new[] { syncPath }, cancellationToken);
 			foreach (PerforceResponse<SyncSummaryRecord> response in responseList)
 			{
 				PerforceError? error = response.Error;
@@ -1627,7 +1627,7 @@ namespace EpicGames.Perforce.Managed
 		private async Task<StreamSnapshotFromMemory> FindClientContentsAsync(IPerforceConnection perforceClient, int changeNumber, CancellationToken cancellationToken)
 		{
 			StreamTreeBuilder builder = new StreamTreeBuilder();
-			
+
 			using TelemetrySpan span = _tracer.StartActiveSpan($"{nameof(ManagedWorkspace)}.{nameof(FindClientContentsAsync)}");
 			using (Trace("FetchMetadata"))
 			using (ILoggerProgress scope = _logger.BeginProgressScope("Fetching metadata..."))
@@ -1794,7 +1794,7 @@ namespace EpicGames.Perforce.Managed
 		public async Task<StreamSnapshotFromMemory> FindClientContentsWithoutHaveTableAsync(IPerforceConnection perforceClient, string streamName, IReadOnlyList<string> view, int changeNumber, CancellationToken cancellationToken)
 		{
 			DepotStreamTreeBuilder builder = new();
-			
+
 			using TelemetrySpan span = _tracer.StartActiveSpan($"{nameof(ManagedWorkspace)}.{nameof(FindClientContentsWithoutHaveTableAsync)}");
 			using (Trace("FetchMetadata"))
 			using (ILoggerProgress scope = _logger.BeginProgressScope("Fetching metadata (without have table)..."))
@@ -1896,7 +1896,7 @@ namespace EpicGames.Perforce.Managed
 			StreamSnapshotFromMemory contents = _options.UseHaveTable
 				? await FindClientContentsAsync(perforceClient, changeNumber, cancellationToken)
 				: await FindClientContentsWithoutHaveTableAsync(perforceClient, basePath.ToString(), view, changeNumber, cancellationToken);
-			
+
 			using TelemetrySpan span = _tracer.StartActiveSpan($"{nameof(ManagedWorkspace)}.{nameof(FindAndSaveClientContentsAsync)}");
 			using (Trace("WriteMetadata"))
 			using (ILoggerProgress scope = _logger.BeginProgressScope($"Saving metadata to {cacheFile}..."))
@@ -2161,7 +2161,7 @@ namespace EpicGames.Perforce.Managed
 				if (RuntimeInformation.IsOSPlatform(OSPlatform.Windows))
 				{
 					ulong freeSpace;
-            		bool success = GetDiskFreeSpaceEx(_baseDir.FullName, out _, out _, out freeSpace);
+					bool success = GetDiskFreeSpaceEx(_baseDir.FullName, out _, out _, out freeSpace);
 					if (!success)
 					{
 						throw new InsufficientSpaceException($"Can't get free space for {_baseDir.FullName}");
@@ -2387,9 +2387,9 @@ namespace EpicGames.Perforce.Managed
 					DateTime lastWriteTime = _options.UseModTime
 						? new DateTime(fileToMove._workspaceFile._lastModifiedTicks, DateTimeKind.Utc)
 						: DateTime.UtcNow;
-					
+
 					FileReference.SetLastWriteTimeUtc(targetFile, lastWriteTime);
-					fileToMove._workspaceFile.UpdateMetadata();	
+					fileToMove._workspaceFile.UpdateMetadata();
 				}
 				catch
 				{
diff --git a/Engine/Source/Programs/Shared/EpicGames.Perforce/PerforceReflection.cs b/Engine/Source/Programs/Shared/EpicGames.Perforce/PerforceReflection.cs
--- a/Engine/Source/Programs/Shared/EpicGames.Perforce/PerforceReflection.cs
+++ b/Engine/Source/Programs/Shared/EpicGames.Perforce/PerforceReflection.cs
@@ -507,7 +507,7 @@ namespace EpicGames.Perforce
 						}
 						rootRecord.RequiredTagsBitMask |= requiredTagBitMask;
 					}
-					
+
 					TaggedPropertyInfo tagInfo = new TaggedPropertyInfo(new Utf8String(tagName), tagAttribute.Optional, parentRecords, property, requiredTagBitMask);
 
 					Type fieldType = property.PropertyType;
@@ -515,7 +515,7 @@ namespace EpicGames.Perforce
 					PropertyInfo propertyCopy = property;
 					if (fieldType == typeof(DateTime))
 					{
-						tagInfo.Write = (writer, value) => WriteUtf8StringWithTag(writer, new Utf8String(((long)((DateTime)value - PerforceReflection.UnixEpoch).TotalSeconds).ToString()));
+						tagInfo.Write = (writer, value) => WriteUtf8StringWithTag(writer, new Utf8String(((long)((DateTime)value - UnixEpoch).TotalSeconds).ToString()));
 						tagInfo.ReadFromString = (obj, value) => propertyCopy.SetValue(obj, ParseStringAsDateTime(value));
 					}
 					else if (fieldType == typeof(bool))
@@ -557,7 +557,7 @@ namespace EpicGames.Perforce
 					}
 					else if (fieldType == typeof(DateTimeOffset?))
 					{
-						tagInfo.ReadFromString = (obj, value) => propertyCopy.SetValue(obj, ParseStringAsNullableDateTimeOffset(value));
+						tagInfo.ReadFromString = (obj, value) => propertyCopy.SetValue(obj, ParseStringAsDateTimeOffset(value));
 					}
 					else if (fieldType == typeof(List<string>))
 					{
@@ -600,12 +600,12 @@ namespace EpicGames.Perforce
 			return (CreateRecordDelegate)dynamicMethod.CreateDelegate(typeof(CreateRecordDelegate));
 		}
 
-		static object ParseString(Utf8String str)
+		static string ParseString(Utf8String str)
 		{
 			return str.ToString();
 		}
 
-		static object ParseStringAsDateTime(Utf8String str)
+		static DateTime ParseStringAsDateTime(Utf8String str)
 		{
 			string text = str.ToString();
 
@@ -616,7 +616,7 @@ namespace EpicGames.Perforce
 			}
 			else
 			{
-				return PerforceReflection.UnixEpoch + TimeSpan.FromSeconds(Int64.Parse(text));
+				return UnixEpoch + TimeSpan.FromSeconds(Int64.Parse(text));
 			}
 		}
 
@@ -656,7 +656,7 @@ namespace EpicGames.Perforce
 			}
 		}
 
-		static object ParseStringAsLong(Utf8String str)
+		static long ParseStringAsLong(Utf8String str)
 		{
 			long value;
 			int bytesConsumed;
@@ -667,7 +667,7 @@ namespace EpicGames.Perforce
 			return value;
 		}
 
-		static object ParseStringAsNullableDateTimeOffset(Utf8String str)
+		static DateTimeOffset ParseStringAsDateTimeOffset(Utf8String str)
 		{
 			string text = str.ToString();
 			return DateTimeOffset.Parse(Regex.Replace(text, "[^0-9]*$", "")); // Strip timezone name (eg. "EST")
diff --git a/Engine/Source/Programs/Shared/EpicGames.Perforce/PerforceSettings.cs b/Engine/Source/Programs/Shared/EpicGames.Perforce/PerforceSettings.cs
--- a/Engine/Source/Programs/Shared/EpicGames.Perforce/PerforceSettings.cs
+++ b/Engine/Source/Programs/Shared/EpicGames.Perforce/PerforceSettings.cs
@@ -124,7 +124,7 @@ namespace EpicGames.Perforce
 		public PerforceSettings(IPerforceEnvironment environment)
 		{
 			ServerAndPort = environment.GetValue("P4PORT") ?? "perforce:1666";
-			UserName = environment.GetValue("P4USER") ?? System.Environment.UserName;
+			UserName = environment.GetValue("P4USER") ?? Environment.UserName;
 			Password = environment.GetValue("P4PASSWD");
 			HostName = environment.GetValue("P4HOST");
 			ClientName = environment.GetValue("P4CLIENT");
diff --git a/Engine/Source/Programs/Shared/EpicGames.ProjectStore/ProjectStoreTypes.cs b/Engine/Source/Programs/Shared/EpicGames.ProjectStore/ProjectStoreTypes.cs
--- a/Engine/Source/Programs/Shared/EpicGames.ProjectStore/ProjectStoreTypes.cs
+++ b/Engine/Source/Programs/Shared/EpicGames.ProjectStore/ProjectStoreTypes.cs
@@ -94,7 +94,7 @@ namespace EpicGames.ProjectStore
 	/// <summary>
 	/// Exception indicating errors when finding or importing snapshots.
 	/// </summary>
-	public class BuildIndexException : Exception 
+	public class BuildIndexException : Exception
 	{
 		/// <summary>
 		/// Default constructor
@@ -127,16 +127,16 @@ namespace EpicGames.ProjectStore
 	/// Base class for all builds we could want to import, whether discovered via looking them up in the P:/ drive or
 	/// discovered via the Unreal Cloud DDC builds API.
 	/// </summary>
-	#pragma warning disable CA1724
+#pragma warning disable CA1724
 	public abstract class Build
-	#pragma warning restore CA1724
+#pragma warning restore CA1724
 	{
 		/// <summary>
 		/// Delegate used to run command-line applications
 		/// </summary>
-		#pragma warning disable CA1051	
+#pragma warning disable CA1051
 		protected readonly RunAppDelegate RunApp;
-		#pragma warning restore CA1051	
+#pragma warning restore CA1051
 
 		/// <summary>
 		/// Constructor.
@@ -209,15 +209,12 @@ namespace EpicGames.ProjectStore
 		[System.Diagnostics.CodeAnalysis.SuppressMessage("Usage", "CA1308", Justification = "We know the hashed strings are safe to lowercase using ToLowerInvariant()")]
 		protected static string GetProjectID(FileReference projectFile)
 		{
-			using (System.Security.Cryptography.MD5 hasher = System.Security.Cryptography.MD5.Create())
-			{
-				byte[] bytes = System.Text.Encoding.UTF8.GetBytes(projectFile.FullName.Replace("\\", "/", StringComparison.InvariantCulture));
-				byte[] hashed = hasher.ComputeHash(bytes);
+			byte[] bytes = Encoding.UTF8.GetBytes(projectFile.FullName.Replace("\\", "/", StringComparison.InvariantCulture));
+			byte[] hashed = System.Security.Cryptography.MD5.HashData(bytes);
 
-				string hexString = Convert.ToHexString(hashed).ToLowerInvariant();
+			string hexString = Convert.ToHexString(hashed).ToLowerInvariant();
 
-				return projectFile.GetFileNameWithoutAnyExtensions() + "." + hexString.Substring(0, 8);
-			}
+			return projectFile.GetFileNameWithoutAnyExtensions() + "." + hexString.Substring(0, 8);
 		}
 
 		/// <summary>
@@ -600,9 +597,9 @@ namespace EpicGames.ProjectStore
 		/// <summary>
 		/// A delegate used to run command-line apps. See documentation of RunAppDelegate for details.
 		/// </summary>
-		#pragma warning disable CA1051	
+#pragma warning disable CA1051
 		protected readonly RunAppDelegate RunApp;
-		#pragma warning restore CA1051	
+#pragma warning restore CA1051
 
 		/// <summary>
 		/// Retrieves a list of changelists for a given branch/project/platform/runtime that can be imported
@@ -674,7 +671,7 @@ namespace EpicGames.ProjectStore
 
 			string runtimeSuffix = (runtime == "game" ? "" : textInfo.ToTitleCase(runtime));
 
-			if (platform.ToLower() == "win64")
+			if (platform.Equals("win64", StringComparison.OrdinalIgnoreCase))
 			{
 				buildsDir = DirectoryReference.Combine(buildsDir, "Windows" + runtimeSuffix);
 			}
@@ -850,7 +847,7 @@ namespace EpicGames.ProjectStore
 
 			IConfiguration config = ProviderConfigurationFactory.BindOptions(authConfig);
 			using (ITokenStore tokenStore = TokenStoreFactory.CreateTokenStore())
-			{ 
+			{
 				OidcTokenManager tokenManager = OidcTokenManager.CreateTokenManager(config, tokenStore);
 
 				try
@@ -942,7 +939,7 @@ namespace EpicGames.ProjectStore
 				platform = "linuxserver";
 			}
 
-			return String.Format("{0}.oplog.{1}.{2}",  project, branch, platform);
+			return String.Format("{0}.oplog.{1}.{2}", project, branch, platform);
 		}
 
 		/// <summary>
diff --git a/Engine/Source/Programs/Shared/EpicGames.Serialization/CbField.cs b/Engine/Source/Programs/Shared/EpicGames.Serialization/CbField.cs
--- a/Engine/Source/Programs/Shared/EpicGames.Serialization/CbField.cs
+++ b/Engine/Source/Programs/Shared/EpicGames.Serialization/CbField.cs
@@ -1318,7 +1318,7 @@ namespace EpicGames.Serialization
 				return defaultValue;
 			}
 		}
-		
+
 		/// <summary>
 		/// Access the field as a UUID. Returns a nil UUID on error.
 		/// </summary>
@@ -1606,7 +1606,7 @@ namespace EpicGames.Serialization
 		/// <param name="name"></param>
 		/// <returns></returns>
 #pragma warning disable CA1043 // Use Integral Or String Argument For Indexers
-		public CbField this[CbFieldName name] => this.FirstOrDefault(field => field.Name == name.Text) ?? CbField.Empty;
+		public CbField this[CbFieldName name] => this.FirstOrDefault(field => field.Name == name.Text) ?? Empty;
 #pragma warning restore CA1043 // Use Integral Or String Argument For Indexers
 
 		/// <summary>
@@ -2426,7 +2426,7 @@ namespace EpicGames.Serialization
 		/// <summary>
 		/// Empty array constant
 		/// </summary>
-		public static CbObject Empty { get; } = CbObject.FromFieldNoCheck(new CbField(new byte[] { (byte)CbFieldType.Object, 0 }));
+		public static CbObject Empty { get; } = FromFieldNoCheck(new CbField(new byte[] { (byte)CbFieldType.Object, 0 }));
 
 		/// <summary>
 		/// The inner field object
@@ -2659,7 +2659,6 @@ namespace EpicGames.Serialization
 			writer.WriteEndObject();
 		}
 
-		
 		/// <summary>
 		/// Write this object to JSON
 		/// </summary>
@@ -2740,18 +2739,18 @@ namespace EpicGames.Serialization
 		{
 			if (!StringUtils.TryParseHexString(text, out byte[]? bytes))
 			{
-				objectId = CbObjectId.Zero;
+				objectId = Zero;
 				return false;
 			}
 			if (bytes.Length != 12)
 			{
-				objectId = CbObjectId.Zero;
+				objectId = Zero;
 				return false;
 			}
 			objectId = new CbObjectId(bytes);
 			return true;
 		}
-		
+
 		/// <summary>
 		/// Generates a new object id
 		/// </summary>
diff --git a/Engine/Source/Programs/Shared/EpicGames.Serialization/CbSerializer.cs b/Engine/Source/Programs/Shared/EpicGames.Serialization/CbSerializer.cs
--- a/Engine/Source/Programs/Shared/EpicGames.Serialization/CbSerializer.cs
+++ b/Engine/Source/Programs/Shared/EpicGames.Serialization/CbSerializer.cs
@@ -231,173 +231,173 @@ namespace EpicGames.Serialization
 			switch (jsonElement.ValueKind)
 			{
 				case JsonValueKind.Object:
-				{
-					if (String.IsNullOrEmpty(fieldName))
 					{
-						writer.BeginObject();
-					}
-					else
-					{
-						writer.BeginObject(fieldName);
-					}
+						if (String.IsNullOrEmpty(fieldName))
+						{
+							writer.BeginObject();
+						}
+						else
+						{
+							writer.BeginObject(fieldName);
+						}
 
-					foreach (JsonProperty prop in jsonElement.EnumerateObject())
-					{
-						ReadJsonField(prop.Value, prop.Name, writer);
-					}
+						foreach (JsonProperty prop in jsonElement.EnumerateObject())
+						{
+							ReadJsonField(prop.Value, prop.Name, writer);
+						}
 
-					writer.EndObject();
-					break;
-				}
-				case JsonValueKind.Array:
-				{
-					if (String.IsNullOrEmpty(fieldName))
-					{
-						writer.BeginArray();
+						writer.EndObject();
+						break;
 					}
-					else
+				case JsonValueKind.Array:
 					{
-						writer.BeginArray(fieldName);
-					}
+						if (String.IsNullOrEmpty(fieldName))
+						{
+							writer.BeginArray();
+						}
+						else
+						{
+							writer.BeginArray(fieldName);
+						}
 
-					foreach (JsonElement prop in jsonElement.EnumerateArray())
-					{
-						ReadJsonField(prop, null, writer);
-					}
+						foreach (JsonElement prop in jsonElement.EnumerateArray())
+						{
+							ReadJsonField(prop, null, writer);
+						}
 
-					writer.EndArray();
-					break;
-				}
-				case JsonValueKind.Null:
-				{
-					if (String.IsNullOrEmpty(fieldName))
-					{
-						writer.WriteNullValue();
+						writer.EndArray();
+						break;
 					}
-					else
+				case JsonValueKind.Null:
 					{
-						writer.WriteNull(fieldName);
+						if (String.IsNullOrEmpty(fieldName))
+						{
+							writer.WriteNullValue();
+						}
+						else
+						{
+							writer.WriteNull(fieldName);
+						}
+						break;
 					}
-					break;
-				}
 				case JsonValueKind.True:
-				{
-					if (String.IsNullOrEmpty(fieldName))
 					{
-						writer.WriteBoolValue(true);
-					}
-					else
-					{
-						writer.WriteBool(fieldName, true);
+						if (String.IsNullOrEmpty(fieldName))
+						{
+							writer.WriteBoolValue(true);
+						}
+						else
+						{
+							writer.WriteBool(fieldName, true);
+						}
+						break;
 					}
-					break;
-				}
 				case JsonValueKind.False:
-				{
-					if (String.IsNullOrEmpty(fieldName))
 					{
-						writer.WriteBoolValue(false);
-					}
-					else
-					{
-						writer.WriteBool(fieldName, false);
+						if (String.IsNullOrEmpty(fieldName))
+						{
+							writer.WriteBoolValue(false);
+						}
+						else
+						{
+							writer.WriteBool(fieldName, false);
+						}
+						break;
 					}
-					break;
-				}
 				case JsonValueKind.Number:
-				{
-					if (String.IsNullOrEmpty(fieldName))
-					{
-						writer.WriteDoubleValue(jsonElement.GetDouble());
-					}
-					else
-					{
-						writer.WriteDouble(fieldName, jsonElement.GetDouble());
-					}
-					break;
-				}
-				case JsonValueKind.String:
-				{
-					string? value = jsonElement.GetString() ?? "";
-
-					if (CbObjectId.TryParse(value, out CbObjectId oid))
 					{
 						if (String.IsNullOrEmpty(fieldName))
 						{
-							writer.WriteObjectIdValue(oid);
+							writer.WriteDoubleValue(jsonElement.GetDouble());
 						}
 						else
 						{
-							writer.WriteObjectId(fieldName, oid);
+							writer.WriteDouble(fieldName, jsonElement.GetDouble());
 						}
-
-						return;
+						break;
 					}
-
+				case JsonValueKind.String:
 					{
-						if (IoHash.TryParse(value, out IoHash hash))
+						string? value = jsonElement.GetString() ?? "";
+
+						if (CbObjectId.TryParse(value, out CbObjectId oid))
 						{
 							if (String.IsNullOrEmpty(fieldName))
 							{
-								writer.WriteHashValue(hash);
+								writer.WriteObjectIdValue(oid);
 							}
 							else
 							{
-								writer.WriteHash(fieldName, hash);
+								writer.WriteObjectId(fieldName, oid);
 							}
 
 							return;
 						}
-					}
 
-					int pos = value.IndexOf(BinaryAttachmentPrefix, StringComparison.OrdinalIgnoreCase);
-					if  (pos != -1)
-					{
-						value = value.Substring(pos);
-						if (IoHash.TryParse(value, out IoHash hash))
 						{
-							if (String.IsNullOrEmpty(fieldName))
-							{
-								writer.WriteBinaryAttachmentValue(hash);
-							}
-							else
+							if (IoHash.TryParse(value, out IoHash hash))
 							{
-								writer.WriteBinaryAttachment(fieldName, hash);
+								if (String.IsNullOrEmpty(fieldName))
+								{
+									writer.WriteHashValue(hash);
+								}
+								else
+								{
+									writer.WriteHash(fieldName, hash);
+								}
+
+								return;
 							}
-
-							return;
 						}
-					}
 
-					pos = value.IndexOf(CompactBinaryAttachmentPrefix, StringComparison.OrdinalIgnoreCase);
-					if  (pos != -1)
-					{
-						value = value.Substring(pos);
-						if (IoHash.TryParse(value, out IoHash hash))
+						int pos = value.IndexOf(BinaryAttachmentPrefix, StringComparison.OrdinalIgnoreCase);
+						if (pos != -1)
 						{
-							if (String.IsNullOrEmpty(fieldName))
+							value = value.Substring(pos);
+							if (IoHash.TryParse(value, out IoHash hash))
 							{
-								writer.WriteObjectAttachmentValue(hash);
+								if (String.IsNullOrEmpty(fieldName))
+								{
+									writer.WriteBinaryAttachmentValue(hash);
+								}
+								else
+								{
+									writer.WriteBinaryAttachment(fieldName, hash);
+								}
+
+								return;
 							}
-							else
+						}
+
+						pos = value.IndexOf(CompactBinaryAttachmentPrefix, StringComparison.OrdinalIgnoreCase);
+						if (pos != -1)
+						{
+							value = value.Substring(pos);
+							if (IoHash.TryParse(value, out IoHash hash))
 							{
-								writer.WriteObjectAttachment(fieldName, hash);
+								if (String.IsNullOrEmpty(fieldName))
+								{
+									writer.WriteObjectAttachmentValue(hash);
+								}
+								else
+								{
+									writer.WriteObjectAttachment(fieldName, hash);
+								}
+
+								return;
 							}
-
-							return;
 						}
-					}
 
-					if (String.IsNullOrEmpty(fieldName))
-					{
-						writer.WriteStringValue(value);
-					}
-					else
-					{
-						writer.WriteString(fieldName, value);
+						if (String.IsNullOrEmpty(fieldName))
+						{
+							writer.WriteStringValue(value);
+						}
+						else
+						{
+							writer.WriteString(fieldName, value);
+						}
+						break;
 					}
-					break;
-				}
 				case JsonValueKind.Undefined:
 					break;
 				default:
diff --git a/Engine/Source/Programs/Shared/EpicGames.UHT/Parsers/UhtFunctionParser.cs b/Engine/Source/Programs/Shared/EpicGames.UHT/Parsers/UhtFunctionParser.cs
--- a/Engine/Source/Programs/Shared/EpicGames.UHT/Parsers/UhtFunctionParser.cs
+++ b/Engine/Source/Programs/Shared/EpicGames.UHT/Parsers/UhtFunctionParser.cs
@@ -120,7 +120,7 @@ namespace EpicGames.UHT.Parsers
 		[SuppressMessage("CodeQuality", "IDE0051:Remove unused private members", Justification = "Attribute accessed method")]
 		private static UhtParseResult ParseCatchAllKeyword(UhtParsingScope topScope, ref UhtToken token)
 		{
-			if (topScope.HeaderParser.NamespaceDepth == 1 && UhtFunctionParser.IsValidateDelegateDeclaration(token))
+			if (topScope.HeaderParser.NamespaceDepth == 1 && IsValidateDelegateDeclaration(token))
 			{
 				return ParseUDelegate(topScope, token, false);
 			}
diff --git a/Engine/Source/Programs/Shared/EpicGames.UHT/Tables/UhtKeywordTable.cs b/Engine/Source/Programs/Shared/EpicGames.UHT/Tables/UhtKeywordTable.cs
--- a/Engine/Source/Programs/Shared/EpicGames.UHT/Tables/UhtKeywordTable.cs
+++ b/Engine/Source/Programs/Shared/EpicGames.UHT/Tables/UhtKeywordTable.cs
@@ -166,7 +166,7 @@ namespace EpicGames.UHT.Tables
 		/// <param name="value">Value to be added</param>
 		public UhtKeywordTable Add(UhtKeyword value)
 		{
-			base.Add(value.Name, value);
+			Add(value.Name, value);
 			return this;
 		}
 
diff --git a/Engine/Source/Programs/Shared/EpicGames.UHT/Tables/UhtSpecifierTable.cs b/Engine/Source/Programs/Shared/EpicGames.UHT/Tables/UhtSpecifierTable.cs
--- a/Engine/Source/Programs/Shared/EpicGames.UHT/Tables/UhtSpecifierTable.cs
+++ b/Engine/Source/Programs/Shared/EpicGames.UHT/Tables/UhtSpecifierTable.cs
@@ -551,7 +551,7 @@ namespace EpicGames.UHT.Tables
 		/// <param name="specifier">The specifier to add</param>
 		public UhtSpecifierTable Add(UhtSpecifier specifier)
 		{
-			base.Add(specifier.Name, specifier);
+			Add(specifier.Name, specifier);
 			return this;
 		}
 	}
diff --git a/Engine/Source/Programs/Shared/EpicGames.UHT/Tables/UhtSpecifierValidatorTable.cs b/Engine/Source/Programs/Shared/EpicGames.UHT/Tables/UhtSpecifierValidatorTable.cs
--- a/Engine/Source/Programs/Shared/EpicGames.UHT/Tables/UhtSpecifierValidatorTable.cs
+++ b/Engine/Source/Programs/Shared/EpicGames.UHT/Tables/UhtSpecifierValidatorTable.cs
@@ -83,7 +83,7 @@ namespace EpicGames.UHT.Tables
 		/// <param name="specifier">Validator to add</param>
 		public UhtSpecifierValidatorTable Add(UhtSpecifierValidator specifier)
 		{
-			base.Add(specifier.Name, specifier);
+			Add(specifier.Name, specifier);
 			return this;
 		}
 	}
diff --git a/Engine/Source/Programs/Shared/EpicGames.UHT/Tokenizer/UhtTokenBufferReader.cs b/Engine/Source/Programs/Shared/EpicGames.UHT/Tokenizer/UhtTokenBufferReader.cs
--- a/Engine/Source/Programs/Shared/EpicGames.UHT/Tokenizer/UhtTokenBufferReader.cs
+++ b/Engine/Source/Programs/Shared/EpicGames.UHT/Tokenizer/UhtTokenBufferReader.cs
@@ -100,7 +100,7 @@ namespace EpicGames.UHT.Tokenizer
 				}
 				else
 				{
-					return new ReadOnlySpan<StringView>(UhtTokenBufferReader.s_emptyComments);
+					return new ReadOnlySpan<StringView>(s_emptyComments);
 				}
 			}
 		}
diff --git a/Engine/Source/Programs/Shared/EpicGames.UHT/Types/Properties/UhtDelegateProperty.cs b/Engine/Source/Programs/Shared/EpicGames.UHT/Types/Properties/UhtDelegateProperty.cs
--- a/Engine/Source/Programs/Shared/EpicGames.UHT/Types/Properties/UhtDelegateProperty.cs
+++ b/Engine/Source/Programs/Shared/EpicGames.UHT/Types/Properties/UhtDelegateProperty.cs
@@ -124,7 +124,7 @@ namespace EpicGames.UHT.Types
 		public override StringBuilder AppendMemberDef(StringBuilder builder, IUhtPropertyMemberContext context, string name, string nameSuffix, string? offset, int tabs)
 		{
 			AppendMemberDefStart(builder, context, name, nameSuffix, offset, tabs, "FDelegatePropertyParams", "UECodeGen_Private::EPropertyGenFlags::Delegate");
-			AppendMemberDefRef(builder, context, Function, Exporters.CodeGen.UhtSingletonType.Registered);
+			AppendMemberDefRef(builder, context, Function, UhtSingletonType.Registered);
 			AppendMemberDefEnd(builder, context, name, nameSuffix);
 			return builder;
 		}
diff --git a/Engine/Source/Programs/Shared/EpicGames.UHT/Types/Properties/UhtEnumProperty.cs b/Engine/Source/Programs/Shared/EpicGames.UHT/Types/Properties/UhtEnumProperty.cs
--- a/Engine/Source/Programs/Shared/EpicGames.UHT/Types/Properties/UhtEnumProperty.cs
+++ b/Engine/Source/Programs/Shared/EpicGames.UHT/Types/Properties/UhtEnumProperty.cs
@@ -130,8 +130,8 @@ namespace EpicGames.UHT.Types
 				case UhtPropertyTextType.ClassFunctionArgOrRetVal:
 				case UhtPropertyTextType.InterfaceFunctionArgOrRetVal:
 				case UhtPropertyTextType.EventFunctionArgOrRetVal:
-					if (enumObj.CppForm == UhtEnumCppForm.EnumClass || 
-						(!isTemplateArgument && (property.PropertyFlags.HasAnyFlags(EPropertyFlags.ReturnParm) || 
+					if (enumObj.CppForm == UhtEnumCppForm.EnumClass ||
+						(!isTemplateArgument && (property.PropertyFlags.HasAnyFlags(EPropertyFlags.ReturnParm) ||
 						!property.PropertyFlags.HasAnyFlags(EPropertyFlags.OutParm))))
 					{
 						builder.Append(enumObj.FullyQualifiedCppType);
@@ -273,14 +273,14 @@ namespace EpicGames.UHT.Types
 					builder.AppendMemberDef(UnderlyingProperty, context, name, GetNameSuffix(nameSuffix, "_Underlying"), "0", tabs);
 				}
 				AppendMemberDefStart(builder, context, name, nameSuffix, offset, tabs, "FEnumPropertyParams", "UECodeGen_Private::EPropertyGenFlags::Enum");
-				AppendMemberDefRef(builder, context, Enum, Exporters.CodeGen.UhtSingletonType.Registered);
+				AppendMemberDefRef(builder, context, Enum, UhtSingletonType.Registered);
 				AppendMemberDefEnd(builder, context, name, nameSuffix);
 				return builder;
 			}
 			else
 			{
 				AppendMemberDefStart(builder, context, name, nameSuffix, offset, tabs, "FBytePropertyParams", "UECodeGen_Private::EPropertyGenFlags::Byte");
-				AppendMemberDefRef(builder, context, Enum, Exporters.CodeGen.UhtSingletonType.Registered);
+				AppendMemberDefRef(builder, context, Enum, UhtSingletonType.Registered);
 				AppendMemberDefEnd(builder, context, name, nameSuffix);
 			}
 			return builder;
@@ -293,7 +293,8 @@ namespace EpicGames.UHT.Types
 			{
 				if (UnderlyingProperty != null)
 				{
-					builder.AppendConstInitMemberDef(UnderlyingProperty, context, name, GetNameSuffix(nameSuffix, "_Underlying"), (builder) => {
+					builder.AppendConstInitMemberDef(UnderlyingProperty, context, name, GetNameSuffix(nameSuffix, "_Underlying"), (builder) =>
+					{
 						builder.AppendConstInitMemberPtr(this, context, name, nameSuffix, tabs, "");
 					}, "0", tabs);
 				}
diff --git a/Engine/Source/Programs/Shared/EpicGames.UHT/Types/Properties/UhtInterfaceProperty.cs b/Engine/Source/Programs/Shared/EpicGames.UHT/Types/Properties/UhtInterfaceProperty.cs
--- a/Engine/Source/Programs/Shared/EpicGames.UHT/Types/Properties/UhtInterfaceProperty.cs
+++ b/Engine/Source/Programs/Shared/EpicGames.UHT/Types/Properties/UhtInterfaceProperty.cs
@@ -155,7 +155,7 @@ namespace EpicGames.UHT.Types
 		{
 			// FScriptInterface<USomeInterface> is valid so in that case we need to pass in the interface class and not the alternate object (which in the end is the same object)
 			AppendMemberDefStart(builder, context, name, nameSuffix, offset, tabs, "FInterfacePropertyParams", "UECodeGen_Private::EPropertyGenFlags::Interface");
-			AppendMemberDefRef(builder, context, InterfaceClass.AlternateObject ?? InterfaceClass, Exporters.CodeGen.UhtSingletonType.Unregistered);
+			AppendMemberDefRef(builder, context, InterfaceClass.AlternateObject ?? InterfaceClass, UhtSingletonType.Unregistered);
 			AppendMemberDefEnd(builder, context, name, nameSuffix);
 			return builder;
 		}
diff --git a/Engine/Source/Programs/Shared/EpicGames.UHT/Types/Properties/UhtSoftClassProperty.cs b/Engine/Source/Programs/Shared/EpicGames.UHT/Types/Properties/UhtSoftClassProperty.cs
--- a/Engine/Source/Programs/Shared/EpicGames.UHT/Types/Properties/UhtSoftClassProperty.cs
+++ b/Engine/Source/Programs/Shared/EpicGames.UHT/Types/Properties/UhtSoftClassProperty.cs
@@ -85,7 +85,7 @@ namespace EpicGames.UHT.Types
 		public override StringBuilder AppendMemberDef(StringBuilder builder, IUhtPropertyMemberContext context, string name, string nameSuffix, string? offset, int tabs)
 		{
 			AppendMemberDefStart(builder, context, name, nameSuffix, offset, tabs, "FSoftClassPropertyParams", "UECodeGen_Private::EPropertyGenFlags::SoftClass");
-			AppendMemberDefRef(builder, context, MetaClass, Exporters.CodeGen.UhtSingletonType.Unregistered);
+			AppendMemberDefRef(builder, context, MetaClass, UhtSingletonType.Unregistered);
 			AppendMemberDefEnd(builder, context, name, nameSuffix);
 			return builder;
 		}
diff --git a/Engine/Source/Programs/Shared/EpicGames.UHT/Types/UhtHeaderFile.cs b/Engine/Source/Programs/Shared/EpicGames.UHT/Types/UhtHeaderFile.cs
--- a/Engine/Source/Programs/Shared/EpicGames.UHT/Types/UhtHeaderFile.cs
+++ b/Engine/Source/Programs/Shared/EpicGames.UHT/Types/UhtHeaderFile.cs
@@ -107,7 +107,7 @@ namespace EpicGames.UHT.Types
 		/// <summary>
 		/// The header being referenced in a direct include
 		/// </summary>
-		Include,	
+		Include,
 
 		/// <summary>
 		/// The header being referenced is a passive reference (i.e. NoExportTypes.h or a header referenced by a header being included)
@@ -227,7 +227,7 @@ namespace EpicGames.UHT.Types
 		/// Resource collector for the header file
 		/// </summary>
 		[JsonIgnore]
-		public UhtReferenceCollector References { get; init; } 
+		public UhtReferenceCollector References { get; init; }
 
 		/// <summary>
 		/// Collection of headers directly included by this header
@@ -267,7 +267,7 @@ namespace EpicGames.UHT.Types
 			_messageSite = new UhtSimpleMessageSite(Session);
 			_sourceFile = new UhtSourceFile(Session, path);
 			_messageSite.MessageSource = _sourceFile;
-			FileNameWithoutExtension = System.IO.Path.GetFileNameWithoutExtension(_sourceFile.FilePath);
+			FileNameWithoutExtension = Path.GetFileNameWithoutExtension(_sourceFile.FilePath);
 			GeneratedHeaderFileName = FileNameWithoutExtension + ".generated.h";
 			IsNoExportTypes = String.Equals(_sourceFile.FileName, "NoExportTypes", StringComparison.OrdinalIgnoreCase);
 		}
diff --git a/Engine/Source/Programs/Shared/EpicGames.UHT/Types/UhtType.cs b/Engine/Source/Programs/Shared/EpicGames.UHT/Types/UhtType.cs
--- a/Engine/Source/Programs/Shared/EpicGames.UHT/Types/UhtType.cs
+++ b/Engine/Source/Programs/Shared/EpicGames.UHT/Types/UhtType.cs
@@ -951,7 +951,7 @@ namespace EpicGames.UHT.Types
 		/// </summary>
 		[JsonConverter(typeof(UhtNullableTypeListJsonConverter<UhtType>))]
 		//TODO - this should return a IReadOnlyList
-		public List<UhtType> Children => _children ?? UhtType.s_emptyTypeList;
+		public List<UhtType> Children => _children ?? s_emptyTypeList;
 
 		/// <summary>
 		/// Helper to string method
diff --git a/Engine/Source/Programs/Shared/EpicGames.UHT/Utils/UhtBuffer.cs b/Engine/Source/Programs/Shared/EpicGames.UHT/Utils/UhtBuffer.cs
--- a/Engine/Source/Programs/Shared/EpicGames.UHT/Utils/UhtBuffer.cs
+++ b/Engine/Source/Programs/Shared/EpicGames.UHT/Utils/UhtBuffer.cs
@@ -167,17 +167,17 @@ namespace EpicGames.UHT.Utils
 		/// <summary>
 		/// Adjustment to the bucket index to account for the minimum bucket size
 		/// </summary>
-		private static readonly int s_buckedAdjustment = BitOperations.Log2((uint)UhtBuffer.MinSize);
+		private static readonly int s_buckedAdjustment = BitOperations.Log2((uint)MinSize);
 
 		/// <summary>
 		/// Total number of supported buckets
 		/// </summary>
-		private static readonly int s_bucketCount = 32 - UhtBuffer.s_buckedAdjustment;
+		private static readonly int s_bucketCount = 32 - s_buckedAdjustment;
 
 		/// <summary>
 		/// Bucket lookaside list
 		/// </summary>
-		private static readonly UhtBuffer?[] s_lookAsideArray = new UhtBuffer?[UhtBuffer.s_bucketCount];
+		private static readonly UhtBuffer?[] s_lookAsideArray = new UhtBuffer?[s_bucketCount];
 
 		/// <summary>
 		/// The bucket index associated with the buffer
@@ -231,9 +231,9 @@ namespace EpicGames.UHT.Utils
 		/// <returns>Buffer that should be returned with a call to Return</returns>
 		public static UhtBuffer Borrow(int size)
 		{
-			if (size <= UhtBuffer.MinSize)
+			if (size <= MinSize)
 			{
-				return BorrowInternal(size, 0, UhtBuffer.MinSize);
+				return BorrowInternal(size, 0, MinSize);
 			}
 			else
 			{
@@ -247,7 +247,7 @@ namespace EpicGames.UHT.Utils
 				usize |= usize >> 8;
 				usize |= usize >> 16;
 				++usize;
-				int bucket = BitOperations.Log2(usize) - UhtBuffer.s_buckedAdjustment;
+				int bucket = BitOperations.Log2(usize) - s_buckedAdjustment;
 				return BorrowInternal(size, bucket, (int)usize);
 			}
 		}
@@ -285,10 +285,10 @@ namespace EpicGames.UHT.Utils
 		/// <param name="buffer">The buffer to be returned.</param>
 		public static void Return(UhtBuffer buffer)
 		{
-			lock (UhtBuffer.s_lookAsideArray)
+			lock (s_lookAsideArray)
 			{
-				buffer.NextBuffer = UhtBuffer.s_lookAsideArray[buffer.Bucket];
-				UhtBuffer.s_lookAsideArray[buffer.Bucket] = buffer;
+				buffer.NextBuffer = s_lookAsideArray[buffer.Bucket];
+				s_lookAsideArray[buffer.Bucket] = buffer;
 			}
 		}
 
@@ -301,12 +301,12 @@ namespace EpicGames.UHT.Utils
 		/// <returns>The allocated buffer</returns>
 		private static UhtBuffer BorrowInternal(int size, int bucket, int bucketSize)
 		{
-			lock (UhtBuffer.s_lookAsideArray)
+			lock (s_lookAsideArray)
 			{
-				if (UhtBuffer.s_lookAsideArray[bucket] != null)
+				if (s_lookAsideArray[bucket] != null)
 				{
-					UhtBuffer buffer = UhtBuffer.s_lookAsideArray[bucket]!;
-					UhtBuffer.s_lookAsideArray[bucket] = buffer.NextBuffer;
+					UhtBuffer buffer = s_lookAsideArray[bucket]!;
+					s_lookAsideArray[bucket] = buffer.NextBuffer;
 					buffer.Reset(size);
 					return buffer;
 				}
@@ -382,17 +382,17 @@ namespace EpicGames.UHT.Utils
 		/// <summary>
 		/// Adjustment to the bucket index to account for the minimum bucket size
 		/// </summary>
-		private static readonly int s_buckedAdjustment = BitOperations.Log2((uint)UhtByteBuffer.MinSize);
+		private static readonly int s_buckedAdjustment = BitOperations.Log2((uint)MinSize);
 
 		/// <summary>
 		/// Total number of supported buckets
 		/// </summary>
-		private static readonly int s_bucketCount = 32 - UhtByteBuffer.s_buckedAdjustment;
+		private static readonly int s_bucketCount = 32 - s_buckedAdjustment;
 
 		/// <summary>
 		/// Bucket lookaside list
 		/// </summary>
-		private static readonly UhtByteBuffer?[] s_lookAsideArray = new UhtByteBuffer?[UhtByteBuffer.s_bucketCount];
+		private static readonly UhtByteBuffer?[] s_lookAsideArray = new UhtByteBuffer?[s_bucketCount];
 
 		/// <summary>
 		/// The bucket index associated with the buffer
@@ -446,9 +446,9 @@ namespace EpicGames.UHT.Utils
 		/// <returns>Buffer that should be returned with a call to Return</returns>
 		public static UhtByteBuffer Borrow(int size)
 		{
-			if (size <= UhtByteBuffer.MinSize)
+			if (size <= MinSize)
 			{
-				return BorrowInternal(size, 0, UhtByteBuffer.MinSize);
+				return BorrowInternal(size, 0, MinSize);
 			}
 			else
 			{
@@ -462,7 +462,7 @@ namespace EpicGames.UHT.Utils
 				usize |= usize >> 8;
 				usize |= usize >> 16;
 				++usize;
-				int bucket = BitOperations.Log2(usize) - UhtByteBuffer.s_buckedAdjustment;
+				int bucket = BitOperations.Log2(usize) - s_buckedAdjustment;
 				return BorrowInternal(size, bucket, (int)usize);
 			}
 		}
@@ -473,10 +473,10 @@ namespace EpicGames.UHT.Utils
 		/// <param name="buffer">The buffer to be returned.</param>
 		public static void Return(UhtByteBuffer buffer)
 		{
-			lock (UhtByteBuffer.s_lookAsideArray)
+			lock (s_lookAsideArray)
 			{
-				buffer.NextBuffer = UhtByteBuffer.s_lookAsideArray[buffer.Bucket];
-				UhtByteBuffer.s_lookAsideArray[buffer.Bucket] = buffer;
+				buffer.NextBuffer = s_lookAsideArray[buffer.Bucket];
+				s_lookAsideArray[buffer.Bucket] = buffer;
 			}
 		}
 
@@ -489,12 +489,12 @@ namespace EpicGames.UHT.Utils
 		/// <returns>The allocated buffer</returns>
 		private static UhtByteBuffer BorrowInternal(int size, int bucket, int bucketSize)
 		{
-			lock (UhtByteBuffer.s_lookAsideArray)
+			lock (s_lookAsideArray)
 			{
-				if (UhtByteBuffer.s_lookAsideArray[bucket] != null)
+				if (s_lookAsideArray[bucket] != null)
 				{
-					UhtByteBuffer buffer = UhtByteBuffer.s_lookAsideArray[bucket]!;
-					UhtByteBuffer.s_lookAsideArray[bucket] = buffer.NextBuffer;
+					UhtByteBuffer buffer = s_lookAsideArray[bucket]!;
+					s_lookAsideArray[bucket] = buffer.NextBuffer;
 					buffer.Reset(size);
 					return buffer;
 				}
diff --git a/Engine/Source/Programs/Shared/EpicGames.UHT/Utils/UhtMessage.cs b/Engine/Source/Programs/Shared/EpicGames.UHT/Utils/UhtMessage.cs
--- a/Engine/Source/Programs/Shared/EpicGames.UHT/Utils/UhtMessage.cs
+++ b/Engine/Source/Programs/Shared/EpicGames.UHT/Utils/UhtMessage.cs
@@ -264,7 +264,7 @@ namespace EpicGames.UHT.Utils
 						{
 							builder.Append(" in ");
 						}
-						builder.Append(UhtMessage.FormatContext(ec));
+						builder.Append(FormatContext(ec));
 					}
 				}
 			}
@@ -543,13 +543,13 @@ namespace EpicGames.UHT.Utils
 		/// Get the extra context associated with this thread
 		/// </summary>
 		/// <returns>Extra context</returns>
-		public static UhtTlsMessageExtraContext? GetTls() { return UhtTlsMessageExtraContext.s_tls.Value; }
+		public static UhtTlsMessageExtraContext? GetTls() { return s_tls.Value; }
 
 		/// <summary>
 		/// Get the extra context interface
 		/// </summary>
 		/// <returns>Extra context interface</returns>
-		public static IUhtMessageExtraContext? GetMessageExtraContext() { return UhtTlsMessageExtraContext.s_tls.Value; }
+		public static IUhtMessageExtraContext? GetMessageExtraContext() { return s_tls.Value; }
 	}
 
 	/// <summary>
diff --git a/Engine/Source/Programs/Shared/EpicGames.UHT/Utils/UhtSession.cs b/Engine/Source/Programs/Shared/EpicGames.UHT/Utils/UhtSession.cs
--- a/Engine/Source/Programs/Shared/EpicGames.UHT/Utils/UhtSession.cs
+++ b/Engine/Source/Programs/Shared/EpicGames.UHT/Utils/UhtSession.cs
@@ -382,7 +382,7 @@ namespace EpicGames.UHT.Utils
 				// Writing billions of files to the same directory causes issues.  Use ourselves to throttle reference writes
 				try
 				{
-					UhtExportFactory.s_writeRefSemaphore.WaitOne();
+					s_writeRefSemaphore.WaitOne();
 					{
 						string outPath = Path.Combine(Session.ReferenceMode == UhtReferenceMode.Reference ? ReferenceDirectory : VerifyDirectory, fileName);
 						if (!Session.WriteSource(outPath, exported.Span))
@@ -393,7 +393,7 @@ namespace EpicGames.UHT.Utils
 				}
 				finally
 				{
-					UhtExportFactory.s_writeRefSemaphore.Release();
+					s_writeRefSemaphore.Release();
 				}
 
 				// If we are verifying, read the existing file and check the contents
@@ -1844,11 +1844,11 @@ namespace EpicGames.UHT.Utils
 			}
 
 			GetMessageParts(msg, out string filePath, out int lineNumber, out string fragmentPath, out string severity, out string message);
-			Logger.Log(logLevel, KnownLogEvents.UHT, "{File}({Line}){FileFragment}: {Severity}: {Message}", 
-				new FileReference(filePath), 
-				new LogValue(LogValueType.LineNumber, lineNumber.ToString()), 
-				fragmentPath, 
-				new LogValue(LogValueType.Severity, severity), 
+			Logger.Log(logLevel, KnownLogEvents.UHT, "{File}({Line}){FileFragment}: {Severity}: {Message}",
+				new FileReference(filePath),
+				new LogValue(LogValueType.LineNumber, lineNumber.ToString()),
+				fragmentPath,
+				new LogValue(LogValueType.Severity, severity),
 				new LogValue(LogValueType.Message, message));
 		}
 
@@ -2062,18 +2062,18 @@ namespace EpicGames.UHT.Utils
 
 				if (Tables != null)
 				{
-				if (Manifest != null)
-				{
+					if (Manifest != null)
+					{
 						Tables.AddPlugins(Manifest.UhtPlugins);
 					}
 					if (EnableInputCacheRead || EnableInputCacheWrite)
 					{
 						Tables.ValidateCacheableTypes(Logger);
 					}
-					}
+				}
 
 				if (Manifest != null)
-					{
+				{
 					AssemblyDateTime = UhtUtilities.GetLatestAssemblyDateTime(Manifest.UhtPlugins);
 				}
 				Logger.LogInformation("UHT compiled-in object format {CompiledInObjectFormat}", CompiledInObjectFormat);
@@ -2140,7 +2140,7 @@ namespace EpicGames.UHT.Utils
 				headerFile.ReadFromCache();
 				if (!headerFile.WasCached)
 				{
-				headerFile.Read();
+					headerFile.Read();
 					UhtHeaderFileParser.Parse(headerFile);
 
 					if (EnableInputCacheWrite && !headerFile.WasCached)
@@ -2257,7 +2257,7 @@ namespace EpicGames.UHT.Utils
 		private void StepLinkModules()
 		{
 			Log.Logger.LogTrace("Step - Link modules");
-		
+
 			foreach (UhtModule module in Modules)
 			{
 				Dictionary<UhtModule, int> referencedModules = module.Headers.SelectMany(h => h.References.Modules)
@@ -3123,7 +3123,7 @@ namespace EpicGames.UHT.Utils
 				{
 					using StreamWriter output = new(Manifest!.ExternalDependenciesFile);
 					foreach (string dep in externalDependencies)
-					{ 
+					{
 						output.WriteLine(dep);
 					}
 				}
