Change to using MSBuild-based project parsing rather than custom project parsing. Make the results from MSBuild-based parsing evaluate to as close to what MSBuild would evaluate to as possible, or fail when this is not possible.

#Synced-CL 49202056

diff --git a/Engine/Source/Programs/AutomationTool/BuildGraph/Tasks/CsCompileTask.cs b/Engine/Source/Programs/AutomationTool/BuildGraph/Tasks/CsCompileTask.cs
--- a/Engine/Source/Programs/AutomationTool/BuildGraph/Tasks/CsCompileTask.cs
+++ b/Engine/Source/Programs/AutomationTool/BuildGraph/Tasks/CsCompileTask.cs
@@ -301,8 +301,8 @@ namespace AutomationTool.Tasks
 			if (!fileToProjectInfo.ContainsKey(file))
 			{
 				// Try to read this project
-				CsProjectInfo projectInfo;
-				if (!CsProjectInfo.TryRead(file, initialProperties, out projectInfo))
+				DotnetProjectInfo projectInfo;
+				if (!DotnetProjectInfo.TryRead(file, initialProperties, out projectInfo))
 				{
 					throw new AutomationException("Couldn't read project '{0}'", file.FullName);
 				}
diff --git a/Engine/Source/Programs/Shared/EpicGames.Build/DotnetProjectInfo.cs b/Engine/Source/Programs/Shared/EpicGames.Build/DotnetProjectInfo.cs
--- a/Engine/Source/Programs/Shared/EpicGames.Build/DotnetProjectInfo.cs
+++ b/Engine/Source/Programs/Shared/EpicGames.Build/DotnetProjectInfo.cs
@@ -78,16 +78,17 @@ public sealed record class DotnetProjectInfo : IDotnetProjectInfo
 	/// Reads project information for the given file.
 	/// </summary>
 	/// <param name="file">The project file to read</param>
-	/// <param name="properties">Initial set of property values</param>
+	/// <param name="globalProperties">Initial set of property values</param>
 	/// <returns>The parsed project info</returns>
-	public static DotnetProjectInfo Read(FileReference file, IReadOnlyDictionary<string, string> properties)
+	[MethodImpl(MethodImplOptions.NoInlining)]
+	public static DotnetProjectInfo Read(FileReference file, IReadOnlyDictionary<string, string> globalProperties)
 	{
-		DotnetProjectInfo? project;
-		if (!TryRead(file, properties, out project))
-		{
-			throw new Exception(String.Format("Unable to read '{0}'", file));
-		}
-		return project;
+		// First register MSBuild so we can use it to parse our project file. By not having the internal routine
+		// inline, we avoid having the issue of the Microsoft.Build libraries being resolved prior to the build path
+		// being set.
+		MsBuildSupport.RegisterMsBuildPath(MsBuildRegistration.Instance);
+
+		return ReadImpl(file, globalProperties);
 	}
 
 	/// <summary>
@@ -102,77 +103,97 @@ public sealed record class DotnetProjectInfo : IDotnetProjectInfo
 		IReadOnlyDictionary<string, string> globalProperties,
 		[NotNullWhen(true)] out DotnetProjectInfo? outProjectInfo)
 	{
-		// First register MSBuild so we can use it to parse our project file. By not having the internal routine
-		// inline, we avoid having the issue of the Microsoft.Build libraries being resolved prior to the build path
-		// being set.
-		MsBuildSupport.RegisterMsBuildPath(MsBuildRegistration.Instance);
-
-		return TryReadImpl(file, globalProperties, out outProjectInfo);
+		try
+		{
+			outProjectInfo = Read(file, globalProperties);
+			return true;
+		}
+		catch
+		{
+			outProjectInfo = null;
+			return false;
+		}
 	}
 
 	[MethodImpl(MethodImplOptions.NoInlining)]
-	private static bool TryReadImpl(
-		FileReference file,
-		IReadOnlyDictionary<string, string> globalProperties,
-		[NotNullWhen(true)] out DotnetProjectInfo? outProjectInfo)
+	private static DotnetProjectInfo ReadImpl(FileReference file, IReadOnlyDictionary<string, string> globalProperties)
 	{
-		try
+		using ProjectCollection projectCollection = new();
+		ProjectOptions projectOptions = new()
 		{
-			using ProjectCollection projectCollection = new();
-			ProjectOptions projectOptions = new()
-			{
-				GlobalProperties = globalProperties.ToDictionary(),
-				LoadSettings = ProjectLoadSettings.IgnoreMissingImports | ProjectLoadSettings.IgnoreInvalidImports,
-				ProjectCollection = projectCollection,
-			};
-			Project project = Project.FromFile(file.FullName, projectOptions);
-
-			Dictionary<string, string> properties = project.Properties.ToDictionary(p => p.Name, p => p.EvaluatedValue);
-
-			DirectoryReference projectDirectory = new(project.DirectoryPath);
-			FileReference GetFullPath(string path) =>
-				Path.IsPathFullyQualified(path) ? new FileReference(path) : FileReference.Combine(projectDirectory, path);
-
-			bool IsPrivate(ProjectItem item) => !item.HasMetadata("Private") || Boolean.Parse(item.GetMetadataValue("Private"));
-
-			IEnumerable<ProjectItem> validReferences = project.GetItems("Reference")
-				.Where(r => r.HasMetadata("HintPath"));
-			Dictionary<FileReference, bool> references = validReferences
-				.ToDictionary(r => GetFullPath(r.GetMetadataValue("HintPath")), IsPrivate);
-
-			Dictionary<FileReference, bool> projectReferences = project.GetItems("ProjectReference")
-				.ToDictionary(pr => GetFullPath(pr.EvaluatedInclude), IsPrivate);
-
-			List<FileReference> compileReferences = [.. project.GetItems("Compile").Select(c => GetFullPath(c.EvaluatedInclude))];
-
-			bool ShouldCopy(ProjectItem item)
-			{
-				string copyTo = item.GetMetadataValue("CopyToOutputDirectory");
-				return copyTo.Equals("Always", StringComparison.OrdinalIgnoreCase) || copyTo.Equals("PreserveNewest", StringComparison.OrdinalIgnoreCase);
-			}
-
-			IEnumerable<ProjectItem> contentAndNoneRefs = project.GetItems("Content").Concat(project.GetItems("None"));
-			Dictionary<FileReference, bool> contentReferences = contentAndNoneRefs
-				.ToDictionary(cr => GetFullPath(cr.EvaluatedInclude), ShouldCopy);
-
-			outProjectInfo = new()
-			{
-				Properties = properties,
-				References = references,
-				ProjectReferences = projectReferences,
-				CompileReferences = compileReferences,
-				ContentReferences = contentReferences,
-				ProjectPath = file,
-			};
-			return true;
+			GlobalProperties = globalProperties.ToDictionary(),
+			LoadSettings = ProjectLoadSettings.IgnoreMissingImports | ProjectLoadSettings.IgnoreInvalidImports,
+			ProjectCollection = projectCollection,
+		};
+		Project project = Project.FromFile(file.FullName, projectOptions);
+
+		Dictionary<string, string> properties = project.Properties.ToDictionary(p => p.Name, p => p.EvaluatedValue);
+
+		DirectoryReference projectDirectory = new(project.DirectoryPath);
+		FileReference GetFullPath(string path) =>
+			Path.IsPathFullyQualified(path) ? new FileReference(path) : FileReference.Combine(projectDirectory, path);
+
+		bool IsPrivate(ProjectItem item) => !item.HasMetadata("Private") || Boolean.Parse(item.GetMetadataValue("Private"));
+
+		// For references with hint paths, multiple values appear to be all loaded, checked that they're the correct assembly,
+		// and the highest versioned assembly is used. It would be unreasonable to recreate this behaviour here, and it's not
+		// clear how we would get MSBuild to evaluate this for us, so we should just warn if people do silly things like this.
+		IEnumerable<string> unparseableReferences = project.GetItems("Reference")
+			.Where(r => !r.HasMetadata("EmbedInteropTypes") || !Boolean.Parse(r.GetMetadataValue("EmbedInteropTypes")))
+			.Where(r => r.HasMetadata("HintPath"))
+			.GroupBy(r => r.EvaluatedInclude)
+			.Where(grp => grp.DistinctBy(r => (hintPath: r.GetMetadataValue("HintPath"), isPrivate: IsPrivate(r))).Count() > 1)
+			.Select(grp => grp.Key);
+		string unparseableReferencesString = string.Join(", ", unparseableReferences);
+		if (unparseableReferencesString.Length != 0)
+		{
+			throw new Exception($"Project '{file}' has multiple <Reference>s for the same Include but with different HintPath / Private values ({unparseableReferencesString}) - please fix.");
 		}
-		catch (Exception)
+
+		IEnumerable<ProjectItem> validReferences = project.GetItems("Reference")
+			.Where(r => !r.HasMetadata("EmbedInteropTypes") || !Boolean.Parse(r.GetMetadataValue("EmbedInteropTypes")))
+			.Where(r => r.HasMetadata("HintPath"))
+			.Reverse()
+			.DistinctBy(r => r.EvaluatedInclude);
+		Dictionary<FileReference, bool> references = validReferences
+			.ToDictionary(r => GetFullPath(r.GetMetadataValue("HintPath")), IsPrivate);
+
+		// For project references, we deal with duplicate elements by taking the last value (hence the .Reverse call)
+		// as this follows MSBuild's actual behaviour as documented.
+		// https://learn.microsoft.com/en-us/visualstudio/msbuild/comparing-properties-and-items?view=visualstudio#property-and-item-evaluation-order
+		IEnumerable<ProjectItem> dedupedProjectReferences = project.GetItems("ProjectReference")
+			.Reverse()
+			.DistinctBy(pr => pr.EvaluatedInclude);
+		Dictionary<FileReference, bool> projectReferences = dedupedProjectReferences
+			.ToDictionary(pr => GetFullPath(pr.EvaluatedInclude), IsPrivate);
+
+		List<FileReference> compileReferences = [.. project.GetItems("Compile").Select(c => GetFullPath(c.EvaluatedInclude)).Distinct()];
+
+		bool ShouldCopy(ProjectItem item)
 		{
-			// Do nothing.
+			string copyTo = item.GetMetadataValue("CopyToOutputDirectory");
+			return copyTo.Equals("Always", StringComparison.OrdinalIgnoreCase) || copyTo.Equals("PreserveNewest", StringComparison.OrdinalIgnoreCase);
 		}
 
-		outProjectInfo = null;
-		return false;
+		// For Content and None items, no matter the ordering, if any reference has CopyToOutputDirectory set to Always or PreserveNewest,
+		// then they will be copied, even if later items explicitly set it to Never.
+		// This contradicts the documented MSBuild behaviour, which is very cool.
+		IEnumerable<(ProjectItem item, bool shouldCopy)> contentAndNoneRefs = project.GetItems("Content")
+			.Concat(project.GetItems("None"))
+			.GroupBy(cr => cr.EvaluatedInclude)
+			.Select(grp => (item: grp.Last(), shouldCopy: grp.Any(ShouldCopy)));
+		Dictionary<FileReference, bool> contentReferences = contentAndNoneRefs
+			.ToDictionary(tup => GetFullPath(tup.item.EvaluatedInclude), tup => tup.shouldCopy);
+
+		return new()
+		{
+			Properties = properties,
+			References = references,
+			ProjectReferences = projectReferences,
+			CompileReferences = compileReferences,
+			ContentReferences = contentReferences,
+			ProjectPath = file,
+		};
 	}
 
 	private class MsBuildRegistration : IMsBuildRegistration
diff --git a/Engine/Source/Programs/Shared/EpicGames.MsBuild/MsBuildSupport.cs b/Engine/Source/Programs/Shared/EpicGames.MsBuild/MsBuildSupport.cs
--- a/Engine/Source/Programs/Shared/EpicGames.MsBuild/MsBuildSupport.cs
+++ b/Engine/Source/Programs/Shared/EpicGames.MsBuild/MsBuildSupport.cs
@@ -13,7 +13,8 @@ namespace EpicGames.MsBuild;
 /// </summary>
 public static class MsBuildSupport
 {
-	static bool s_hasRegisteredMsBuildPath = false;
+	private static bool s_hasRegisteredMsBuildPath = false;
+	private static readonly object s_lock = new object();
 
 	/// <summary>
 	/// Register our bundled dotnet installation to be used by Microsoft.Build
@@ -25,44 +26,53 @@ public static class MsBuildSupport
 		{
 			return;
 		}
-		s_hasRegisteredMsBuildPath = true;
 
-		// Find our bundled dotnet SDK
-		List<string> listOfSdks = [];
-		ProcessStartInfo startInfo = new ProcessStartInfo
+		lock (s_lock)
 		{
-			FileName = hook.DotnetPath.FullName,
-			RedirectStandardOutput = true,
-			UseShellExecute = false,
-			ArgumentList = { "--list-sdks" }
-		};
-		startInfo.EnvironmentVariables["DOTNET_MULTILEVEL_LOOKUP"] = "0"; // use only the bundled dotnet installation - ignore any other/system dotnet install
+			if (s_hasRegisteredMsBuildPath)
+			{
+				return;
+			}
 
-		Process dotnetProcess = Process.Start(startInfo);
-		{
-			string line;
-			while ((line = dotnetProcess.StandardOutput.ReadLine()) != null)
+			// Find our bundled dotnet SDK
+			List<string> listOfSdks = [];
+			ProcessStartInfo startInfo = new ProcessStartInfo
+			{
+				FileName = hook.DotnetPath.FullName,
+				RedirectStandardOutput = true,
+				UseShellExecute = false,
+				ArgumentList = { "--list-sdks" }
+			};
+			startInfo.EnvironmentVariables["DOTNET_MULTILEVEL_LOOKUP"] = "0"; // use only the bundled dotnet installation - ignore any other/system dotnet install
+
+			Process dotnetProcess = Process.Start(startInfo);
 			{
-				listOfSdks.Add(line);
+				string line;
+				while ((line = dotnetProcess.StandardOutput.ReadLine()) != null)
+				{
+					listOfSdks.Add(line);
+				}
 			}
-		}
-		dotnetProcess.WaitForExit();
+			dotnetProcess.WaitForExit();
 
-		if (listOfSdks.Count != 1)
-		{
-			throw new Exception("Expected only one sdk installed for bundled dotnet");
-		}
+			if (listOfSdks.Count != 1)
+			{
+				throw new Exception("Expected only one sdk installed for bundled dotnet");
+			}
 
-		// Expected output has this form:
-		// 3.1.403 [D:\UE5_Main\engine\binaries\ThirdParty\DotNet\Windows\sdk]
-		string sdkVersion = listOfSdks[0].Split(' ')[0];
+			// Expected output has this form:
+			// 3.1.403 [D:\UE5_Main\engine\binaries\ThirdParty\DotNet\Windows\sdk]
+			string sdkVersion = listOfSdks[0].Split(' ')[0];
 
-		DirectoryReference dotnetSdkDirectory = DirectoryReference.Combine(hook.DotnetDirectory, "sdk", sdkVersion);
-		if (!DirectoryReference.Exists(dotnetSdkDirectory))
-		{
-			throw new Exception($"Failed to find .NET SDK directory: {dotnetSdkDirectory.FullName}");
-		}
+			DirectoryReference dotnetSdkDirectory = DirectoryReference.Combine(hook.DotnetDirectory, "sdk", sdkVersion);
+			if (!DirectoryReference.Exists(dotnetSdkDirectory))
+			{
+				throw new Exception($"Failed to find .NET SDK directory: {dotnetSdkDirectory.FullName}");
+			}
 
-		MSBuildLocator.RegisterMSBuildPath(dotnetSdkDirectory.FullName);
+			MSBuildLocator.RegisterMSBuildPath(dotnetSdkDirectory.FullName);
+
+			s_hasRegisteredMsBuildPath = true;
+		}
 	}
 }
diff --git a/Engine/Source/Programs/UnrealBuildTool/ProjectFiles/VisualStudio/VCProject.cs b/Engine/Source/Programs/UnrealBuildTool/ProjectFiles/VisualStudio/VCProject.cs
--- a/Engine/Source/Programs/UnrealBuildTool/ProjectFiles/VisualStudio/VCProject.cs
+++ b/Engine/Source/Programs/UnrealBuildTool/ProjectFiles/VisualStudio/VCProject.cs
@@ -2712,24 +2712,20 @@ namespace UnrealBuildTool
 		/// <summary>
 		/// Extract information from the csproj file based on the supplied configuration
 		/// </summary>
-		public IDotnetProjectInfo? GetProjectInfo(UnrealTargetConfiguration InConfiguration)
+		public IDotnetProjectInfo GetProjectInfo(UnrealTargetConfiguration InConfiguration)
 		{
 			if (CachedProjectInfo.TryGetValue(InConfiguration, out IDotnetProjectInfo? value))
 			{
 				return value;
 			}
 
-			CsProjectInfo? Info;
-
 			Dictionary<string, string> Properties = [];
 			Properties.Add("Platform", "AnyCPU");
 			Properties.Add("Configuration", InConfiguration.ToString());
 			Properties.Add("EngineDirectory", Unreal.EngineDirectory.FullName);
-			if (CsProjectInfo.TryRead(ProjectFilePath, Properties, out Info))
-			{
-				CachedProjectInfo.Add(InConfiguration, Info);
-			}
 
+			DotnetProjectInfo Info = DotnetProjectInfo.Read(ProjectFilePath, Properties);
+			CachedProjectInfo.Add(InConfiguration, Info);
 			return Info;
 		}
 
@@ -2738,7 +2734,7 @@ namespace UnrealBuildTool
 		/// </summary>
 		public bool IsDotnetCoreOrDotnet5Plus()
 		{
-			IDotnetProjectInfo Info = GetProjectInfo(UnrealTargetConfiguration.Debug)!;
+			IDotnetProjectInfo Info = GetProjectInfo(UnrealTargetConfiguration.Debug);
 			return Info.IsDotnetCoreOrDotnet5Plus();
 		}
 
