Shared: Misc analysis fixes

#rnx
#Synced-CL 49186938

diff --git a/Engine/Source/Programs/Shared/EpicGames.Build/CsProjectInfo.cs b/Engine/Source/Programs/Shared/EpicGames.Build/CsProjectInfo.cs
--- a/Engine/Source/Programs/Shared/EpicGames.Build/CsProjectInfo.cs
+++ b/Engine/Source/Programs/Shared/EpicGames.Build/CsProjectInfo.cs
@@ -381,8 +381,10 @@ namespace EpicGames.Core
 			}
 		}
 
+		/// <summary>
 		/// recursive helper used by the function below that will append RemainingComponents one by one to ExistingPath, 
 		/// expanding wildcards as necessary. The complete list of files that match the complete path is returned out OutFoundFiles
+		/// </summary>
 		static void ProcessPathComponents(DirectoryReference existingPath, IEnumerable<string> remainingComponents, List<FileReference> outFoundFiles)
 		{
 			if (!remainingComponents.Any())
@@ -666,7 +668,7 @@ namespace EpicGames.Core
 			}
 			catch (CsProjectParseException ex)
 			{
-				throw new CsProjectParseException($"{ex.Message} while parsing {element} in project file {projectInfo.ProjectPath}", ex);
+				throw new CsProjectParseException($"{ex.Message} while parsing {element.Name} in project file {projectInfo.ProjectPath}", ex);
 			}
 		}
 
diff --git a/Engine/Source/Programs/Shared/EpicGames.Build/System/CompileScriptModules.cs b/Engine/Source/Programs/Shared/EpicGames.Build/System/CompileScriptModules.cs
--- a/Engine/Source/Programs/Shared/EpicGames.Build/System/CompileScriptModules.cs
+++ b/Engine/Source/Programs/Shared/EpicGames.Build/System/CompileScriptModules.cs
@@ -434,9 +434,9 @@ namespace UnrealBuildBase
 						BuildRecord = JsonSerializer.Deserialize<CsProjBuildRecord>(FileReference.ReadAllText(JsonFile));
 						Logger.LogTrace("Loaded script module build record {JsonFile}", JsonFile);
 					}
-					catch (Exception Ex)
+					catch (Exception ex)
 					{
-						Logger.LogWarning("[{JsonFile}] Failed to load build record: {Message}", JsonFile, Ex.Message);
+						Logger.LogWarning(ex, "[{JsonFile}] Failed to load build record: {Message}", JsonFile, ex.Message);
 					}
 
 					if (BuildRecord != null && BuildRecord.ProjectPath != null)
diff --git a/Engine/Source/Programs/Shared/EpicGames.Build/System/NativeProjectsBase.cs b/Engine/Source/Programs/Shared/EpicGames.Build/System/NativeProjectsBase.cs
--- a/Engine/Source/Programs/Shared/EpicGames.Build/System/NativeProjectsBase.cs
+++ b/Engine/Source/Programs/Shared/EpicGames.Build/System/NativeProjectsBase.cs
@@ -17,7 +17,7 @@ namespace UnrealBuildBase
 		/// <summary>
 		/// Object used for synchronizing access to static fields
 		/// </summary>
-		protected static object LockObject = new object();
+		static protected object LockObject = new object();
 
 		/// <summary>
 		/// The native project base directories
diff --git a/Engine/Source/Programs/Shared/EpicGames.Core/JsonTracer.cs b/Engine/Source/Programs/Shared/EpicGames.Core/JsonTracer.cs
--- a/Engine/Source/Programs/Shared/EpicGames.Core/JsonTracer.cs
+++ b/Engine/Source/Programs/Shared/EpicGames.Core/JsonTracer.cs
@@ -259,7 +259,7 @@ namespace EpicGames.Core
 			{
 				int hashCode = 2083322454;
 				hashCode = hashCode * -1521134295 + EqualityComparer<JsonTracerSpanContext>.Default.GetHashCode(Context);
-				hashCode = hashCode * -1521134295 + EqualityComparer<string>.Default.GetHashCode(ReferenceType);
+				hashCode = hashCode * -1521134295 + StringComparer.Ordinal.GetHashCode(ReferenceType);
 				return hashCode;
 			}
 #pragma warning restore IDE0070 // Use 'System.HashCode'
diff --git a/Engine/Source/Programs/Shared/EpicGames.Core/Log.cs b/Engine/Source/Programs/Shared/EpicGames.Core/Log.cs
--- a/Engine/Source/Programs/Shared/EpicGames.Core/Log.cs
+++ b/Engine/Source/Programs/Shared/EpicGames.Core/Log.cs
@@ -915,7 +915,7 @@ namespace EpicGames.Core
 		/// <summary>
 		/// Temporary status message displayed on the console.
 		/// </summary>
-		[DebuggerDisplay("{HeadingText}")]
+		[DebuggerDisplay("{_headingText}")]
 		class StatusMessage
 		{
 			/// <summary>
diff --git a/Engine/Source/Programs/Shared/EpicGames.Perforce.Managed/ManagedWorkspace.cs b/Engine/Source/Programs/Shared/EpicGames.Perforce.Managed/ManagedWorkspace.cs
--- a/Engine/Source/Programs/Shared/EpicGames.Perforce.Managed/ManagedWorkspace.cs
+++ b/Engine/Source/Programs/Shared/EpicGames.Perforce.Managed/ManagedWorkspace.cs
@@ -1144,7 +1144,7 @@ namespace EpicGames.Perforce.Managed
 			PerforceResponseList<PrintRecord> printResponse = await perforce.TryPrintAsync($"{perforceClient.Root}{Path.DirectorySeparatorChar}...", $"//{perforceClient.Name}/...@={unshelveChangelist}", cancellationToken);
 			if (!printResponse.Succeeded)
 			{
-				_logger.LogWarning("Unable to print shelved changelist: {Error}", printResponse.ToString());
+				_logger.LogWarning("Unable to print shelved changelist: {Error}", printResponse.FirstError?.ToString());
 			}
 
 			_logger.LogInformation("Completed in {TimeSeconds}s", $"{timer.Elapsed.TotalSeconds:0.0}");
diff --git a/Engine/Source/Programs/Shared/EpicGames.Serialization/CbField.cs b/Engine/Source/Programs/Shared/EpicGames.Serialization/CbField.cs
--- a/Engine/Source/Programs/Shared/EpicGames.Serialization/CbField.cs
+++ b/Engine/Source/Programs/Shared/EpicGames.Serialization/CbField.cs
@@ -2790,9 +2790,9 @@ namespace EpicGames.Serialization
 		/// <returns></returns>
 		public byte[] ToByteArray()
 		{
-			byte[] b = new byte[12];
-			CopyTo(b);
-			return b;
+			byte[] bytes = new byte[12];
+			CopyTo(bytes);
+			return bytes;
 		}
 
 		/// <summary>
diff --git a/Engine/Source/Programs/Shared/EpicGames.UBA/EpicGames.UBA.csproj b/Engine/Source/Programs/Shared/EpicGames.UBA/EpicGames.UBA.csproj
--- a/Engine/Source/Programs/Shared/EpicGames.UBA/EpicGames.UBA.csproj
+++ b/Engine/Source/Programs/Shared/EpicGames.UBA/EpicGames.UBA.csproj
@@ -37,7 +37,7 @@
 
   <ItemGroup>
     <Compile Include="../../../../Restricted/*/Source/Programs/Shared/EpicGames.UBA/**/*.cs">
-      <Link>$([System.Text.RegularExpressions.Regex]::Replace(%(FullPath), '^.+?[\\/]EpicGames.UBA[\\/]', ''))</Link>
+      <Link>Restricted/$([System.Text.RegularExpressions.Regex]::Replace(%(FullPath), '^.+?[\\/]EpicGames.UBA[\\/]', ''))</Link>
     </Compile>
   </ItemGroup>
 
diff --git a/Engine/Source/Programs/Shared/EpicGames.UHT/Exporters/CodeGen/UhtRigVMStringBuilderExtensions.cs b/Engine/Source/Programs/Shared/EpicGames.UHT/Exporters/CodeGen/UhtRigVMStringBuilderExtensions.cs
--- a/Engine/Source/Programs/Shared/EpicGames.UHT/Exporters/CodeGen/UhtRigVMStringBuilderExtensions.cs
+++ b/Engine/Source/Programs/Shared/EpicGames.UHT/Exporters/CodeGen/UhtRigVMStringBuilderExtensions.cs
@@ -39,7 +39,7 @@ namespace EpicGames.UHT.Exporters.CodeGen
 			builder.Append("const ");
 			string paramterType = parameter.TypeOriginal(castType);
 			ReadOnlySpan<char> parameterTypeSpan = paramterType.AsSpan();
-			if (parameterTypeSpan.EndsWith("&"))
+			if (parameterTypeSpan.EndsWith("&", StringComparison.Ordinal))
 			{
 				parameterTypeSpan = parameterTypeSpan[0..^1];
 			}
@@ -55,7 +55,7 @@ namespace EpicGames.UHT.Exporters.CodeGen
 		{
 			string paramterType = parameter.TypeOriginal(castType);
 			ReadOnlySpan<char> parameterTypeSpan = paramterType.AsSpan();
-			if (parameterTypeSpan.EndsWith("&"))
+			if (parameterTypeSpan.EndsWith("&", StringComparison.Ordinal))
 			{
 				builder.Append(parameterTypeSpan);
 			}
@@ -70,7 +70,7 @@ namespace EpicGames.UHT.Exporters.CodeGen
 		{
 			string paramterType = parameter.TypeOriginal(castType);
 			ReadOnlySpan<char> parameterTypeSpan = paramterType.AsSpan();
-			if (parameterTypeSpan.EndsWith("&"))
+			if (parameterTypeSpan.EndsWith("&", StringComparison.Ordinal))
 			{
 				parameterTypeSpan = parameterTypeSpan[0..^1];
 			}
diff --git a/Engine/Source/Programs/Shared/EpicGames.UHT/Parsers/UhtFunctionParser.cs b/Engine/Source/Programs/Shared/EpicGames.UHT/Parsers/UhtFunctionParser.cs
--- a/Engine/Source/Programs/Shared/EpicGames.UHT/Parsers/UhtFunctionParser.cs
+++ b/Engine/Source/Programs/Shared/EpicGames.UHT/Parsers/UhtFunctionParser.cs
@@ -588,7 +588,7 @@ namespace EpicGames.UHT.Parsers
 
 		internal static bool IsValidateDelegateDeclaration(UhtToken token)
 		{
-			return (token.IsIdentifier() && token.Value.Span.StartsWith("DECLARE_DYNAMIC_"));
+			return (token.IsIdentifier() && token.Value.Span.StartsWith("DECLARE_DYNAMIC_", StringComparison.Ordinal));
 		}
 
 		private static void ParseParameterList(UhtParsingScope topScope, UhtPropertyParseOptions options)
diff --git a/Engine/Source/Programs/Shared/EpicGames.UHT/Parsers/UhtHeaderFileParser.cs b/Engine/Source/Programs/Shared/EpicGames.UHT/Parsers/UhtHeaderFileParser.cs
--- a/Engine/Source/Programs/Shared/EpicGames.UHT/Parsers/UhtHeaderFileParser.cs
+++ b/Engine/Source/Programs/Shared/EpicGames.UHT/Parsers/UhtHeaderFileParser.cs
@@ -780,7 +780,7 @@ namespace EpicGames.UHT.Parsers
 			}
 
 			// Identifiers which end in _API are known
-			if (token.Value.Span.Length > 4 && token.Value.Span.EndsWith("_API"))
+			if (token.Value.Span.Length > 4 && token.Value.Span.EndsWith("_API", StringComparison.Ordinal))
 			{
 				return UhtParseResult.Unhandled;
 			}
@@ -1256,7 +1256,7 @@ namespace EpicGames.UHT.Parsers
 			while (true)
 			{
 				UhtToken peekToken = replayReader.PeekToken();
-				if (!peekToken.IsValue("virtual") && !peekToken.Value.Span.EndsWith("_API"))
+				if (!peekToken.IsValue("virtual") && !peekToken.Value.Span.EndsWith("_API", StringComparison.Ordinal))
 				{
 					break;
 				}
@@ -1272,7 +1272,7 @@ namespace EpicGames.UHT.Parsers
 			// Allow explicit constructors
 			{
 				bool foundExplicit = replayReader.TryOptional("explicit");
-				if (replayReader.PeekToken().Value.Span.EndsWith("_API"))
+				if (replayReader.PeekToken().Value.Span.EndsWith("_API", StringComparison.Ordinal))
 				{
 					replayReader.ConsumeToken();
 					if (!foundExplicit)
diff --git a/Engine/Source/Programs/Shared/EpicGames.UHT/Tokenizer/UhtTokenReaderGeneralExtensions.cs b/Engine/Source/Programs/Shared/EpicGames.UHT/Tokenizer/UhtTokenReaderGeneralExtensions.cs
--- a/Engine/Source/Programs/Shared/EpicGames.UHT/Tokenizer/UhtTokenReaderGeneralExtensions.cs
+++ b/Engine/Source/Programs/Shared/EpicGames.UHT/Tokenizer/UhtTokenReaderGeneralExtensions.cs
@@ -222,7 +222,7 @@ namespace EpicGames.UHT.Tokenizer
 		public static IUhtTokenReader OptionalEndsWith(this IUhtTokenReader tokenReader, string text, Action<string> action)
 		{
 			ref UhtToken token = ref tokenReader.PeekToken();
-			if (token.IsIdentifier() && token.Value.Span.EndsWith(text))
+			if (token.IsIdentifier() && token.Value.Span.EndsWith(text, StringComparison.Ordinal))
 			{
 				tokenReader.ConsumeToken();
 				action(token.Value.Span.ToString());
diff --git a/Engine/Source/Programs/Shared/EpicGames.UHT/Tokenizer/UhtTokenReaderUtilityExtensions.cs b/Engine/Source/Programs/Shared/EpicGames.UHT/Tokenizer/UhtTokenReaderUtilityExtensions.cs
--- a/Engine/Source/Programs/Shared/EpicGames.UHT/Tokenizer/UhtTokenReaderUtilityExtensions.cs
+++ b/Engine/Source/Programs/Shared/EpicGames.UHT/Tokenizer/UhtTokenReaderUtilityExtensions.cs
@@ -44,7 +44,7 @@ namespace EpicGames.UHT.Tokenizer
 		public static bool TryOptionalAPIMacro(this IUhtTokenReader tokenReader, out UhtToken apiMacroToken)
 		{
 			ref UhtToken token = ref tokenReader.PeekToken();
-			if (token.IsIdentifier() && token.Value.Span.EndsWith("_API"))
+			if (token.IsIdentifier() && token.Value.Span.EndsWith("_API", StringComparison.Ordinal))
 			{
 				apiMacroToken = token;
 				tokenReader.ConsumeToken();
@@ -207,9 +207,9 @@ namespace EpicGames.UHT.Tokenizer
 			}
 
 			// Test for known delegate and event macros.
-			if (span.StartsWith("DECLARE_MULTICAST_DELEGATE") ||
-				span.StartsWith("DECLARE_DELEGATE") ||
-				span.StartsWith("DECLARE_EVENT"))
+			if (span.StartsWith("DECLARE_MULTICAST_DELEGATE", StringComparison.Ordinal) ||
+				span.StartsWith("DECLARE_DELEGATE", StringComparison.Ordinal) ||
+				span.StartsWith("DECLARE_EVENT", StringComparison.Ordinal))
 			{
 				return true;
 			}
diff --git a/Engine/Source/Programs/Shared/EpicGames.UHT/Types/UhtClass.cs b/Engine/Source/Programs/Shared/EpicGames.UHT/Types/UhtClass.cs
--- a/Engine/Source/Programs/Shared/EpicGames.UHT/Types/UhtClass.cs
+++ b/Engine/Source/Programs/Shared/EpicGames.UHT/Types/UhtClass.cs
@@ -1547,7 +1547,7 @@ namespace EpicGames.UHT.Types
 			int tokenCount = tokens.Length;
 
 			// Skip all of the API and virtual keywords
-			for (; tokenIndex < tokenCount && (tokens[tokenIndex].IsValue("virtual") || tokens[tokenIndex].Value.Span.EndsWith("_API")); ++tokenIndex)
+			for (; tokenIndex < tokenCount && (tokens[tokenIndex].IsValue("virtual") || tokens[tokenIndex].Value.Span.EndsWith("_API", StringComparison.Ordinal)); ++tokenIndex)
 			{
 				// Nothing to do in the body
 			}
@@ -1692,7 +1692,7 @@ namespace EpicGames.UHT.Types
 
 					// Make sure we match the expected type
 					if (tokenSpan.Length > expectedSpan.Length ||
-						!expectedSpan.StartsWith(tokenSpan))
+						!expectedSpan.StartsWith(tokenSpan, StringComparison.Ordinal))
 					{
 						return false;
 					}
diff --git a/Engine/Source/Programs/Shared/EpicGames.UHT/Types/UhtHeaderFile.cs b/Engine/Source/Programs/Shared/EpicGames.UHT/Types/UhtHeaderFile.cs
--- a/Engine/Source/Programs/Shared/EpicGames.UHT/Types/UhtHeaderFile.cs
+++ b/Engine/Source/Programs/Shared/EpicGames.UHT/Types/UhtHeaderFile.cs
@@ -317,7 +317,7 @@ namespace EpicGames.UHT.Types
 				catch (Exception ex)
 				{
 					Session.FilesCacheReadErrors.Increment();
-					Session.Logger.LogInformation("File {FileName} failed to deserialize from input cache: {Message}", _sourceFile.FilePath, ex.Message);
+					Session.Logger.LogInformation(ex, "File {FileName} failed to deserialize from input cache: {Message}", _sourceFile.FilePath, ex.Message);
 				}
 			}
 			else
diff --git a/Engine/Source/Programs/Shared/EpicGames.UHT/Types/UhtNamespace.cs b/Engine/Source/Programs/Shared/EpicGames.UHT/Types/UhtNamespace.cs
--- a/Engine/Source/Programs/Shared/EpicGames.UHT/Types/UhtNamespace.cs
+++ b/Engine/Source/Programs/Shared/EpicGames.UHT/Types/UhtNamespace.cs
@@ -66,7 +66,7 @@ namespace EpicGames.UHT.Types
 		public StringBuilder AppendSingleLine(StringBuilder builder, Action<StringBuilder> body)
 		{
 			ReadOnlySpan<char> namespaceName = FullSourceName.AsSpan();
-			if (namespaceName.EndsWith("::"))
+			if (namespaceName.EndsWith("::", StringComparison.Ordinal))
 			{
 				namespaceName = namespaceName[..^2];
 			}
@@ -92,7 +92,7 @@ namespace EpicGames.UHT.Types
 		public StringBuilder AppendMultipleLines(StringBuilder builder, Action<StringBuilder> body)
 		{
 			ReadOnlySpan<char> namespaceName = FullSourceName.AsSpan();
-			if (namespaceName.EndsWith("::"))
+			if (namespaceName.EndsWith("::", StringComparison.Ordinal))
 			{
 				namespaceName = namespaceName[..^2];
 			}
diff --git a/Engine/Source/Programs/Shared/EpicGames.UHT/Types/UhtProperty.cs b/Engine/Source/Programs/Shared/EpicGames.UHT/Types/UhtProperty.cs
--- a/Engine/Source/Programs/Shared/EpicGames.UHT/Types/UhtProperty.cs
+++ b/Engine/Source/Programs/Shared/EpicGames.UHT/Types/UhtProperty.cs
@@ -1996,7 +1996,7 @@ namespace EpicGames.UHT.Types
 							{
 								foreach (string cast in s_casts)
 								{
-									if (dim.StartsWith(cast))
+									if (dim.StartsWith(cast, StringComparison.Ordinal))
 									{
 										dim = dim[cast.Length..];
 										again = true;
diff --git a/Engine/Source/Programs/Shared/EpicGames.UHT/Utils/UhtExport.cs b/Engine/Source/Programs/Shared/EpicGames.UHT/Utils/UhtExport.cs
--- a/Engine/Source/Programs/Shared/EpicGames.UHT/Utils/UhtExport.cs
+++ b/Engine/Source/Programs/Shared/EpicGames.UHT/Utils/UhtExport.cs
@@ -36,15 +36,15 @@ namespace EpicGames.UHT.Utils
 		/// </summary>
 		/// <param name="prereqs">Tasks that must be completed prior to this task running</param>
 		/// <param name="action">Action to be invoked to generate the output(s)</param>
-		/// <returns>Task object or null if the task was immediately executed.</returns>
-		public Task? CreateTask(List<Task?>? prereqs, UhtExportTaskDelegate action);
+		/// <returns>Task object or Task.CompletedTask if the task was immediately executed.</returns>
+		public Task CreateTask(List<Task?>? prereqs, UhtExportTaskDelegate action);
 
 		/// <summary>
 		/// Create a task
 		/// </summary>
 		/// <param name="action">Action to be invoked to generate the output(s)</param>
-		/// <returns>Task object or null if the task was immediately executed.</returns>
-		public Task? CreateTask(UhtExportTaskDelegate action);
+		/// <returns>Task object or Task.CompletedTask if the task was immediately executed.</returns>
+		public Task CreateTask(UhtExportTaskDelegate action);
 
 		/// <summary>
 		/// Commit the contents of the string builder as the output.
diff --git a/Engine/Source/Programs/Shared/EpicGames.UHT/Utils/UhtInputCache.cs b/Engine/Source/Programs/Shared/EpicGames.UHT/Utils/UhtInputCache.cs
--- a/Engine/Source/Programs/Shared/EpicGames.UHT/Utils/UhtInputCache.cs
+++ b/Engine/Source/Programs/Shared/EpicGames.UHT/Utils/UhtInputCache.cs
@@ -1148,7 +1148,7 @@ namespace EpicGames.UHT.Utils
 			}
 			catch (Exception ex)
 			{
-				module.Session.Logger.LogTrace("Error reading input cache {Message}", ex.Message);
+				module.Session.Logger.LogTrace(ex, "Error reading input cache {Message}", ex.Message);
 				DeleteCache(module);
 				module.Session.CacheFilesSkippedCorrupted.Increment();
 				return new();
@@ -1208,7 +1208,7 @@ namespace EpicGames.UHT.Utils
 			}
 			catch (Exception ex)
 			{
-				module.Session.Logger.LogTrace("Error creating input cache {Message}", ex.Message);
+				module.Session.Logger.LogTrace(ex, "Error creating input cache {Message}", ex.Message);
 				DeleteCache(module);
 			}
 		}
diff --git a/Engine/Source/Programs/Shared/EpicGames.UHT/Utils/UhtPropertyResolveArgs.cs b/Engine/Source/Programs/Shared/EpicGames.UHT/Utils/UhtPropertyResolveArgs.cs
--- a/Engine/Source/Programs/Shared/EpicGames.UHT/Utils/UhtPropertyResolveArgs.cs
+++ b/Engine/Source/Programs/Shared/EpicGames.UHT/Utils/UhtPropertyResolveArgs.cs
@@ -247,7 +247,7 @@ namespace EpicGames.UHT.Utils
 			{
 				UhtTokenList last = identifier.Last;
 				UhtToken token = last.Token;
-				if (!token.Value.Span.StartsWith("I") || !token.Value.Span.EndsWith(UhtNames.VerseProxySuffix))
+				if (!token.Value.Span.StartsWith("I", StringComparison.Ordinal) || !token.Value.Span.EndsWith(UhtNames.VerseProxySuffix, StringComparison.Ordinal))
 				{
 					TokenReader.LogError(token.InputLine, $"Expected a Verse interface proxy.  Must begin with 'I' and end with '{UhtNames.VerseProxySuffix}'");
 					return null;
diff --git a/Engine/Source/Programs/Shared/EpicGames.UHT/Utils/UhtSession.cs b/Engine/Source/Programs/Shared/EpicGames.UHT/Utils/UhtSession.cs
--- a/Engine/Source/Programs/Shared/EpicGames.UHT/Utils/UhtSession.cs
+++ b/Engine/Source/Programs/Shared/EpicGames.UHT/Utils/UhtSession.cs
@@ -233,8 +233,8 @@ namespace EpicGames.UHT.Utils
 		/// </summary>
 		/// <param name="prereqs">Tasks that must be completed prior to this task running</param>
 		/// <param name="action">Action to be invoked to generate the output</param>
-		/// <returns>Task object or null if the task was immediately executed.</returns>
-		public Task? CreateTask(List<Task?>? prereqs, UhtExportTaskDelegate action)
+		/// <returns>Task object or Task.CompletedTask if the task was immediately executed.</returns>
+		public Task CreateTask(List<Task?>? prereqs, UhtExportTaskDelegate action)
 		{
 			if (Session.GoWide)
 			{
@@ -251,7 +251,7 @@ namespace EpicGames.UHT.Utils
 			else
 			{
 				Session.TryNoErrorCheck(null, () => action(this));
-				return null;
+				return Task.CompletedTask;
 			}
 		}
 
@@ -259,8 +259,8 @@ namespace EpicGames.UHT.Utils
 		/// Create a task to export two files
 		/// </summary>
 		/// <param name="action">Action to be invoked to generate the output</param>
-		/// <returns>Task object or null if the task was immediately executed.</returns>
-		public Task? CreateTask(UhtExportTaskDelegate action)
+		/// <returns>Task object or Task.CompletedTask if the task was immediately executed.</returns>
+		public Task CreateTask(UhtExportTaskDelegate action)
 		{
 			return CreateTask(null, action);
 		}
@@ -1674,7 +1674,7 @@ namespace EpicGames.UHT.Utils
 			StringView parts = new(fullName);
 			while (parts.Length > 0)
 			{
-				int sep = parts.Span.IndexOf("::");
+				int sep = parts.Span.IndexOf("::", StringComparison.Ordinal);
 				if (sep == -1)
 				{
 					sep = parts.Length;
@@ -2153,7 +2153,7 @@ namespace EpicGames.UHT.Utils
 						}
 						catch (Exception ex)
 						{
-							Logger.LogInformation("File {Filename} failed to serialize to input cache: {Message}", headerFile.FilePath, ex.Message);
+							Logger.LogInformation(ex, "File {Filename} failed to serialize to input cache: {Message}", headerFile.FilePath, ex.Message);
 							FilesCacheWriteErrors.Increment();
 						}
 					}
diff --git a/Engine/Source/Programs/Shared/EpicGames.UHT/Utils/UhtUtilities.cs b/Engine/Source/Programs/Shared/EpicGames.UHT/Utils/UhtUtilities.cs
--- a/Engine/Source/Programs/Shared/EpicGames.UHT/Utils/UhtUtilities.cs
+++ b/Engine/Source/Programs/Shared/EpicGames.UHT/Utils/UhtUtilities.cs
@@ -213,7 +213,7 @@ namespace EpicGames.UHT.Utils
 				case 'A':
 				case 'U':
 					// If it is a class prefix, check for deprecated class prefix also
-					if (sourceName.Span.Length > 12 && sourceName.Span[1..].StartsWith("DEPRECATED_"))
+					if (sourceName.Span.Length > 12 && sourceName.Span[1..].StartsWith("DEPRECATED_", StringComparison.Ordinal))
 					{
 						return new UhtEngineNameParts { Prefix = new StringView(sourceName, 0, 12), EngineName = new StringView(sourceName, 12), IsDeprecated = true };
 					}
