EpicGames.UBA: Add EventId to UBA logging

#rnx
#Synced-CL 49215227

diff --git a/Engine/Source/Programs/Shared/EpicGames.Core/KnownLogEvents.cs b/Engine/Source/Programs/Shared/EpicGames.Core/KnownLogEvents.cs
--- a/Engine/Source/Programs/Shared/EpicGames.Core/KnownLogEvents.cs
+++ b/Engine/Source/Programs/Shared/EpicGames.Core/KnownLogEvents.cs
@@ -4,10 +4,10 @@ using Microsoft.Extensions.Logging;
 
 namespace EpicGames.Core
 {
-#pragma warning disable CA1707 // Identifiers should not contain underscores
 	/// <summary>
 	/// Well known log events
 	/// </summary>
+	[System.Diagnostics.CodeAnalysis.SuppressMessage("Naming", "CA1707:Identifiers should not contain underscores", Justification = "Intentional")]
 	public static class KnownLogEvents
 	{
 		/// <summary>
@@ -60,6 +60,11 @@ namespace EpicGames.Core
 		/// </summary>
 		public static EventId UHT { get; } = new EventId(110);
 
+		/// <summary>
+		/// UBA message
+		/// </summary>
+		public static EventId UBA { get; } = new EventId(111);
+
 		/// <summary>
 		/// Linker error
 		/// </summary>
@@ -421,5 +426,4 @@ namespace EpicGames.Core
 		/// </summary>
 		public static EventId ZenCli_Max { get; } = new EventId(1199);
 	}
-#pragma warning restore CA1707 // Identifiers should not contain underscores
 }
diff --git a/Engine/Source/Programs/Shared/EpicGames.UBA/EpicGames.UBA.csproj b/Engine/Source/Programs/Shared/EpicGames.UBA/EpicGames.UBA.csproj
--- a/Engine/Source/Programs/Shared/EpicGames.UBA/EpicGames.UBA.csproj
+++ b/Engine/Source/Programs/Shared/EpicGames.UBA/EpicGames.UBA.csproj
@@ -53,4 +53,7 @@
     <PackageReference Include="Microsoft.CSharp" Version="4.7.0" />
     <PackageReference Include="Microsoft.Extensions.Logging" Version="8.0.1" />
   </ItemGroup>
+  <ItemGroup>
+    <ProjectReference Include="..\EpicGames.Core\EpicGames.Core.csproj" />
+  </ItemGroup>
 </Project>
\ No newline at end of file
diff --git a/Engine/Source/Programs/Shared/EpicGames.UBA/Impl/LoggerImpl.cs b/Engine/Source/Programs/Shared/EpicGames.UBA/Impl/LoggerImpl.cs
--- a/Engine/Source/Programs/Shared/EpicGames.UBA/Impl/LoggerImpl.cs
+++ b/Engine/Source/Programs/Shared/EpicGames.UBA/Impl/LoggerImpl.cs
@@ -1,12 +1,16 @@
 // Copyright Epic Games, Inc. All Rights Reserved.
 
-using Microsoft.Extensions.Logging;
 using System;
 using System.Runtime.InteropServices;
 using System.Threading;
+using EpicGames.Core;
+using Microsoft.Extensions.Logging;
 
 namespace EpicGames.UBA.Impl
 {
+	/// <summary>
+	/// Logger implementation which logs to an ILogger
+	/// </summary>
 	internal class LoggerImpl : ILogger
 	{
 		delegate void BeginScopeCallback();
@@ -52,6 +56,7 @@ namespace EpicGames.UBA.Impl
 		#region IDisposable
 		~LoggerImpl() => Dispose(false);
 
+		/// <inheritdoc/>
 		public void Dispose()
 		{
 			Dispose(true);
@@ -73,60 +78,39 @@ namespace EpicGames.UBA.Impl
 		#endregion
 
 		#region ILogger
-
+		/// <inheritdoc/>
 		public nint GetHandle() => _handle;
 
-		public void BeginScope()
-		{
-			Monitor.Enter(_lock);
-		}
+		/// <inheritdoc/>
+		public void BeginScope() => Monitor.Enter(_lock);
 
-		public void EndScope()
-		{
-			Monitor.Exit(_lock);
-		}
+		/// <inheritdoc/>
+		public void EndScope() => Monitor.Exit(_lock);
 
+		/// <inheritdoc/>
 		public void Log(LogEntryType type, string message)
 		{
+			LogLevel logLevel = LogEntryTypeToLogLevel(type);
 			lock (_lock)
 			{
-				switch (type)
-				{
-					case LogEntryType.Error: _logger.LogError("{Message}", message); break;
-					case LogEntryType.Warning: _logger.LogWarning("{Message}", message); break;
-					case LogEntryType.Info: _logger.LogInformation("{Message}", message); break;
-					case LogEntryType.Detail: _logger.LogDebug("{Message}", message); break;
-					case LogEntryType.Debug: _logger.LogDebug("{Message}", message); break;
-				}
+				_logger.Log(logLevel, KnownLogEvents.UBA, "{Message}", message);
 			}
 		}
 		#endregion
 
 		void Log(LogEntryType type, nint ptr, uint len) => Log(type, Marshal.PtrToStringAuto(ptr, (int)len) ?? String.Empty);
-	}
 
-	internal class NullLogger : ILogger
-	{
-		#region DllImport
-		[DllImport("UbaHost", CharSet = CharSet.Auto)]
-		static extern nint LogWriter_GetNull();
-		#endregion
-
-		/// <inheritdoc/>
-		public void BeginScope() { }
-		/// <inheritdoc/>
-		public void EndScope() { }
-		/// <inheritdoc/>
-		public void Log(LogEntryType type, string message) { }
-		/// <inheritdoc/>
-		public nint GetHandle() => LogWriter_GetNull();
-		/// <inheritdoc/>
-		public void Dispose()
+		static LogLevel LogEntryTypeToLogLevel(LogEntryType type)
 		{
-			Dispose(true);
-			GC.SuppressFinalize(this);
+			switch (type)
+			{
+				case LogEntryType.Error: return LogLevel.Error;
+				case LogEntryType.Warning: return LogLevel.Warning;
+				case LogEntryType.Info: return LogLevel.Information;
+				case LogEntryType.Detail:
+				case LogEntryType.Debug:
+				default: return LogLevel.Debug;
+			}
 		}
-		/// <inheritdoc/>
-		protected virtual void Dispose(bool disposing) { }
 	}
 }
\ No newline at end of file
diff --git a/Engine/Source/Programs/Shared/EpicGames.UBA/Impl/NullLogger.cs b/Engine/Source/Programs/Shared/EpicGames.UBA/Impl/NullLogger.cs
new file mode 100644
--- /dev/null
+++ b/Engine/Source/Programs/Shared/EpicGames.UBA/Impl/NullLogger.cs
@@ -0,0 +1,40 @@
+// Copyright Epic Games, Inc. All Rights Reserved.
+
+using System;
+using System.Runtime.InteropServices;
+
+namespace EpicGames.UBA.Impl
+{
+	/// <summary>
+	/// Logger implementation which discards output
+	/// </summary>
+	internal class NullLogger : ILogger
+	{
+		#region DllImport
+		[DllImport("UbaHost", CharSet = CharSet.Auto)]
+		static extern nint LogWriter_GetNull();
+		#endregion
+
+		/// <inheritdoc/>
+		public void BeginScope() { }
+
+		/// <inheritdoc/>
+		public void EndScope() { }
+
+		/// <inheritdoc/>
+		public void Log(LogEntryType type, string message) { }
+
+		/// <inheritdoc/>
+		public nint GetHandle() => LogWriter_GetNull();
+
+		/// <inheritdoc/>
+		public void Dispose()
+		{
+			Dispose(true);
+			GC.SuppressFinalize(this);
+		}
+
+		/// <inheritdoc/>
+		protected virtual void Dispose(bool disposing) { }
+	}
+}
