EpicGames.Build: Don't use MD5 in ProcessSingleton

#rnx
#Synced-CL 49215227

diff --git a/Engine/Source/Programs/Shared/EpicGames.Build/Automation/ProcessSingleton.cs b/Engine/Source/Programs/Shared/EpicGames.Build/Automation/ProcessSingleton.cs
--- a/Engine/Source/Programs/Shared/EpicGames.Build/Automation/ProcessSingleton.cs
+++ b/Engine/Source/Programs/Shared/EpicGames.Build/Automation/ProcessSingleton.cs
@@ -1,25 +1,16 @@
 ﻿// Copyright Epic Games, Inc. All Rights Reserved.
 
-using EpicGames.Core;
-using Microsoft.Extensions.Logging;
 using System;
-using System.Collections.Generic;
-using System.IO;
 using System.Reflection;
-using System.Text;
 using System.Threading;
 using System.Threading.Tasks;
+using EpicGames.Core;
+using Microsoft.Extensions.Logging;
 
 namespace UnrealBuildBase
 {
-	public class ProcessSingleton
+	public static class ProcessSingleton
 	{
-		public static string GetUniqueMutexForPath(string Name, string UniquePath)
-		{
-			// generate a md5 hash of the path, as GetHashCode is not guaranteed to generate a stable hash
-			return String.Format("Global\\{0}_{1}", Name, ContentHash.MD5(UniquePath.ToUpperInvariant()));
-		}
-
 		/// <summary>
 		/// Returns true/false based on whether this is the only instance
 		/// running (checked at startup).
@@ -29,71 +20,72 @@ namespace UnrealBuildBase
 		/// <summary>
 		/// Runs the specified delegate checking if this is the only instance of the application.
 		/// </summary>
-		/// <param name="Main"></param>
+		/// <param name="main"></param>
 		/// <param name="bWaitForUATMutex"></param>
-		/// <param name="Logger"></param>
-		public static async Task<ExitCode> RunSingleInstanceAsync(Func<Task<ExitCode>> Main, bool bWaitForUATMutex, ILogger Logger)
+		/// <param name="logger"></param>
+		public static async Task<ExitCode> RunSingleInstanceAsync(Func<Task<ExitCode>> main, bool bWaitForUATMutex, ILogger logger)
 		{
 			// Need to execute this logic on a background thread, since mutex ownership on Linux has thread affinity (ie. mutexes be released on the
 			// same thread that acquires it, which is not guaranteed through an async continuation)
-			TaskCompletionSource<ExitCode> Result = new TaskCompletionSource<ExitCode>();
-			Thread Thread = new Thread(() => RunSingleInstanceThread(Main, Result, bWaitForUATMutex, Logger));
-			Thread.Start();
-			return await Result.Task;
+			TaskCompletionSource<ExitCode> result = new TaskCompletionSource<ExitCode>();
+			Thread thread = new Thread(() => RunSingleInstanceThread(main, result, bWaitForUATMutex, logger));
+			thread.Start();
+			return await result.Task;
 		}
 
-		public static void RunSingleInstanceThread(Func<Task<ExitCode>> Main, TaskCompletionSource<ExitCode> Result, bool bWaitForUATMutex, ILogger Logger)
+		static void RunSingleInstanceThread(Func<Task<ExitCode>> main, TaskCompletionSource<ExitCode> result, bool bWaitForUATMutex, ILogger logger)
 		{
 			try
 			{
-				bool AllowMultipleInsances = (Environment.GetEnvironmentVariable("uebp_UATMutexNoWait") == "1");
+				bool allowMultipleInsances = (Environment.GetEnvironmentVariable("uebp_UATMutexNoWait") == "1");
 
-				string EntryAssemblyLocation = Assembly.GetEntryAssembly()!.GetOriginalLocation();
+				FileReference entryAssemblyLocation = FileReference.FromString(Assembly.GetEntryAssembly()?.GetOriginalLocation())
+					?? Unreal.DotnetPath;
 
-				string MutexName = GetUniqueMutexForPath(Path.GetFileNameWithoutExtension(EntryAssemblyLocation), EntryAssemblyLocation);
-				using (Mutex SingleInstanceMutex = new Mutex(true, MutexName, out bool bCreatedMutex))
+				string mutexName = GlobalSingleInstanceMutex.GetUniqueMutexForPath(entryAssemblyLocation.GetFileNameWithoutExtension(), entryAssemblyLocation);
+				using (Mutex singleInstanceMutex = new Mutex(true, mutexName, out bool bCreatedMutex))
 				{
 					IsSoleInstance = bCreatedMutex;
 
-					if (!IsSoleInstance && AllowMultipleInsances == false)
+					if (!IsSoleInstance && allowMultipleInsances == false)
 					{
 						if (bWaitForUATMutex)
 						{
-							Logger.LogWarning("Another instance of UAT at '{File}' is running, and the -WaitForUATMutex parameter has been used. Waiting for other UAT to finish...", EntryAssemblyLocation);
-							int Seconds = 0;
-							while (WaitMutexNoExceptions(SingleInstanceMutex, 15 * 1000) == false)
+							logger.LogWarning("Another instance of UAT at '{File}' is running, and the -WaitForUATMutex parameter has been used. Waiting for other UAT to finish...", entryAssemblyLocation);
+							int seconds = 0;
+							while (WaitMutexNoExceptions(singleInstanceMutex, 15 * 1000) == false)
 							{
-								Seconds += 15;
-								Logger.LogInformation("Still waiting for Mutex. {TimeSeconds} seconds has passed...", Seconds);
+								seconds += 15;
+								logger.LogInformation("Still waiting for Mutex. {TimeSeconds} seconds has passed...", seconds);
 							}
 						}
 						else
 						{
-							throw new Exception($"A conflicting instance of AutomationTool is already running. Current location: {EntryAssemblyLocation}. A process manager may be used to determine the conflicting process and what tool may have launched it");
+							throw new Exception($"A conflicting instance of AutomationTool is already running. Current location: {entryAssemblyLocation}. A process manager may be used to determine the conflicting process and what tool may have launched it");
 						}
 					}
 
-					ExitCode ExitCode = Task.Run(() => Main()).Result;
+					ExitCode exitCode = Task.Run(() => main()).Result;
 
 					if (IsSoleInstance)
 					{
-						SingleInstanceMutex.ReleaseMutex();
+						singleInstanceMutex.ReleaseMutex();
 					}
 
-					Result.SetResult(ExitCode);
+					result.SetResult(exitCode);
 				}
 			}
-			catch (Exception Ex)
+			catch (Exception ex)
 			{
-				Result.TrySetException(Ex);
+				result.TrySetException(ex);
 			}
 		}
 
-		static bool WaitMutexNoExceptions(Mutex Mutex, int TimeoutMs)
+		static bool WaitMutexNoExceptions(Mutex mutex, int timeoutMs = 15 * 1000)
 		{
 			try
 			{
-				return Mutex.WaitOne(15 * 1000);
+				return mutex.WaitOne(timeoutMs);
 			}
 			catch (AbandonedMutexException)
 			{
