Fix an analysis warning with EpicGames.Analytics.Generators.

#Synced-CL 49271786

diff --git a/Engine/Source/Programs/Shared/EpicGames.Analytics.Generators/AnalyticsTableGenerator.cs b/Engine/Source/Programs/Shared/EpicGames.Analytics.Generators/AnalyticsTableGenerator.cs
--- a/Engine/Source/Programs/Shared/EpicGames.Analytics.Generators/AnalyticsTableGenerator.cs
+++ b/Engine/Source/Programs/Shared/EpicGames.Analytics.Generators/AnalyticsTableGenerator.cs
@@ -168,7 +168,7 @@ namespace EpicGames.Analytics.Generators
 				}
 			}
 
-			return String.IsNullOrEmpty(tableName) ? null : new AnalyticsTableToGenerate(telemetryName, tableName, columns);
+			return String.IsNullOrEmpty(tableName) ? null : new AnalyticsTableToGenerate(telemetryName, tableName!, columns);
 #nullable disable
 		}
 
