Refactor shared projects to remove default values and pull out all common features to a Directory.Build.props file.

#Synced-CL 49271786

diff --git a/Engine/Source/Programs/Shared/EpicGames.Analytics.Generators/EpicGames.Analytics.Generators.csproj b/Engine/Source/Programs/Shared/EpicGames.Analytics.Generators/EpicGames.Analytics.Generators.csproj
--- a/Engine/Source/Programs/Shared/EpicGames.Analytics.Generators/EpicGames.Analytics.Generators.csproj
+++ b/Engine/Source/Programs/Shared/EpicGames.Analytics.Generators/EpicGames.Analytics.Generators.csproj
@@ -1,34 +1,18 @@
 ﻿<Project Sdk="Microsoft.NET.Sdk">
-  <Import Project="..\UnrealEngine.csproj.props" />
-
   <PropertyGroup>
     <TargetFramework>netstandard2.0</TargetFramework>
     <LangVersion>latest</LangVersion>
     <IsRoslynComponent>true</IsRoslynComponent>
     <EnforceExtendedAnalyzerRules>true</EnforceExtendedAnalyzerRules>
-    <Configuration Condition=" '$(Configuration)' == '' ">Development</Configuration>
-    <Configurations>Debug;Release;Development;Analyze</Configurations>
   </PropertyGroup>
 
-  <PropertyGroup Condition=" '$(Configuration)|$(Platform)' == 'Debug|AnyCPU' ">
-    <DocumentationFile></DocumentationFile>
-  </PropertyGroup>
   <PropertyGroup Condition="'$(Configuration)|$(Platform)' == 'Development|AnyCPU'">
-    <DocumentationFile></DocumentationFile>
     <Optimize>false</Optimize>
   </PropertyGroup>
-  <PropertyGroup Condition="'$(Configuration)|$(Platform)' == 'Release|AnyCPU'">
-    <DocumentationFile></DocumentationFile>
-    <Optimize>true</Optimize>
-  </PropertyGroup>
-  <PropertyGroup Condition=" '$(Configuration)|$(Platform)' == 'Analyze|AnyCPU' ">
-    <DocumentationFile></DocumentationFile>
-  </PropertyGroup>
 
   <ItemGroup>
     <PackageReference Include="Microsoft.CodeAnalysis.Analyzers" Version="3.3.4" PrivateAssets="all" />
     <PackageReference Include="Microsoft.CodeAnalysis.Common" Version="4.8.0" PrivateAssets="all" />
     <PackageReference Include="Microsoft.CodeAnalysis.CSharp" Version="4.8.0" PrivateAssets="all" />
   </ItemGroup>
-
-</Project>
\ No newline at end of file
+</Project>
diff --git a/Engine/Source/Programs/Shared/EpicGames.Analytics.Telemetry/EpicGames.Analytics.Telemetry.csproj b/Engine/Source/Programs/Shared/EpicGames.Analytics.Telemetry/EpicGames.Analytics.Telemetry.csproj
--- a/Engine/Source/Programs/Shared/EpicGames.Analytics.Telemetry/EpicGames.Analytics.Telemetry.csproj
+++ b/Engine/Source/Programs/Shared/EpicGames.Analytics.Telemetry/EpicGames.Analytics.Telemetry.csproj
@@ -1,39 +1,12 @@
 ﻿<Project Sdk="Microsoft.NET.Sdk">
-  <Import Project="..\UnrealEngine.csproj.props" />
-
   <PropertyGroup>
     <TargetFramework>net8.0</TargetFramework>
-    <Configuration Condition=" '$(Configuration)' == '' ">Development</Configuration>
-    <OutputType>Library</OutputType>
     <RootNamespace>EpicGames.Analytics</RootNamespace>
-    <GenerateAssemblyInfo>false</GenerateAssemblyInfo>
-    <GenerateTargetFrameworkAttribute>false</GenerateTargetFrameworkAttribute>
-    <Configurations>Debug;Release;Development;Analyze</Configurations>
-    <Nullable>enable</Nullable>
     <ProduceReferenceAssemblyInOutDir>true</ProduceReferenceAssemblyInOutDir>
     <AllowUnsafeBlocks>true</AllowUnsafeBlocks>
-    <RunAnalyzersDuringBuild>false</RunAnalyzersDuringBuild>
-  </PropertyGroup>
-  <PropertyGroup Condition=" '$(Configuration)|$(Platform)' == 'Debug|AnyCPU' ">
-    <DocumentationFile></DocumentationFile>
   </PropertyGroup>
+
   <PropertyGroup Condition="'$(Configuration)|$(Platform)' == 'Development|AnyCPU'">
-    <DocumentationFile></DocumentationFile>
     <Optimize>false</Optimize>
   </PropertyGroup>
-  <PropertyGroup Condition="'$(Configuration)|$(Platform)' == 'Release|AnyCPU'">
-    <DocumentationFile></DocumentationFile>
-    <Optimize>true</Optimize>
-  </PropertyGroup>
-  <PropertyGroup Condition=" '$(Configuration)|$(Platform)' == 'Analyze|AnyCPU' ">
-    <DocumentationFile></DocumentationFile>
-    <EnforceCodeStyleInBuild>True</EnforceCodeStyleInBuild>
-    <GenerateDocumentationFile>True</GenerateDocumentationFile>
-    <RunAnalyzersDuringBuild>True</RunAnalyzersDuringBuild>
-  </PropertyGroup>
-  <ItemGroup>
-    <Compile Include="..\..\Shared\MetaData.cs">
-      <Link>Properties\MetaData.cs</Link>
-    </Compile>
-  </ItemGroup>
 </Project>
diff --git a/Engine/Source/Programs/Shared/EpicGames.AspNet/EpicGames.AspNet.csproj b/Engine/Source/Programs/Shared/EpicGames.AspNet/EpicGames.AspNet.csproj
--- a/Engine/Source/Programs/Shared/EpicGames.AspNet/EpicGames.AspNet.csproj
+++ b/Engine/Source/Programs/Shared/EpicGames.AspNet/EpicGames.AspNet.csproj
@@ -1,49 +1,19 @@
-<Project Sdk="Microsoft.NET.Sdk.Web">
-  <Import Project="..\UnrealEngine.csproj.props" />
-
+﻿<Project Sdk="Microsoft.NET.Sdk.Web">
   <PropertyGroup>
     <TargetFramework>net8.0</TargetFramework>
-    <Configuration Condition=" '$(Configuration)' == '' ">Development</Configuration>
+    <!-- This needs to be explictly declared as a library again to build. -->
     <OutputType>Library</OutputType>
-    <RootNamespace>EpicGames.AspNet</RootNamespace>
-    <GenerateAssemblyInfo>false</GenerateAssemblyInfo>
-    <GenerateTargetFrameworkAttribute>false</GenerateTargetFrameworkAttribute>
-    <Configurations>Debug;Release;Development;Analyze</Configurations>
-    <Nullable>enable</Nullable>
     <DebugType>pdbonly</DebugType>
-  </PropertyGroup>
-  <PropertyGroup Condition=" '$(Configuration)|$(Platform)' == 'Debug|AnyCPU' ">
-    <DocumentationFile></DocumentationFile>
-    <NoWarn>$(NoWarn);1701;1702;1591</NoWarn>
-  </PropertyGroup>
-  <PropertyGroup Condition="'$(Configuration)|$(Platform)' == 'Development|AnyCPU'">
-    <DocumentationFile></DocumentationFile>
-    <Optimize>true</Optimize>
-    <NoWarn>$(NoWarn);1701;1702;1591</NoWarn>
-  </PropertyGroup>
-  <PropertyGroup Condition="'$(Configuration)|$(Platform)' == 'Release|AnyCPU'">
-    <DocumentationFile></DocumentationFile>
-    <Optimize>true</Optimize>
-    <NoWarn>$(NoWarn);1701;1702;1591</NoWarn>
-  </PropertyGroup>
-  <PropertyGroup Condition=" '$(Configuration)|$(Platform)' == 'Analyze|AnyCPU' ">
-    <RunAnalyzersDuringBuild>True</RunAnalyzersDuringBuild>
-    <TreatWarningsAsErrors>False</TreatWarningsAsErrors>
-    <EnforceCodeStyleInBuild>True</EnforceCodeStyleInBuild>
     <NoWarn>$(NoWarn);1701;1702;1591</NoWarn>
   </PropertyGroup>
 
-  <ItemGroup>
-    <Compile Include="..\MetaData.cs">
-      <Link>Properties\MetaData.cs</Link>
-    </Compile>
-  </ItemGroup>
   <ItemGroup>
     <ProjectReference Include="..\..\Shared\EpicGames.Core\EpicGames.Core.csproj" />
     <ProjectReference Include="..\EpicGames.Serialization\EpicGames.Serialization.csproj" />
   </ItemGroup>
+
   <ItemGroup>
     <PackageReference Include="Microsoft.CSharp" Version="4.7.0" />
     <PackageReference Include="System.Text.Json" Version="8.0.6" />
   </ItemGroup>
-</Project>
\ No newline at end of file
+</Project>
diff --git a/Engine/Source/Programs/Shared/EpicGames.Benchmarks/EpicGames.Benchmarks.csproj b/Engine/Source/Programs/Shared/EpicGames.Benchmarks/EpicGames.Benchmarks.csproj
--- a/Engine/Source/Programs/Shared/EpicGames.Benchmarks/EpicGames.Benchmarks.csproj
+++ b/Engine/Source/Programs/Shared/EpicGames.Benchmarks/EpicGames.Benchmarks.csproj
@@ -1,6 +1,4 @@
 ﻿<Project Sdk="Microsoft.NET.Sdk">
-  <Import Project="..\UnrealEngine.csproj.props" />
-
   <PropertyGroup>
     <TargetFramework>net8.0</TargetFramework>
     <OutputType>Exe</OutputType>
diff --git a/Engine/Source/Programs/Shared/EpicGames.Build.Tests/EpicGames.Build.Tests.csproj b/Engine/Source/Programs/Shared/EpicGames.Build.Tests/EpicGames.Build.Tests.csproj
--- a/Engine/Source/Programs/Shared/EpicGames.Build.Tests/EpicGames.Build.Tests.csproj
+++ b/Engine/Source/Programs/Shared/EpicGames.Build.Tests/EpicGames.Build.Tests.csproj
@@ -1,30 +1,8 @@
-﻿<Project Sdk="Microsoft.NET.Sdk">
-  <Import Project="..\UnrealEngine.csproj.props" />
-  
+﻿<Project Sdk="Microsoft.NET.Sdk">  
   <PropertyGroup>
     <TargetFramework>net8.0</TargetFramework>
-    <!-- Configurations are required to be in the project itself despite being declared in UnrealEngine.csproj.props due to
-      a reported Mac-specific issue with project generation. -->
-    <Configuration Condition=" '$(Configuration)' == '' ">Development</Configuration>
-    <Configurations>Debug;Release;Development;Analyze</Configurations>
-    <OutputType>Library</OutputType>
-    <Nullable>enable</Nullable>
-    <IsPackable>false</IsPackable>
-    <NoWarn>$(NoWarn);1591</NoWarn>
   </PropertyGroup>
 
-  <ItemGroup>
-    <PackageReference Include="Microsoft.NET.Test.Sdk" Version="17.14.1" />
-    <PackageReference Include="MSTest.TestAdapter" Version="3.6.4">
-      <TreatAsUsed>true</TreatAsUsed>
-    </PackageReference>
-    <PackageReference Include="MSTest.TestFramework" Version="3.6.4" />
-    <PackageReference Include="coverlet.collector" Version="6.0.4">
-      <PrivateAssets>all</PrivateAssets>
-      <IncludeAssets>runtime; build; native; contentfiles; analyzers; buildtransitive</IncludeAssets>
-    </PackageReference>
-  </ItemGroup>
-
   <ItemGroup>
     <ProjectReference Include="..\EpicGames.Build\EpicGames.Build.csproj" />
   </ItemGroup>
diff --git a/Engine/Source/Programs/Shared/EpicGames.Build/EpicGames.Build.csproj b/Engine/Source/Programs/Shared/EpicGames.Build/EpicGames.Build.csproj
--- a/Engine/Source/Programs/Shared/EpicGames.Build/EpicGames.Build.csproj
+++ b/Engine/Source/Programs/Shared/EpicGames.Build/EpicGames.Build.csproj
@@ -1,46 +1,21 @@
 ﻿<Project Sdk="Microsoft.NET.Sdk">
-  <Import Project="..\UnrealEngine.csproj.props" />
-  
   <PropertyGroup>
     <TargetFramework>net8.0</TargetFramework>
-    <Configuration Condition=" '$(Configuration)' == '' ">Development</Configuration>
-    <OutputType>Library</OutputType>
     <RootNamespace>UnrealBuildBase</RootNamespace>
-    <GenerateAssemblyInfo>false</GenerateAssemblyInfo>
-    <GenerateTargetFrameworkAttribute>false</GenerateTargetFrameworkAttribute>
-    <Configurations>Debug;Release;Development;Analyze</Configurations>
-    <Nullable>enable</Nullable>
     <ProduceReferenceAssemblyInOutDir>true</ProduceReferenceAssemblyInOutDir>
     <AllowUnsafeBlocks>true</AllowUnsafeBlocks>
-    <RunAnalyzersDuringBuild>false</RunAnalyzersDuringBuild>
-  </PropertyGroup>
-  <PropertyGroup Condition=" '$(Configuration)|$(Platform)' == 'Debug|AnyCPU' ">
-    <DocumentationFile></DocumentationFile>
-  </PropertyGroup>
-  <PropertyGroup Condition="'$(Configuration)|$(Platform)' == 'Development|AnyCPU'">
-    <DocumentationFile></DocumentationFile>
-    <Optimize>true</Optimize>
-  </PropertyGroup>
-  <PropertyGroup Condition="'$(Configuration)|$(Platform)' == 'Release|AnyCPU'">
-    <DocumentationFile></DocumentationFile>
-    <Optimize>true</Optimize>
   </PropertyGroup>
+
   <PropertyGroup Condition=" '$(Configuration)|$(Platform)' == 'Analyze|AnyCPU' ">
-    <DocumentationFile></DocumentationFile>
-    <RunAnalyzersDuringBuild>True</RunAnalyzersDuringBuild>
-    <GenerateDocumentationFile>True</GenerateDocumentationFile>
     <NoWarn>$(NoWarn);1701;1702;1591</NoWarn>
   </PropertyGroup>
-  <ItemGroup>
-    <Compile Include="..\MetaData.cs">
-      <Link>Properties\MetaData.cs</Link>
-    </Compile>
-  </ItemGroup>
+
   <ItemGroup>
     <ProjectReference Include="..\..\Shared\EpicGames.Core\EpicGames.Core.csproj" />
     <ProjectReference Include="..\EpicGames.IoHash\EpicGames.IoHash.csproj" />
     <ProjectReference Include="..\EpicGames.MsBuild\EpicGames.MsBuild.csproj" />
   </ItemGroup>
+
   <ItemGroup>
     <PackageReference Include="Microsoft.CSharp" Version="4.7.0" />
     <PackageReference Include="Microsoft.Extensions.FileSystemGlobbing" Version="8.0.0" />
diff --git a/Engine/Source/Programs/Shared/EpicGames.BuildGraph.Tests/EpicGames.BuildGraph.Tests.csproj b/Engine/Source/Programs/Shared/EpicGames.BuildGraph.Tests/EpicGames.BuildGraph.Tests.csproj
--- a/Engine/Source/Programs/Shared/EpicGames.BuildGraph.Tests/EpicGames.BuildGraph.Tests.csproj
+++ b/Engine/Source/Programs/Shared/EpicGames.BuildGraph.Tests/EpicGames.BuildGraph.Tests.csproj
@@ -1,59 +1,10 @@
 ﻿<Project Sdk="Microsoft.NET.Sdk">
-  <Import Project="..\UnrealEngine.csproj.props" />
-
   <PropertyGroup>
     <TargetFramework>net8.0</TargetFramework>
-    <Configuration Condition=" '$(Configuration)' == '' ">Development</Configuration>
-    <OutputType>Library</OutputType>
-    <Deterministic>true</Deterministic>
-    <RootNamespace>EpicGames.BuildGraph.Tests</RootNamespace>
-    <GenerateAssemblyCompanyAttribute>false</GenerateAssemblyCompanyAttribute>
-    <GenerateAssemblyInformationalVersionAttribute>false</GenerateAssemblyInformationalVersionAttribute>
-    <GenerateAssemblyProductAttribute>false</GenerateAssemblyProductAttribute>
-    <GenerateAssemblyVersionAttribute>false</GenerateAssemblyVersionAttribute>
-    <GenerateAssemblyFileVersionAttribute>false</GenerateAssemblyFileVersionAttribute>
-    <Configurations>Debug;Release;Development;Analyze</Configurations>
-    <IsPackable>false</IsPackable>
-    <Nullable>enable</Nullable>
-  </PropertyGroup>
-
-  <PropertyGroup Condition=" '$(Configuration)|$(Platform)' == 'Analyze|AnyCPU' ">
-    <RunAnalyzersDuringBuild>True</RunAnalyzersDuringBuild>
-    <TreatWarningsAsErrors>False</TreatWarningsAsErrors>
-    <EnforceCodeStyleInBuild>True</EnforceCodeStyleInBuild>
-    <GenerateDocumentationFile>True</GenerateDocumentationFile>
-    <DocumentationFile></DocumentationFile>
-    <NoWarn>$(NoWarn);1701;1702;1591</NoWarn>
-  </PropertyGroup>
-
-  <PropertyGroup Condition="'$(Configuration)|$(Platform)'=='Debug|AnyCPU'">
-    <NoWarn>$(NoWarn);1701;1702;1591</NoWarn>
-  </PropertyGroup>
-
-  <PropertyGroup Condition="'$(Configuration)|$(Platform)'=='Release|AnyCPU'">
-    <NoWarn>$(NoWarn);1701;1702;1591</NoWarn>
-  </PropertyGroup>
-
-  <PropertyGroup Condition="'$(Configuration)|$(Platform)'=='Development|AnyCPU'">
-    <NoWarn>$(NoWarn);1701;1702;1591</NoWarn>
+    <NoWarn>$(NoWarn);1701;1702</NoWarn>
   </PropertyGroup>
 
   <ItemGroup>
-    <Compile Include="..\..\Shared\MetaData.cs">
-      <Link>Properties\MetaData.cs</Link>
-    </Compile>
-  </ItemGroup>
-
-  <ItemGroup>
-    <PackageReference Include="Microsoft.NET.Test.Sdk" Version="17.14.1" />
-    <PackageReference Include="MSTest.TestAdapter" Version="3.6.4">
-      <TreatAsUsed>true</TreatAsUsed>
-    </PackageReference>
-    <PackageReference Include="MSTest.TestFramework" Version="3.6.4" />
-    <PackageReference Include="coverlet.collector" Version="6.0.4">
-      <PrivateAssets>all</PrivateAssets>
-      <IncludeAssets>runtime; build; native; contentfiles; analyzers; buildtransitive</IncludeAssets>
-    </PackageReference>
     <PackageReference Include="System.Text.Json" Version="8.0.6" />
   </ItemGroup>
 
diff --git a/Engine/Source/Programs/Shared/EpicGames.BuildGraph/EpicGames.BuildGraph.csproj b/Engine/Source/Programs/Shared/EpicGames.BuildGraph/EpicGames.BuildGraph.csproj
--- a/Engine/Source/Programs/Shared/EpicGames.BuildGraph/EpicGames.BuildGraph.csproj
+++ b/Engine/Source/Programs/Shared/EpicGames.BuildGraph/EpicGames.BuildGraph.csproj
@@ -1,46 +1,16 @@
 ﻿<Project Sdk="Microsoft.NET.Sdk">
-  <Import Project="..\UnrealEngine.csproj.props" />
-
   <PropertyGroup>
     <TargetFramework>net8.0</TargetFramework>
-    <Configuration Condition=" '$(Configuration)' == '' ">Development</Configuration>
-    <OutputType>Library</OutputType>
-    <RootNamespace>EpicGames.BuildGraph</RootNamespace>
-    <GenerateAssemblyInfo>false</GenerateAssemblyInfo>
-    <GenerateTargetFrameworkAttribute>false</GenerateTargetFrameworkAttribute>
-    <Configurations>Debug;Release;Development;Analyze</Configurations>
-    <Nullable>enable</Nullable>
     <DebugType>pdbonly</DebugType>
     <DebugType Condition="'$([System.Runtime.InteropServices.RuntimeInformation]::IsOSPlatform($([System.Runtime.InteropServices.OSPlatform]::Windows)))' == 'true' And '$([System.Runtime.InteropServices.RuntimeInformation]::OSArchitecture)' == 'Arm64'">portable</DebugType>
     <GenerateDocumentationFile>True</GenerateDocumentationFile>
   </PropertyGroup>
-  <PropertyGroup Condition="'$(Configuration)|$(Platform)' == 'Debug|AnyCPU'">
-    <DocumentationFile></DocumentationFile>
-  </PropertyGroup>
-  <PropertyGroup Condition="'$(Configuration)|$(Platform)' == 'Development|AnyCPU'">
-    <DocumentationFile></DocumentationFile>
-    <Optimize>true</Optimize>
-  </PropertyGroup>
-  <PropertyGroup Condition="'$(Configuration)|$(Platform)' == 'Release|AnyCPU'">
-    <DocumentationFile></DocumentationFile>
-    <Optimize>true</Optimize>
-  </PropertyGroup>
-  <PropertyGroup Condition=" '$(Configuration)|$(Platform)' == 'Analyze|AnyCPU' ">
-    <RunAnalyzersDuringBuild>True</RunAnalyzersDuringBuild>
-    <TreatWarningsAsErrors>false</TreatWarningsAsErrors>
-    <GenerateDocumentationFile>True</GenerateDocumentationFile>
-    <DocumentationFile></DocumentationFile>
-  </PropertyGroup>
 
-  <ItemGroup>
-    <Compile Include="..\MetaData.cs">
-      <Link>Properties\MetaData.cs</Link>
-    </Compile>
-  </ItemGroup>
   <ItemGroup>
     <ProjectReference Include="..\..\Shared\EpicGames.Core\EpicGames.Core.csproj" />
   </ItemGroup>
+
   <ItemGroup>
     <ProjectReference Include="..\EpicGames.Serialization\EpicGames.Serialization.csproj" />
   </ItemGroup>
-</Project>
\ No newline at end of file
+</Project>
diff --git a/Engine/Source/Programs/Shared/EpicGames.Core.Tests/EpicGames.Core.Tests.csproj b/Engine/Source/Programs/Shared/EpicGames.Core.Tests/EpicGames.Core.Tests.csproj
--- a/Engine/Source/Programs/Shared/EpicGames.Core.Tests/EpicGames.Core.Tests.csproj
+++ b/Engine/Source/Programs/Shared/EpicGames.Core.Tests/EpicGames.Core.Tests.csproj
@@ -1,64 +1,10 @@
 <Project Sdk="Microsoft.NET.Sdk">
-  <Import Project="..\UnrealEngine.csproj.props" />
-
   <PropertyGroup>
     <TargetFramework>net8.0</TargetFramework>
-    <Configuration Condition=" '$(Configuration)' == '' ">Development</Configuration>
-    <OutputType>Library</OutputType>
-    <Deterministic>true</Deterministic>
-    <RootNamespace>EpicGames.Core.Tests</RootNamespace>
-    <GenerateAssemblyCompanyAttribute>false</GenerateAssemblyCompanyAttribute>
-    <GenerateAssemblyInformationalVersionAttribute>false</GenerateAssemblyInformationalVersionAttribute>
-    <GenerateAssemblyProductAttribute>false</GenerateAssemblyProductAttribute>
-    <GenerateAssemblyVersionAttribute>false</GenerateAssemblyVersionAttribute>
-    <GenerateAssemblyFileVersionAttribute>false</GenerateAssemblyFileVersionAttribute>
-    <Configurations>Debug;Release;Development;Analyze</Configurations>
-    <IsPackable>false</IsPackable>
-    <Nullable>enable</Nullable>
-  </PropertyGroup>
-
-  <PropertyGroup Condition=" '$(Configuration)|$(Platform)' == 'Analyze|AnyCPU' ">
-    <RunAnalyzersDuringBuild>True</RunAnalyzersDuringBuild>
-    <TreatWarningsAsErrors>False</TreatWarningsAsErrors>
-    <EnforceCodeStyleInBuild>True</EnforceCodeStyleInBuild>
-    <GenerateDocumentationFile>True</GenerateDocumentationFile>
-    <DocumentationFile></DocumentationFile>
-    <NoWarn>$(NoWarn);1701;1702;1591</NoWarn>
-  </PropertyGroup>
-
-  <PropertyGroup Condition="'$(Configuration)|$(Platform)'=='Debug|AnyCPU'">
-    <NoWarn>$(NoWarn);1701;1702;1591</NoWarn>
-    <TreatWarningsAsErrors>False</TreatWarningsAsErrors>
-  </PropertyGroup>
-
-  <PropertyGroup Condition="'$(Configuration)|$(Platform)'=='Release|AnyCPU'">
-    <NoWarn>$(NoWarn);1701;1702;1591</NoWarn>
+    <NoWarn>$(NoWarn);1701;1702</NoWarn>
     <TreatWarningsAsErrors>False</TreatWarningsAsErrors>
   </PropertyGroup>
 
-  <PropertyGroup Condition="'$(Configuration)|$(Platform)'=='Development|AnyCPU'">
-    <NoWarn>$(NoWarn);1701;1702;1591</NoWarn>
-    <TreatWarningsAsErrors>False</TreatWarningsAsErrors>
-  </PropertyGroup>
-
-  <ItemGroup>
-    <Compile Include="..\..\Shared\MetaData.cs">
-      <Link>Properties\MetaData.cs</Link>
-    </Compile>
-  </ItemGroup>
-
-  <ItemGroup>
-    <PackageReference Include="Microsoft.NET.Test.Sdk" Version="17.14.1" />
-    <PackageReference Include="MSTest.TestAdapter" Version="3.6.4">
-      <TreatAsUsed>true</TreatAsUsed>
-    </PackageReference>
-    <PackageReference Include="MSTest.TestFramework" Version="3.6.4" />
-    <PackageReference Include="coverlet.collector" Version="6.0.4">
-      <PrivateAssets>all</PrivateAssets>
-      <IncludeAssets>runtime; build; native; contentfiles; analyzers; buildtransitive</IncludeAssets>
-    </PackageReference>
-  </ItemGroup>
-
   <ItemGroup>
     <ProjectReference Include="..\EpicGames.Core\EpicGames.Core.csproj" />
   </ItemGroup>
diff --git a/Engine/Source/Programs/Shared/EpicGames.Core/EpicGames.Core.csproj b/Engine/Source/Programs/Shared/EpicGames.Core/EpicGames.Core.csproj
--- a/Engine/Source/Programs/Shared/EpicGames.Core/EpicGames.Core.csproj
+++ b/Engine/Source/Programs/Shared/EpicGames.Core/EpicGames.Core.csproj
@@ -1,47 +1,15 @@
 ﻿<Project Sdk="Microsoft.NET.Sdk">
-  <Import Project="..\UnrealEngine.csproj.props" />
-  
   <PropertyGroup>
     <TargetFramework>net8.0</TargetFramework>
-    <Configuration Condition=" '$(Configuration)' == '' ">Development</Configuration>
-    <OutputType>Library</OutputType>
-    <Deterministic>true</Deterministic>
-    <RootNamespace>EpicGames.Core</RootNamespace>
-    <GenerateAssemblyInfo>false</GenerateAssemblyInfo>
-    <GenerateTargetFrameworkAttribute>false</GenerateTargetFrameworkAttribute>
-    <Configurations>Debug;Release;Development;Analyze</Configurations>
-    <Nullable>enable</Nullable>
     <AllowUnsafeBlocks>true</AllowUnsafeBlocks>
-    <RunAnalyzersDuringBuild>false</RunAnalyzersDuringBuild>
   </PropertyGroup>
-  <PropertyGroup Condition=" '$(Configuration)|$(Platform)' == 'Debug|AnyCPU' ">
-    <DocumentationFile></DocumentationFile>
-  </PropertyGroup>
-  <PropertyGroup Condition="'$(Configuration)|$(Platform)' == 'Development|AnyCPU'">
-    <DocumentationFile></DocumentationFile>
-    <Optimize>true</Optimize>
-  </PropertyGroup>
-  <PropertyGroup Condition="'$(Configuration)|$(Platform)' == 'Release|AnyCPU'">
-    <DocumentationFile></DocumentationFile>
-    <Optimize>true</Optimize>
-  </PropertyGroup>
-  <PropertyGroup Condition=" '$(Configuration)|$(Platform)' == 'Analyze|AnyCPU' ">
-    <RunAnalyzersDuringBuild>True</RunAnalyzersDuringBuild>
-    <TreatWarningsAsErrors>False</TreatWarningsAsErrors>
-    <EnforceCodeStyleInBuild>True</EnforceCodeStyleInBuild>
-    <GenerateDocumentationFile>True</GenerateDocumentationFile>
-    <DocumentationFile></DocumentationFile>
-  </PropertyGroup>
-  <ItemGroup>
-    <Compile Include="..\..\Shared\MetaData.cs">
-      <Link>Properties\MetaData.cs</Link>
-    </Compile>
-  </ItemGroup>
+
   <ItemGroup>
     <Compile Include="..\..\..\..\Restricted\*\Source\Programs\Shared\EpicGames.Core\ScopeProfiler\*.cs">
       <Link>$([System.Text.RegularExpressions.Regex]::Replace(%(FullPath), '^.+?[\\/]EpicGames.Core[\\/]', ''))</Link>
     </Compile>
   </ItemGroup>
+
   <Choose>
     <When Condition="Exists('..\..\..\..\Restricted\NoRedist\Source\Programs\Shared\EpicGames.Core\ScopeProfiler\')">
       <PropertyGroup>
@@ -49,6 +17,7 @@
       </PropertyGroup>
     </When>
   </Choose>
+
   <ItemGroup>
     <PackageReference Include="JetBrains.Annotations" Version="2024.3.0" />
     <PackageReference Include="Microsoft.Extensions.Configuration" Version="8.0.0" />
@@ -59,4 +28,4 @@
     <PackageReference Include="System.Memory" Version="4.6.3" />
     <PackageReference Include="System.Text.Json" Version="8.0.6" />
   </ItemGroup>
-</Project>
\ No newline at end of file
+</Project>
diff --git a/Engine/Source/Programs/Shared/EpicGames.Horde.Tests/AssemblyInfo.cs b/Engine/Source/Programs/Shared/EpicGames.Horde.Tests/AssemblyInfo.cs
deleted file mode 100644
diff --git a/Engine/Source/Programs/Shared/EpicGames.Horde.Tests/EpicGames.Horde.Tests.csproj b/Engine/Source/Programs/Shared/EpicGames.Horde.Tests/EpicGames.Horde.Tests.csproj
--- a/Engine/Source/Programs/Shared/EpicGames.Horde.Tests/EpicGames.Horde.Tests.csproj
+++ b/Engine/Source/Programs/Shared/EpicGames.Horde.Tests/EpicGames.Horde.Tests.csproj
@@ -1,44 +1,10 @@
 <Project Sdk="Microsoft.NET.Sdk">
-  <Import Project="..\UnrealEngine.csproj.props" />
-
   <PropertyGroup>
     <TargetFramework>net8.0</TargetFramework>
-    <Configurations>Debug;Release;Development;Analyze</Configurations>
-    <Nullable>enable</Nullable>
-    <IsPackable>false</IsPackable>
-  </PropertyGroup>
-
-  <PropertyGroup Condition=" '$(Configuration)|$(Platform)' == 'Analyze|AnyCPU' ">
-    <RunAnalyzersDuringBuild>True</RunAnalyzersDuringBuild>
-    <TreatWarningsAsErrors>False</TreatWarningsAsErrors>
-    <EnforceCodeStyleInBuild>True</EnforceCodeStyleInBuild>
-    <GenerateDocumentationFile>True</GenerateDocumentationFile>
-    <DocumentationFile></DocumentationFile>
-    <NoWarn>$(NoWarn);1701;1702;1591</NoWarn>
-  </PropertyGroup>
-
-  <PropertyGroup Condition="'$(Configuration)|$(Platform)'=='Debug|AnyCPU'">
-    <NoWarn>$(NoWarn);1701;1702;1591</NoWarn>
-  </PropertyGroup>
-
-  <PropertyGroup Condition="'$(Configuration)|$(Platform)'=='Release|AnyCPU'">
-    <NoWarn>$(NoWarn);1701;1702;1591</NoWarn>
-  </PropertyGroup>
-
-  <PropertyGroup Condition="'$(Configuration)|$(Platform)'=='Development|AnyCPU'">
-    <NoWarn>$(NoWarn);1701;1702;1591</NoWarn>
+    <NoWarn>$(NoWarn);1701;1702</NoWarn>
   </PropertyGroup>
 
   <ItemGroup>
-    <PackageReference Include="Microsoft.NET.Test.Sdk" Version="17.14.1" />
-    <PackageReference Include="MSTest.TestAdapter" Version="3.6.4">
-      <TreatAsUsed>true</TreatAsUsed>
-    </PackageReference>
-    <PackageReference Include="MSTest.TestFramework" Version="3.6.4" />
-    <PackageReference Include="coverlet.collector" Version="6.0.4">
-      <PrivateAssets>all</PrivateAssets>
-      <IncludeAssets>runtime; build; native; contentfiles; analyzers; buildtransitive</IncludeAssets>
-    </PackageReference>
     <PackageReference Include="System.Linq.Async" Version="6.0.3" />
   </ItemGroup>
 
diff --git a/Engine/Source/Programs/Shared/EpicGames.Horde.Tests/Properties/AssemblyInfo.cs b/Engine/Source/Programs/Shared/EpicGames.Horde.Tests/Properties/AssemblyInfo.cs
--- a/Engine/Source/Programs/Shared/EpicGames.Horde.Tests/Properties/AssemblyInfo.cs
+++ b/Engine/Source/Programs/Shared/EpicGames.Horde.Tests/Properties/AssemblyInfo.cs
@@ -1,5 +1,5 @@
 // Copyright Epic Games, Inc. All Rights Reserved.
 
-using System.Resources;
+using Microsoft.VisualStudio.TestTools.UnitTesting;
 
-[assembly: NeutralResourcesLanguage("en-US")]
+[assembly: Parallelize(Workers = 0, Scope = ExecutionScope.MethodLevel)]
diff --git a/Engine/Source/Programs/Shared/EpicGames.Horde/EpicGames.Horde.csproj b/Engine/Source/Programs/Shared/EpicGames.Horde/EpicGames.Horde.csproj
--- a/Engine/Source/Programs/Shared/EpicGames.Horde/EpicGames.Horde.csproj
+++ b/Engine/Source/Programs/Shared/EpicGames.Horde/EpicGames.Horde.csproj
@@ -1,42 +1,14 @@
-<Project Sdk="Microsoft.NET.Sdk">
-  <Import Project="..\UnrealEngine.csproj.props" />
-
+﻿<Project Sdk="Microsoft.NET.Sdk">
   <PropertyGroup>
     <TargetFramework>net8.0</TargetFramework>
-    <Configuration Condition=" '$(Configuration)' == '' ">Development</Configuration>
-    <OutputType>Library</OutputType>
-    <Deterministic>true</Deterministic>
-    <RootNamespace>EpicGames.Horde</RootNamespace>
-    <GenerateAssemblyInfo>false</GenerateAssemblyInfo>
-    <GenerateTargetFrameworkAttribute>false</GenerateTargetFrameworkAttribute>
-    <Configurations>Debug;Release;Development;Analyze</Configurations>
-    <Nullable>enable</Nullable>
     <GenerateDocumentationFile>True</GenerateDocumentationFile>
     <AllowUnsafeBlocks>True</AllowUnsafeBlocks>
-    <RunAnalyzersDuringBuild>false</RunAnalyzersDuringBuild>
-  </PropertyGroup>
-  <PropertyGroup Condition=" '$(Configuration)|$(Platform)' == 'Debug|AnyCPU' ">
-    <DocumentationFile></DocumentationFile>
   </PropertyGroup>
+
   <PropertyGroup Condition="'$(Configuration)|$(Platform)' == 'Development|AnyCPU'">
-    <DocumentationFile></DocumentationFile>
     <Optimize>false</Optimize>
   </PropertyGroup>
-  <PropertyGroup Condition="'$(Configuration)|$(Platform)' == 'Release|AnyCPU'">
-    <DocumentationFile></DocumentationFile>
-    <Optimize>true</Optimize>
-  </PropertyGroup>
-  <PropertyGroup Condition=" '$(Configuration)|$(Platform)' == 'Analyze|AnyCPU' ">
-    <DocumentationFile></DocumentationFile>
-    <EnforceCodeStyleInBuild>True</EnforceCodeStyleInBuild>
-    <GenerateDocumentationFile>True</GenerateDocumentationFile>
-    <RunAnalyzersDuringBuild>True</RunAnalyzersDuringBuild>
-  </PropertyGroup>
-  <ItemGroup>
-    <Compile Include="..\..\Shared\MetaData.cs">
-      <Link>Properties\MetaData.cs</Link>
-    </Compile>
-  </ItemGroup>
+
   <ItemGroup>
     <PackageReference Include="BitFaster.Caching" Version="2.4.1" />
     <PackageReference Include="Dapper" Version="2.1.66" />
@@ -58,6 +30,7 @@
     <PackageReference Include="System.Linq.Async" Version="6.0.3" />
     <PackageReference Include="ZstdSharp.Port" Version="0.8.6" />
   </ItemGroup>
+
   <ItemGroup>
     <ProjectReference Include="..\EpicGames.Core\EpicGames.Core.csproj" />
     <ProjectReference Include="..\EpicGames.IoHash\EpicGames.IoHash.csproj" />
@@ -65,8 +38,9 @@
     <ProjectReference Include="..\EpicGames.Oodle\EpicGames.Oodle.csproj" />
     <ProjectReference Include="..\EpicGames.Serialization\EpicGames.Serialization.csproj" />
   </ItemGroup>
+
   <ItemGroup>
     <Protobuf Include="Protos\horde\log_rpc.proto" ProtoRoot="Protos" GrpcServices="Both" Access="Public" />
     <Protobuf Include="Protos\horde\log_rpc_messages.proto" ProtoRoot="Protos" Access="Public" />
   </ItemGroup>
-</Project>
\ No newline at end of file
+</Project>
diff --git a/Engine/Source/Programs/Shared/EpicGames.IoHash.Tests/EpicGames.IoHash.Tests.csproj b/Engine/Source/Programs/Shared/EpicGames.IoHash.Tests/EpicGames.IoHash.Tests.csproj
--- a/Engine/Source/Programs/Shared/EpicGames.IoHash.Tests/EpicGames.IoHash.Tests.csproj
+++ b/Engine/Source/Programs/Shared/EpicGames.IoHash.Tests/EpicGames.IoHash.Tests.csproj
@@ -1,33 +1,8 @@
 <Project Sdk="Microsoft.NET.Sdk">
-  <Import Project="..\UnrealEngine.csproj.props" />
-
   <PropertyGroup>
     <TargetFramework>net8.0</TargetFramework>
-    <Configurations>Debug;Release;Development;Analyze</Configurations>
-    <Nullable>enable</Nullable>
-    <IsPackable>false</IsPackable>
-    <NoWarn>$(NoWarn);1591</NoWarn>
-  </PropertyGroup>
-
-  <PropertyGroup Condition=" '$(Configuration)|$(Platform)' == 'Analyze|AnyCPU' ">
-    <RunAnalyzersDuringBuild>True</RunAnalyzersDuringBuild>
-    <TreatWarningsAsErrors>false</TreatWarningsAsErrors>
-    <GenerateDocumentationFile>True</GenerateDocumentationFile>
-    <DocumentationFile></DocumentationFile>
   </PropertyGroup>
 
-  <ItemGroup>
-    <PackageReference Include="Microsoft.NET.Test.Sdk" Version="17.14.1" />
-    <PackageReference Include="MSTest.TestAdapter" Version="3.6.4">
-      <TreatAsUsed>true</TreatAsUsed>
-    </PackageReference>
-    <PackageReference Include="MSTest.TestFramework" Version="3.6.4" />
-    <PackageReference Include="coverlet.collector" Version="6.0.4">
-      <PrivateAssets>all</PrivateAssets>
-      <IncludeAssets>runtime; build; native; contentfiles; analyzers; buildtransitive</IncludeAssets>
-    </PackageReference>
-  </ItemGroup>
-
   <ItemGroup>
     <ProjectReference Include="..\EpicGames.IoHash\EpicGames.IoHash.csproj" />
   </ItemGroup>
diff --git a/Engine/Source/Programs/Shared/EpicGames.IoHash/EpicGames.IoHash.csproj b/Engine/Source/Programs/Shared/EpicGames.IoHash/EpicGames.IoHash.csproj
--- a/Engine/Source/Programs/Shared/EpicGames.IoHash/EpicGames.IoHash.csproj
+++ b/Engine/Source/Programs/Shared/EpicGames.IoHash/EpicGames.IoHash.csproj
@@ -1,46 +1,16 @@
 ﻿<Project Sdk="Microsoft.NET.Sdk">
-  <Import Project="..\UnrealEngine.csproj.props" />
-  
   <PropertyGroup>
     <TargetFramework>net8.0</TargetFramework>
-    <Configuration Condition=" '$(Configuration)' == '' ">Development</Configuration>
-    <OutputType>Library</OutputType>
-    <Deterministic>true</Deterministic>
     <RootNamespace>EpicGames.Core</RootNamespace>
-    <GenerateAssemblyInfo>false</GenerateAssemblyInfo>
-    <GenerateTargetFrameworkAttribute>false</GenerateTargetFrameworkAttribute>
-    <Configurations>Debug;Release;Development;Analyze</Configurations>
-    <Nullable>enable</Nullable>
   </PropertyGroup>
-  <PropertyGroup Condition=" '$(Configuration)|$(Platform)' == 'Debug|AnyCPU' ">
-    <DocumentationFile></DocumentationFile>
-  </PropertyGroup>
-  <PropertyGroup Condition="'$(Configuration)|$(Platform)' == 'Development|AnyCPU'">
-    <DocumentationFile></DocumentationFile>
-    <Optimize>true</Optimize>
-  </PropertyGroup>
-  <PropertyGroup Condition="'$(Configuration)|$(Platform)' == 'Release|AnyCPU'">
-    <DocumentationFile></DocumentationFile>
-    <Optimize>true</Optimize>
-  </PropertyGroup>
-  <PropertyGroup Condition=" '$(Configuration)|$(Platform)' == 'Analyze|AnyCPU' ">
-    <RunAnalyzersDuringBuild>True</RunAnalyzersDuringBuild>
-    <TreatWarningsAsErrors>False</TreatWarningsAsErrors>
-    <EnforceCodeStyleInBuild>True</EnforceCodeStyleInBuild>
-    <GenerateDocumentationFile>True</GenerateDocumentationFile>
-    <DocumentationFile></DocumentationFile>
-  </PropertyGroup>
-  <ItemGroup>
-    <Compile Include="..\..\Shared\MetaData.cs">
-      <Link>Properties\MetaData.cs</Link>
-    </Compile>
-  </ItemGroup>
+
   <ItemGroup>
     <!-- Blake3 2.0.0 has a massive performance regression on Windows x64 - takes 2.25x as long for large (1MB) inputs -->
     <PackageReference Include="Blake3" Version="1.1.0" />
     <PackageReference Include="System.Memory" Version="4.6.3" />
   </ItemGroup>
+
   <ItemGroup>
     <ProjectReference Include="..\EpicGames.Core\EpicGames.Core.csproj" />
   </ItemGroup>
-</Project>
\ No newline at end of file
+</Project>
diff --git a/Engine/Source/Programs/Shared/EpicGames.LogMatchers.Tests/EpicGames.LogMatchers.Tests.csproj b/Engine/Source/Programs/Shared/EpicGames.LogMatchers.Tests/EpicGames.LogMatchers.Tests.csproj
--- a/Engine/Source/Programs/Shared/EpicGames.LogMatchers.Tests/EpicGames.LogMatchers.Tests.csproj
+++ b/Engine/Source/Programs/Shared/EpicGames.LogMatchers.Tests/EpicGames.LogMatchers.Tests.csproj
@@ -1,53 +1,9 @@
 <Project Sdk="Microsoft.NET.Sdk">
-  <Import Project="..\UnrealEngine.csproj.props" />
-  
   <PropertyGroup>
     <TargetFramework>net8.0</TargetFramework>
-    <Configuration Condition=" '$(Configuration)' == '' ">Development</Configuration>
-    <RootNamespace>EpicGames.LogMatchers.Tests</RootNamespace>
-    <Configurations>Debug;Release;Development;Analyze</Configurations>
-    <Nullable>enable</Nullable>
-    <RunAnalyzersDuringBuild>false</RunAnalyzersDuringBuild>
-
-    <IsPackable>false</IsPackable>
-    <IsTestProject>true</IsTestProject>
-    <NoWarn>$(NoWarn);1701;1702;1591</NoWarn>
-  </PropertyGroup>
-
-  <PropertyGroup Condition=" '$(Configuration)|$(Platform)' == 'Debug|AnyCPU' ">
-    <DocumentationFile></DocumentationFile>
-  </PropertyGroup>
-
-  <PropertyGroup Condition="'$(Configuration)|$(Platform)' == 'Development|AnyCPU'">
-    <DocumentationFile></DocumentationFile>
-    <Optimize>true</Optimize>
-  </PropertyGroup>
-
-  <PropertyGroup Condition="'$(Configuration)|$(Platform)' == 'Release|AnyCPU'">
-    <DocumentationFile></DocumentationFile>
-    <Optimize>true</Optimize>
-  </PropertyGroup>
-
-  <PropertyGroup Condition=" '$(Configuration)|$(Platform)' == 'Analyze|AnyCPU' ">
-    <RunAnalyzersDuringBuild>True</RunAnalyzersDuringBuild>
-    <TreatWarningsAsErrors>False</TreatWarningsAsErrors>
-    <EnforceCodeStyleInBuild>True</EnforceCodeStyleInBuild>
-    <GenerateDocumentationFile>True</GenerateDocumentationFile>
-    <DocumentationFile></DocumentationFile>
+    <NoWarn>$(NoWarn);1701;1702</NoWarn>
   </PropertyGroup>
 
-  <ItemGroup>
-    <PackageReference Include="Microsoft.NET.Test.Sdk" Version="17.14.1" />
-    <PackageReference Include="MSTest.TestAdapter" Version="3.6.4">
-      <TreatAsUsed>true</TreatAsUsed>
-    </PackageReference>
-    <PackageReference Include="MSTest.TestFramework" Version="3.6.4" />
-    <PackageReference Include="coverlet.collector" Version="6.0.4">
-      <PrivateAssets>all</PrivateAssets>
-      <IncludeAssets>runtime; build; native; contentfiles; analyzers; buildtransitive</IncludeAssets>
-    </PackageReference>
-  </ItemGroup>
-
   <ItemGroup>
     <ProjectReference Include="..\EpicGames.Core\EpicGames.Core.csproj" />
     <ProjectReference Include="..\EpicGames.LogMatchers\EpicGames.LogMatchers.csproj" />
diff --git a/Engine/Source/Programs/Shared/EpicGames.LogMatchers/EpicGames.LogMatchers.csproj b/Engine/Source/Programs/Shared/EpicGames.LogMatchers/EpicGames.LogMatchers.csproj
--- a/Engine/Source/Programs/Shared/EpicGames.LogMatchers/EpicGames.LogMatchers.csproj
+++ b/Engine/Source/Programs/Shared/EpicGames.LogMatchers/EpicGames.LogMatchers.csproj
@@ -1,47 +1,8 @@
-<Project Sdk="Microsoft.NET.Sdk">
-  <Import Project="..\UnrealEngine.csproj.props" />
-  
+﻿<Project Sdk="Microsoft.NET.Sdk">
   <PropertyGroup>
     <TargetFramework>net8.0</TargetFramework>
-    <Configuration Condition=" '$(Configuration)' == '' ">Development</Configuration>
-    <OutputType>Library</OutputType>
-    <Deterministic>true</Deterministic>
-    <RootNamespace>EpicGames.LogMatchers</RootNamespace>
-    <GenerateAssemblyInfo>false</GenerateAssemblyInfo>
-    <GenerateTargetFrameworkAttribute>false</GenerateTargetFrameworkAttribute>
-    <Configurations>Debug;Release;Development;Analyze</Configurations>
-    <Nullable>enable</Nullable>
-    <RunAnalyzersDuringBuild>false</RunAnalyzersDuringBuild>
   </PropertyGroup>
 
-  <PropertyGroup Condition=" '$(Configuration)|$(Platform)' == 'Debug|AnyCPU' ">
-    <DocumentationFile></DocumentationFile>
-  </PropertyGroup>
-
-  <PropertyGroup Condition="'$(Configuration)|$(Platform)' == 'Development|AnyCPU'">
-    <DocumentationFile></DocumentationFile>
-    <Optimize>true</Optimize>
-  </PropertyGroup>
-
-  <PropertyGroup Condition="'$(Configuration)|$(Platform)' == 'Release|AnyCPU'">
-    <DocumentationFile></DocumentationFile>
-    <Optimize>true</Optimize>
-  </PropertyGroup>
-
-  <PropertyGroup Condition=" '$(Configuration)|$(Platform)' == 'Analyze|AnyCPU' ">
-    <RunAnalyzersDuringBuild>True</RunAnalyzersDuringBuild>
-    <TreatWarningsAsErrors>False</TreatWarningsAsErrors>
-    <EnforceCodeStyleInBuild>True</EnforceCodeStyleInBuild>
-    <GenerateDocumentationFile>True</GenerateDocumentationFile>
-    <DocumentationFile></DocumentationFile>
-  </PropertyGroup>
-
-  <ItemGroup>
-    <Compile Include="..\..\Shared\MetaData.cs">
-      <Link>Properties\MetaData.cs</Link>
-    </Compile>
-  </ItemGroup>
-
   <ItemGroup>
     <PackageReference Include="Microsoft.Extensions.Logging" Version="8.0.1" />
   </ItemGroup>
diff --git a/Engine/Source/Programs/Shared/EpicGames.MongoDB/EpicGames.MongoDB.csproj b/Engine/Source/Programs/Shared/EpicGames.MongoDB/EpicGames.MongoDB.csproj
--- a/Engine/Source/Programs/Shared/EpicGames.MongoDB/EpicGames.MongoDB.csproj
+++ b/Engine/Source/Programs/Shared/EpicGames.MongoDB/EpicGames.MongoDB.csproj
@@ -1,38 +1,8 @@
-<Project Sdk="Microsoft.NET.Sdk">
-  <Import Project="..\UnrealEngine.csproj.props" />
-
+﻿<Project Sdk="Microsoft.NET.Sdk">
   <PropertyGroup>
     <TargetFramework>net8.0</TargetFramework>
-    <Configuration Condition=" '$(Configuration)' == '' ">Development</Configuration>
-    <OutputType>Library</OutputType>
-    <RootNamespace>EpicGames.Core</RootNamespace>
-    <GenerateAssemblyInfo>false</GenerateAssemblyInfo>
-    <GenerateTargetFrameworkAttribute>false</GenerateTargetFrameworkAttribute>
-    <Configurations>Debug;Release;Development;Analyze</Configurations>
-    <Nullable>enable</Nullable>
-  </PropertyGroup>
-  <PropertyGroup Condition=" '$(Configuration)|$(Platform)' == 'Debug|AnyCPU' ">
-    <DocumentationFile></DocumentationFile>
-  </PropertyGroup>
-  <PropertyGroup Condition="'$(Configuration)|$(Platform)' == 'Development|AnyCPU'">
-    <DocumentationFile></DocumentationFile>
-    <Optimize>true</Optimize>
-  </PropertyGroup>
-  <PropertyGroup Condition="'$(Configuration)|$(Platform)' == 'Release|AnyCPU'">
-    <DocumentationFile></DocumentationFile>
-    <Optimize>true</Optimize>
   </PropertyGroup>
-  <PropertyGroup Condition=" '$(Configuration)|$(Platform)' == 'Analyze|AnyCPU' ">
-    <RunAnalyzersDuringBuild>True</RunAnalyzersDuringBuild>
-    <TreatWarningsAsErrors>false</TreatWarningsAsErrors>
-    <GenerateDocumentationFile>True</GenerateDocumentationFile>
-    <DocumentationFile></DocumentationFile>
-  </PropertyGroup>
-  <ItemGroup>
-    <Compile Include="..\..\Shared\MetaData.cs">
-      <Link>Properties\MetaData.cs</Link>
-    </Compile>
-  </ItemGroup>
+
   <ItemGroup>
     <PackageReference Include="MongoDB.Driver" Version="2.19.2" />
   </ItemGroup>
diff --git a/Engine/Source/Programs/Shared/EpicGames.MsBuild/EpicGames.MsBuild.csproj b/Engine/Source/Programs/Shared/EpicGames.MsBuild/EpicGames.MsBuild.csproj
--- a/Engine/Source/Programs/Shared/EpicGames.MsBuild/EpicGames.MsBuild.csproj
+++ b/Engine/Source/Programs/Shared/EpicGames.MsBuild/EpicGames.MsBuild.csproj
@@ -1,39 +1,9 @@
 ﻿<Project Sdk="Microsoft.NET.Sdk">
-  <Import Project="..\UnrealEngine.csproj.props" />
-  
   <PropertyGroup>
     <TargetFramework>net8.0</TargetFramework>
-    <Configuration Condition=" '$(Configuration)' == '' ">Development</Configuration>
-    <OutputType>Library</OutputType>
-    <RootNamespace>EpicGames.MsBuild</RootNamespace>
-    <GenerateAssemblyInfo>false</GenerateAssemblyInfo>
-    <GenerateTargetFrameworkAttribute>false</GenerateTargetFrameworkAttribute>
-    <Configurations>Debug;Release;Development;Analyze</Configurations>
     <ProduceReferenceAssemblyInOutDir>true</ProduceReferenceAssemblyInOutDir>
+    <Nullable>disable</Nullable>
   </PropertyGroup>
-  <PropertyGroup Condition=" '$(Configuration)|$(Platform)' == 'Debug|AnyCPU' ">
-    <DocumentationFile></DocumentationFile>
-  </PropertyGroup>
-  <PropertyGroup Condition="'$(Configuration)|$(Platform)' == 'Development|AnyCPU'">
-    <DocumentationFile></DocumentationFile>
-    <Optimize>true</Optimize>
-  </PropertyGroup>
-  <PropertyGroup Condition="'$(Configuration)|$(Platform)' == 'Release|AnyCPU'">
-    <DocumentationFile></DocumentationFile>
-    <Optimize>true</Optimize>
-  </PropertyGroup>
-  <PropertyGroup Condition=" '$(Configuration)|$(Platform)' == 'Analyze|AnyCPU' ">
-    <RunAnalyzersDuringBuild>True</RunAnalyzersDuringBuild>
-    <TreatWarningsAsErrors>false</TreatWarningsAsErrors>
-    <GenerateDocumentationFile>True</GenerateDocumentationFile>
-    <DocumentationFile></DocumentationFile>
-  </PropertyGroup>
-  
-  <ItemGroup>
-    <Compile Include="..\MetaData.cs">
-      <Link>Properties\MetaData.cs</Link>
-    </Compile>
-  </ItemGroup>
 
   <ItemGroup>
     <PackageReference Include="Microsoft.Build" Version="17.11.48" ExcludeAssets="runtime" />
diff --git a/Engine/Source/Programs/Shared/EpicGames.OIDC/EpicGames.OIDC.csproj b/Engine/Source/Programs/Shared/EpicGames.OIDC/EpicGames.OIDC.csproj
--- a/Engine/Source/Programs/Shared/EpicGames.OIDC/EpicGames.OIDC.csproj
+++ b/Engine/Source/Programs/Shared/EpicGames.OIDC/EpicGames.OIDC.csproj
@@ -1,40 +1,9 @@
-<Project Sdk="Microsoft.NET.Sdk">
-  <Import Project="..\UnrealEngine.csproj.props" />
-  
+﻿<Project Sdk="Microsoft.NET.Sdk">
   <PropertyGroup>
-    <OutputType>Library</OutputType>
-    <RootNamespace>EpicGames.OIDC</RootNamespace>
     <TargetFramework>net8.0</TargetFramework>
-    <Nullable>enable</Nullable>
-    <GenerateAssemblyInfo>false</GenerateAssemblyInfo>
-    <GenerateTargetFrameworkAttribute>false</GenerateTargetFrameworkAttribute>
-    <Configurations>Debug;Release;Development;Analyze</Configurations>
-    <Deterministic>true</Deterministic>
-    <RunAnalyzersDuringBuild>false</RunAnalyzersDuringBuild>
     <EnableConfigurationBindingGenerator>true</EnableConfigurationBindingGenerator>
   </PropertyGroup>
 
-  <PropertyGroup Condition=" '$(Configuration)|$(Platform)' == 'Debug|AnyCPU' ">
-  </PropertyGroup>
-  <PropertyGroup Condition="'$(Configuration)|$(Platform)' == 'Development|AnyCPU'">
-    <Optimize>true</Optimize>
-  </PropertyGroup>
-  <PropertyGroup Condition="'$(Configuration)|$(Platform)' == 'Release|AnyCPU'">
-    <Optimize>true</Optimize>
-  </PropertyGroup>
-  <PropertyGroup Condition=" '$(Configuration)|$(Platform)' == 'Analyze|AnyCPU' ">
-    <RunAnalyzersDuringBuild>True</RunAnalyzersDuringBuild>
-    <TreatWarningsAsErrors>false</TreatWarningsAsErrors>
-    <GenerateDocumentationFile>True</GenerateDocumentationFile>
-    <DocumentationFile></DocumentationFile>
-  </PropertyGroup>
-
-  <ItemGroup>
-    <Compile Include="..\MetaData.cs">
-      <Link>Properties\MetaData.cs</Link>
-    </Compile>
-  </ItemGroup>
-  
   <ItemGroup>
     <PackageReference Include="IdentityModel.OidcClient" Version="6.0.0" />
     <PackageReference Include="Microsoft.Extensions.Configuration.Binder" Version="8.0.2" />
diff --git a/Engine/Source/Programs/Shared/EpicGames.Oodle/EpicGames.Oodle.csproj b/Engine/Source/Programs/Shared/EpicGames.Oodle/EpicGames.Oodle.csproj
--- a/Engine/Source/Programs/Shared/EpicGames.Oodle/EpicGames.Oodle.csproj
+++ b/Engine/Source/Programs/Shared/EpicGames.Oodle/EpicGames.Oodle.csproj
@@ -1,41 +1,15 @@
-<Project Sdk="Microsoft.NET.Sdk">
-  <Import Project="..\UnrealEngine.csproj.props" />
-  
+﻿<Project Sdk="Microsoft.NET.Sdk">
   <PropertyGroup>
     <TargetFramework>net8.0</TargetFramework>
-    <Configuration Condition=" '$(Configuration)' == '' ">Development</Configuration>
-    <OutputType>Library</OutputType>
-    <Deterministic>true</Deterministic>
     <RootNamespace>EpicGames.Compression</RootNamespace>
-    <GenerateAssemblyInfo>false</GenerateAssemblyInfo>
-    <GenerateTargetFrameworkAttribute>false</GenerateTargetFrameworkAttribute>
-    <Configurations>Debug;Release;Development;Analyze</Configurations>
-    <Nullable>enable</Nullable>
     <GenerateDocumentationFile>True</GenerateDocumentationFile>
     <AllowUnsafeBlocks>True</AllowUnsafeBlocks>
-    <RunAnalyzersDuringBuild>false</RunAnalyzersDuringBuild>
-  </PropertyGroup>
-  <PropertyGroup Condition=" '$(Configuration)|$(Platform)' == 'Debug|AnyCPU' ">
-    <DocumentationFile></DocumentationFile>
   </PropertyGroup>
+
   <PropertyGroup Condition="'$(Configuration)|$(Platform)' == 'Development|AnyCPU'">
-    <DocumentationFile></DocumentationFile>
     <Optimize>false</Optimize>
   </PropertyGroup>
-  <PropertyGroup Condition="'$(Configuration)|$(Platform)' == 'Release|AnyCPU'">
-    <DocumentationFile></DocumentationFile>
-    <Optimize>true</Optimize>
-  </PropertyGroup>
-  <PropertyGroup Condition=" '$(Configuration)|$(Platform)' == 'Analyze|AnyCPU' ">
-    <DocumentationFile></DocumentationFile>
-    <GenerateDocumentationFile>True</GenerateDocumentationFile>
-    <RunAnalyzersDuringBuild>True</RunAnalyzersDuringBuild>
-  </PropertyGroup>
-  <ItemGroup>
-    <Compile Include="..\..\Shared\MetaData.cs">
-      <Link>Properties\MetaData.cs</Link>
-    </Compile>
-  </ItemGroup>
+
   <ItemGroup>
     <None Condition="'$([System.Runtime.InteropServices.RuntimeInformation]::IsOSPlatform($([System.Runtime.InteropServices.OSPlatform]::Linux)))' == 'true' Or Exists('Sdk\2.9.10\linux\lib\liboo2corelinux64.so.9')" Include="Sdk\2.9.10\linux\lib\liboo2corelinux64.so.9" Link="liboo2corelinux64.so.9" CopyToOutputDirectory="PreserveNewest" Visible="false" />
     <None Condition="'$([System.Runtime.InteropServices.RuntimeInformation]::IsOSPlatform($([System.Runtime.InteropServices.OSPlatform]::Linux)))' == 'true' Or Exists('Sdk\2.9.10\linuxarm\lib\liboo2corelinuxarm32.so.9')" Include="Sdk\2.9.10\linuxarm\lib\liboo2corelinuxarm32.so.9" Link="liboo2corelinuxarm32.so.9" CopyToOutputDirectory="PreserveNewest" Visible="false" />
@@ -45,4 +19,4 @@
     <None Condition="'$([System.Runtime.InteropServices.RuntimeInformation]::IsOSPlatform($([System.Runtime.InteropServices.OSPlatform]::Windows)))' == 'true' Or Exists('Sdk\2.9.10\win\redist\oo2core_9_win32.dll')" Include="Sdk\2.9.10\win\redist\oo2core_9_win32.dll" Link="oo2core_9_win32.dll" CopyToOutputDirectory="PreserveNewest" Visible="false" />
     <None Condition="'$([System.Runtime.InteropServices.RuntimeInformation]::IsOSPlatform($([System.Runtime.InteropServices.OSPlatform]::Windows)))' == 'true' Or Exists('Sdk\2.9.10\winuwp\redist\oo2core_9_winuwparm64.dll')" Include="Sdk\2.9.10\winuwp\redist\oo2core_9_winuwparm64.dll" Link="oo2core_9_winuwparm64.dll" CopyToOutputDirectory="PreserveNewest" Visible="false" />
   </ItemGroup>
-</Project>
\ No newline at end of file
+</Project>
diff --git a/Engine/Source/Programs/Shared/EpicGames.Perforce.Fixture/EpicGames.Perforce.Fixture.csproj b/Engine/Source/Programs/Shared/EpicGames.Perforce.Fixture/EpicGames.Perforce.Fixture.csproj
--- a/Engine/Source/Programs/Shared/EpicGames.Perforce.Fixture/EpicGames.Perforce.Fixture.csproj
+++ b/Engine/Source/Programs/Shared/EpicGames.Perforce.Fixture/EpicGames.Perforce.Fixture.csproj
@@ -1,38 +1,11 @@
 <Project Sdk="Microsoft.NET.Sdk">
-  <Import Project="..\UnrealEngine.csproj.props" />
-
   <PropertyGroup>
     <TargetFramework>net8.0</TargetFramework>
-    <Configuration Condition=" '$(Configuration)' == '' ">Development</Configuration>
-    <OutputType>Library</OutputType>
-    <Deterministic>true</Deterministic>
-    <RootNamespace>EpicGames.Perforce.Fixture</RootNamespace>
-    <GenerateAssemblyCompanyAttribute>false</GenerateAssemblyCompanyAttribute>
-    <GenerateAssemblyInformationalVersionAttribute>false</GenerateAssemblyInformationalVersionAttribute>
-    <GenerateAssemblyProductAttribute>false</GenerateAssemblyProductAttribute>
-    <GenerateAssemblyVersionAttribute>false</GenerateAssemblyVersionAttribute>
-    <GenerateAssemblyFileVersionAttribute>false</GenerateAssemblyFileVersionAttribute>
-    <Configurations>Debug;Release;Development;Analyze</Configurations>
+    <StartupObject>EpicGames.Perforce.Fixture.Program</StartupObject>
     <IsPackable>false</IsPackable>
-    <Nullable>enable</Nullable>
     <NoWarn>$(NoWarn);1591</NoWarn>
-    <StartupObject>EpicGames.Perforce.Fixture.Program</StartupObject>
-  </PropertyGroup>
-
-  <PropertyGroup Condition=" '$(Configuration)|$(Platform)' == 'Analyze|AnyCPU' ">
-    <RunAnalyzersDuringBuild>True</RunAnalyzersDuringBuild>
-    <TreatWarningsAsErrors>false</TreatWarningsAsErrors>
-    <EnforceCodeStyleInBuild>True</EnforceCodeStyleInBuild>
-    <GenerateDocumentationFile>True</GenerateDocumentationFile>
-    <DocumentationFile></DocumentationFile>
   </PropertyGroup>
 
-  <ItemGroup>
-    <Compile Include="..\..\Shared\MetaData.cs">
-      <Link>Properties\MetaData.cs</Link>
-    </Compile>
-  </ItemGroup>
-
   <ItemGroup>
     <EmbeddedResource Include="snapshot.json" />
   </ItemGroup>
diff --git a/Engine/Source/Programs/Shared/EpicGames.Perforce.Managed.Tests/EpicGames.Perforce.Managed.Tests.csproj b/Engine/Source/Programs/Shared/EpicGames.Perforce.Managed.Tests/EpicGames.Perforce.Managed.Tests.csproj
--- a/Engine/Source/Programs/Shared/EpicGames.Perforce.Managed.Tests/EpicGames.Perforce.Managed.Tests.csproj
+++ b/Engine/Source/Programs/Shared/EpicGames.Perforce.Managed.Tests/EpicGames.Perforce.Managed.Tests.csproj
@@ -1,51 +1,7 @@
 <Project Sdk="Microsoft.NET.Sdk">
-  <Import Project="..\UnrealEngine.csproj.props" />
-
   <PropertyGroup>
     <TargetFramework>net8.0</TargetFramework>
-    <Configuration Condition=" '$(Configuration)' == '' ">Development</Configuration>
-    <OutputType>Library</OutputType>
-    <Deterministic>true</Deterministic>
-    <RootNamespace>EpicGames.Perforce.Managed.Tests</RootNamespace>
-    <GenerateAssemblyCompanyAttribute>false</GenerateAssemblyCompanyAttribute>
-    <GenerateAssemblyInformationalVersionAttribute>false</GenerateAssemblyInformationalVersionAttribute>
-    <GenerateAssemblyProductAttribute>false</GenerateAssemblyProductAttribute>
-    <GenerateAssemblyVersionAttribute>false</GenerateAssemblyVersionAttribute>
-    <GenerateAssemblyFileVersionAttribute>false</GenerateAssemblyFileVersionAttribute>
-    <Configurations>Debug;Release;Development</Configurations>
-    <IsPackable>false</IsPackable>
-    <Nullable>enable</Nullable>
-    <Configurations>Debug;Release;Development;Analyze</Configurations>
-    <NoWarn>$(NoWarn);1591</NoWarn>
-  </PropertyGroup>
-  <PropertyGroup Condition=" '$(Configuration)|$(Platform)' == 'Analyze|AnyCPU' ">
-    <RunAnalyzersDuringBuild>True</RunAnalyzersDuringBuild>
-    <TreatWarningsAsErrors>false</TreatWarningsAsErrors>
-    <EnforceCodeStyleInBuild>True</EnforceCodeStyleInBuild>
-    <GenerateDocumentationFile>True</GenerateDocumentationFile>
-    <DocumentationFile></DocumentationFile>
   </PropertyGroup>
-  <ItemGroup>
-    <Compile Include="..\..\Shared\MetaData.cs">
-      <Link>Properties\MetaData.cs</Link>
-    </Compile>
-  </ItemGroup>
-
-  <ItemGroup>
-    <PackageReference Include="Microsoft.NET.Test.Sdk" Version="17.14.1">
-      <TreatAsUsed>true</TreatAsUsed>
-    </PackageReference>
-    <PackageReference Include="MSTest.TestAdapter" Version="3.6.4">
-      <TreatAsUsed>true</TreatAsUsed>
-    </PackageReference>
-    <PackageReference Include="MSTest.TestFramework" Version="3.6.4">
-      <TreatAsUsed>true</TreatAsUsed>
-    </PackageReference>
-    <PackageReference Include="coverlet.collector" Version="6.0.4">
-      <PrivateAssets>all</PrivateAssets>
-      <IncludeAssets>runtime; build; native; contentfiles; analyzers; buildtransitive</IncludeAssets>
-    </PackageReference>
-  </ItemGroup>
 
   <ItemGroup>
     <ProjectReference Include="..\EpicGames.Perforce.Managed\EpicGames.Perforce.Managed.csproj" />
diff --git a/Engine/Source/Programs/Shared/EpicGames.Perforce.Managed/EpicGames.Perforce.Managed.csproj b/Engine/Source/Programs/Shared/EpicGames.Perforce.Managed/EpicGames.Perforce.Managed.csproj
--- a/Engine/Source/Programs/Shared/EpicGames.Perforce.Managed/EpicGames.Perforce.Managed.csproj
+++ b/Engine/Source/Programs/Shared/EpicGames.Perforce.Managed/EpicGames.Perforce.Managed.csproj
@@ -1,48 +1,18 @@
-<Project Sdk="Microsoft.NET.Sdk">
-  <Import Project="..\UnrealEngine.csproj.props" />
-
+﻿<Project Sdk="Microsoft.NET.Sdk">
   <PropertyGroup>
     <TargetFramework>net8.0</TargetFramework>
-    <Nullable>enable</Nullable>
-    <GenerateAssemblyInfo>false</GenerateAssemblyInfo>
-    <GenerateTargetFrameworkAttribute>false</GenerateTargetFrameworkAttribute>
-    <Configurations>Debug;Release;Development;Analyze</Configurations>
-    <DebugType>full</DebugType>
     <DebugType>pdbonly</DebugType>
   </PropertyGroup>
-  <PropertyGroup Condition=" '$(Configuration)|$(Platform)' == 'Debug|AnyCPU' ">
-    <DocumentationFile></DocumentationFile>
-  </PropertyGroup>
-  <PropertyGroup Condition="'$(Configuration)|$(Platform)' == 'Development|AnyCPU'">
-    <DocumentationFile></DocumentationFile>
-    <Optimize>true</Optimize>
-  </PropertyGroup>
-  <PropertyGroup Condition="'$(Configuration)|$(Platform)' == 'Release|AnyCPU'">
-    <DocumentationFile></DocumentationFile>
-    <Optimize>true</Optimize>
-  </PropertyGroup>
-  <PropertyGroup Condition=" '$(Configuration)|$(Platform)' == 'Analyze|AnyCPU' ">
-    <RunAnalyzersDuringBuild>True</RunAnalyzersDuringBuild>
-    <TreatWarningsAsErrors>false</TreatWarningsAsErrors>
-    <EnforceCodeStyleInBuild>True</EnforceCodeStyleInBuild>
-    <GenerateDocumentationFile>True</GenerateDocumentationFile>
-    <DocumentationFile></DocumentationFile>
-  </PropertyGroup>
 
   <ItemGroup>
     <PackageReference Include="OpenTelemetry" Version="1.8.1" />
   </ItemGroup>
-  
+
   <ItemGroup>
     <ProjectReference Include="..\EpicGames.Core\EpicGames.Core.csproj" />
     <ProjectReference Include="..\EpicGames.IoHash\EpicGames.IoHash.csproj" />
     <ProjectReference Include="..\EpicGames.Perforce\EpicGames.Perforce.csproj" />
     <ProjectReference Include="..\EpicGames.Serialization\EpicGames.Serialization.csproj" />
   </ItemGroup>
-  <ItemGroup>
-    <Compile Include="..\..\Shared\MetaData.cs">
-      <Link>Properties\MetaData.cs</Link>
-    </Compile>
-  </ItemGroup>
 
 </Project>
diff --git a/Engine/Source/Programs/Shared/EpicGames.Perforce.Tests/EpicGames.Perforce.Tests.csproj b/Engine/Source/Programs/Shared/EpicGames.Perforce.Tests/EpicGames.Perforce.Tests.csproj
--- a/Engine/Source/Programs/Shared/EpicGames.Perforce.Tests/EpicGames.Perforce.Tests.csproj
+++ b/Engine/Source/Programs/Shared/EpicGames.Perforce.Tests/EpicGames.Perforce.Tests.csproj
@@ -1,46 +1,7 @@
 <Project Sdk="Microsoft.NET.Sdk">
-  <Import Project="..\UnrealEngine.csproj.props" />
-
   <PropertyGroup>
     <TargetFramework>net8.0</TargetFramework>
-    <Configuration Condition=" '$(Configuration)' == '' ">Development</Configuration>
-    <OutputType>Library</OutputType>
-    <Deterministic>true</Deterministic>
-    <RootNamespace>EpicGames.Perforce.Tests</RootNamespace>
-    <GenerateAssemblyCompanyAttribute>false</GenerateAssemblyCompanyAttribute>
-    <GenerateAssemblyInformationalVersionAttribute>false</GenerateAssemblyInformationalVersionAttribute>
-    <GenerateAssemblyProductAttribute>false</GenerateAssemblyProductAttribute>
-    <GenerateAssemblyVersionAttribute>false</GenerateAssemblyVersionAttribute>
-    <GenerateAssemblyFileVersionAttribute>false</GenerateAssemblyFileVersionAttribute>
-    <Configurations>Debug;Release;Development;Analyze</Configurations>
-    <IsPackable>false</IsPackable>
-    <Nullable>enable</Nullable>
-    <NoWarn>$(NoWarn);1591</NoWarn>
-  </PropertyGroup>
-  <PropertyGroup Condition=" '$(Configuration)|$(Platform)' == 'Analyze|AnyCPU' ">
-    <RunAnalyzersDuringBuild>True</RunAnalyzersDuringBuild>
-    <TreatWarningsAsErrors>false</TreatWarningsAsErrors>
-    <EnforceCodeStyleInBuild>True</EnforceCodeStyleInBuild>
-    <GenerateDocumentationFile>True</GenerateDocumentationFile>
-    <DocumentationFile></DocumentationFile>
   </PropertyGroup>
-  <ItemGroup>
-    <Compile Include="..\..\Shared\MetaData.cs">
-      <Link>Properties\MetaData.cs</Link>
-    </Compile>
-  </ItemGroup>
-
-  <ItemGroup>
-    <PackageReference Include="Microsoft.NET.Test.Sdk" Version="17.14.1" />
-    <PackageReference Include="MSTest.TestAdapter" Version="3.6.4">
-      <TreatAsUsed>true</TreatAsUsed>
-    </PackageReference>
-    <PackageReference Include="MSTest.TestFramework" Version="3.6.4" />
-    <PackageReference Include="coverlet.collector" Version="6.0.4">
-      <PrivateAssets>all</PrivateAssets>
-      <IncludeAssets>runtime; build; native; contentfiles; analyzers; buildtransitive</IncludeAssets>
-    </PackageReference>
-  </ItemGroup>
 
   <ItemGroup>
     <ProjectReference Include="..\EpicGames.Perforce.Managed\EpicGames.Perforce.Managed.csproj" />
diff --git a/Engine/Source/Programs/Shared/EpicGames.Perforce/EpicGames.Perforce.csproj b/Engine/Source/Programs/Shared/EpicGames.Perforce/EpicGames.Perforce.csproj
--- a/Engine/Source/Programs/Shared/EpicGames.Perforce/EpicGames.Perforce.csproj
+++ b/Engine/Source/Programs/Shared/EpicGames.Perforce/EpicGames.Perforce.csproj
@@ -1,43 +1,12 @@
-<Project Sdk="Microsoft.NET.Sdk">
-  <Import Project="..\UnrealEngine.csproj.props" />
-
+﻿<Project Sdk="Microsoft.NET.Sdk">
   <PropertyGroup>
     <TargetFramework>net8.0</TargetFramework>
-    <Configuration Condition=" '$(Configuration)' == '' ">Development</Configuration>
-    <OutputType>Library</OutputType>
-    <RootNamespace>EpicGames.Perforce</RootNamespace>
-    <GenerateAssemblyInfo>false</GenerateAssemblyInfo>
-    <GenerateTargetFrameworkAttribute>false</GenerateTargetFrameworkAttribute>
-    <Configurations>Analyze;Debug;Release;Development</Configurations>
-    <nullable>enable</nullable>
-  </PropertyGroup>
-  <PropertyGroup Condition=" '$(Configuration)|$(Platform)' == 'Debug|AnyCPU' ">
-    <DocumentationFile>EpicGames.Perforce.xml</DocumentationFile>
-  </PropertyGroup>
-  <PropertyGroup Condition="'$(Configuration)|$(Platform)' == 'Development|AnyCPU'">
-    <DocumentationFile></DocumentationFile>
-    <Optimize>true</Optimize>
-  </PropertyGroup>
-  <PropertyGroup Condition="'$(Configuration)|$(Platform)' == 'Release|AnyCPU'">
-    <DocumentationFile></DocumentationFile>
-    <Optimize>true</Optimize>
-  </PropertyGroup>
-  <PropertyGroup Condition=" '$(Configuration)|$(Platform)' == 'Analyze|AnyCPU' ">
-    <RunAnalyzersDuringBuild>True</RunAnalyzersDuringBuild>
-    <TreatWarningsAsErrors>false</TreatWarningsAsErrors>
-    <EnforceCodeStyleInBuild>True</EnforceCodeStyleInBuild>
-    <GenerateDocumentationFile>True</GenerateDocumentationFile>
-    <DocumentationFile></DocumentationFile>
   </PropertyGroup>
+
   <ItemGroup>
     <None Remove="EpicGames.Perforce.uatbuildrecord" />
     <None Remove="EpicGames.Perforce.xml" />
   </ItemGroup>
-  <ItemGroup>
-    <Compile Include="..\..\Shared\MetaData.cs">
-      <Link>Properties\MetaData.cs</Link>
-    </Compile>
-  </ItemGroup>
 
   <ItemGroup>
     <ProjectReference Include="..\EpicGames.Core\EpicGames.Core.csproj" />
diff --git a/Engine/Source/Programs/Shared/EpicGames.ProjectStore/EpicGames.ProjectStore.csproj b/Engine/Source/Programs/Shared/EpicGames.ProjectStore/EpicGames.ProjectStore.csproj
--- a/Engine/Source/Programs/Shared/EpicGames.ProjectStore/EpicGames.ProjectStore.csproj
+++ b/Engine/Source/Programs/Shared/EpicGames.ProjectStore/EpicGames.ProjectStore.csproj
@@ -1,42 +1,11 @@
-<Project Sdk="Microsoft.NET.Sdk">
-  <Import Project="..\UnrealEngine.csproj.props" />
-  
+﻿<Project Sdk="Microsoft.NET.Sdk">
   <PropertyGroup>
     <TargetFramework>net8.0</TargetFramework>
-    <Configuration Condition=" '$(Configuration)' == '' ">Development</Configuration>
-    <OutputType>Library</OutputType>
-    <Deterministic>true</Deterministic>
-    <RootNamespace>EpicGames.ProjectStore</RootNamespace>
-    <GenerateAssemblyInfo>false</GenerateAssemblyInfo>
-    <GenerateTargetFrameworkAttribute>false</GenerateTargetFrameworkAttribute>
-    <Configurations>Debug;Release;Development;Analyze</Configurations>
-    <Nullable>enable</Nullable>
-    <RunAnalyzersDuringBuild>false</RunAnalyzersDuringBuild>
   </PropertyGroup>
-  <PropertyGroup Condition=" '$(Configuration)|$(Platform)' == 'Debug|AnyCPU' ">
-    <DocumentationFile></DocumentationFile>
-  </PropertyGroup>
-  <PropertyGroup Condition="'$(Configuration)|$(Platform)' == 'Development|AnyCPU'">
-    <DocumentationFile></DocumentationFile>
-    <Optimize>true</Optimize>
-  </PropertyGroup>
-  <PropertyGroup Condition="'$(Configuration)|$(Platform)' == 'Release|AnyCPU'">
-    <DocumentationFile></DocumentationFile>
-    <Optimize>true</Optimize>
-  </PropertyGroup>
-  <PropertyGroup Condition=" '$(Configuration)|$(Platform)' == 'Analyze|AnyCPU' ">
-    <DocumentationFile></DocumentationFile>
-    <RunAnalyzersDuringBuild>True</RunAnalyzersDuringBuild>
-    <GenerateDocumentationFile>True</GenerateDocumentationFile>
-  </PropertyGroup>
-  <ItemGroup>
-    <Compile Include="..\..\Shared\MetaData.cs">
-      <Link>Properties\MetaData.cs</Link>
-    </Compile>
-  </ItemGroup>
+
   <ItemGroup>
     <ProjectReference Include="..\EpicGames.Build\EpicGames.Build.csproj" />
     <ProjectReference Include="..\EpicGames.Core\EpicGames.Core.csproj" />
     <ProjectReference Include="..\EpicGames.Horde\EpicGames.Horde.csproj" />
   </ItemGroup>
-</Project>
\ No newline at end of file
+</Project>
diff --git a/Engine/Source/Programs/Shared/EpicGames.Redis.Tests/EpicGames.Redis.Tests.csproj b/Engine/Source/Programs/Shared/EpicGames.Redis.Tests/EpicGames.Redis.Tests.csproj
--- a/Engine/Source/Programs/Shared/EpicGames.Redis.Tests/EpicGames.Redis.Tests.csproj
+++ b/Engine/Source/Programs/Shared/EpicGames.Redis.Tests/EpicGames.Redis.Tests.csproj
@@ -1,46 +1,8 @@
 <Project Sdk="Microsoft.NET.Sdk">
-  <Import Project="..\UnrealEngine.csproj.props" />
-
   <PropertyGroup>
     <TargetFramework>net8.0</TargetFramework>
-    <Configuration Condition=" '$(Configuration)' == '' ">Development</Configuration>
-    <OutputType>Library</OutputType>
-    <Deterministic>true</Deterministic>
-    <RootNamespace>EpicGames.Serialization.Tests</RootNamespace>
-    <GenerateAssemblyCompanyAttribute>false</GenerateAssemblyCompanyAttribute>
-    <GenerateAssemblyInformationalVersionAttribute>false</GenerateAssemblyInformationalVersionAttribute>
-    <GenerateAssemblyProductAttribute>false</GenerateAssemblyProductAttribute>
-    <GenerateAssemblyVersionAttribute>false</GenerateAssemblyVersionAttribute>
-    <GenerateAssemblyFileVersionAttribute>false</GenerateAssemblyFileVersionAttribute>
-    <Configurations>Debug;Release;Development;Analyze</Configurations>
-    <IsPackable>false</IsPackable>
-    <Nullable>enable</Nullable>
-    <NoWarn>$(NoWarn);1701;1702;1591</NoWarn>
-  </PropertyGroup>
-  <PropertyGroup Condition=" '$(Configuration)|$(Platform)' == 'Analyze|AnyCPU' ">
-    <RunAnalyzersDuringBuild>True</RunAnalyzersDuringBuild>
-    <TreatWarningsAsErrors>false</TreatWarningsAsErrors>
-    <EnforceCodeStyleInBuild>True</EnforceCodeStyleInBuild>
-    <GenerateDocumentationFile>True</GenerateDocumentationFile>
-    <DocumentationFile></DocumentationFile>
+    <NoWarn>$(NoWarn);1701;1702</NoWarn>
   </PropertyGroup>
-  <ItemGroup>
-    <Compile Include="..\..\Shared\MetaData.cs">
-      <Link>Properties\MetaData.cs</Link>
-    </Compile>
-  </ItemGroup>
-
-  <ItemGroup>
-    <PackageReference Include="Microsoft.NET.Test.Sdk" Version="17.14.1" />
-    <PackageReference Include="MSTest.TestAdapter" Version="3.6.4">
-      <TreatAsUsed>true</TreatAsUsed>
-    </PackageReference>
-    <PackageReference Include="MSTest.TestFramework" Version="3.6.4" />
-    <PackageReference Include="coverlet.collector" Version="6.0.4">
-      <PrivateAssets>all</PrivateAssets>
-      <IncludeAssets>runtime; build; native; contentfiles; analyzers; buildtransitive</IncludeAssets>
-    </PackageReference>
-  </ItemGroup>
 
   <ItemGroup>
     <ProjectReference Include="..\EpicGames.Redis\EpicGames.Redis.csproj" />
diff --git a/Engine/Source/Programs/Shared/EpicGames.Redis/EpicGames.Redis.csproj b/Engine/Source/Programs/Shared/EpicGames.Redis/EpicGames.Redis.csproj
--- a/Engine/Source/Programs/Shared/EpicGames.Redis/EpicGames.Redis.csproj
+++ b/Engine/Source/Programs/Shared/EpicGames.Redis/EpicGames.Redis.csproj
@@ -1,48 +1,17 @@
-<Project Sdk="Microsoft.NET.Sdk">
-  <Import Project="..\UnrealEngine.csproj.props" />
-  
+﻿<Project Sdk="Microsoft.NET.Sdk">
   <PropertyGroup>
     <TargetFramework>net8.0</TargetFramework>
-    <Configuration Condition=" '$(Configuration)' == '' ">Development</Configuration>
-    <OutputType>Library</OutputType>
-    <Deterministic>true</Deterministic>
-    <RootNamespace>EpicGames.Redis</RootNamespace>
-    <GenerateAssemblyInfo>false</GenerateAssemblyInfo>
-    <GenerateTargetFrameworkAttribute>false</GenerateTargetFrameworkAttribute>
-    <Configurations>Debug;Release;Development;Analyze</Configurations>
-    <Nullable>enable</Nullable>
     <GenerateDocumentationFile>True</GenerateDocumentationFile>
   </PropertyGroup>
-  <PropertyGroup Condition=" '$(Configuration)|$(Platform)' == 'Debug|AnyCPU' ">
-    <DocumentationFile></DocumentationFile>
-  </PropertyGroup>
-  <PropertyGroup Condition="'$(Configuration)|$(Platform)' == 'Development|AnyCPU'">
-    <DocumentationFile></DocumentationFile>
-    <Optimize>false</Optimize>
-  </PropertyGroup>
-  <PropertyGroup Condition="'$(Configuration)|$(Platform)' == 'Release|AnyCPU'">
-    <DocumentationFile></DocumentationFile>
-    <Optimize>true</Optimize>
-  </PropertyGroup>
-  <PropertyGroup Condition=" '$(Configuration)|$(Platform)' == 'Analyze|AnyCPU' ">
-    <RunAnalyzersDuringBuild>True</RunAnalyzersDuringBuild>
-    <TreatWarningsAsErrors>false</TreatWarningsAsErrors>
-    <EnforceCodeStyleInBuild>True</EnforceCodeStyleInBuild>
-    <GenerateDocumentationFile>True</GenerateDocumentationFile>
-    <DocumentationFile></DocumentationFile>
-  </PropertyGroup>
-  <ItemGroup>
-    <Compile Include="..\..\Shared\MetaData.cs">
-      <Link>Properties\MetaData.cs</Link>
-    </Compile>
-  </ItemGroup>
+
   <ItemGroup>
     <PackageReference Include="Google.Protobuf" Version="3.25.8" />
     <PackageReference Include="protobuf-net" Version="3.2.56" />
     <PackageReference Include="StackExchange.Redis" Version="2.8.58" />
   </ItemGroup>
+
   <ItemGroup>
     <ProjectReference Include="..\EpicGames.Core\EpicGames.Core.csproj" />
     <ProjectReference Include="..\EpicGames.Serialization\EpicGames.Serialization.csproj" />
   </ItemGroup>
-</Project>
\ No newline at end of file
+</Project>
diff --git a/Engine/Source/Programs/Shared/EpicGames.ScriptBuild/EpicGames.ScriptBuild.csproj b/Engine/Source/Programs/Shared/EpicGames.ScriptBuild/EpicGames.ScriptBuild.csproj
--- a/Engine/Source/Programs/Shared/EpicGames.ScriptBuild/EpicGames.ScriptBuild.csproj
+++ b/Engine/Source/Programs/Shared/EpicGames.ScriptBuild/EpicGames.ScriptBuild.csproj
@@ -1,43 +1,10 @@
 ﻿<Project Sdk="Microsoft.NET.Sdk">
-  <Import Project="..\UnrealEngine.csproj.props" />
-  
   <PropertyGroup>
     <TargetFramework>net8.0</TargetFramework>
-    <Configuration Condition=" '$(Configuration)' == '' ">Development</Configuration>
-    <OutputType>Library</OutputType>
-    <Deterministic>true</Deterministic>
-    <RootNamespace>EpicGames.ScriptBuild</RootNamespace>
-    <GenerateAssemblyInfo>false</GenerateAssemblyInfo>
-    <GenerateTargetFrameworkAttribute>false</GenerateTargetFrameworkAttribute>
-    <Configurations>Debug;Release;Development;Analyze</Configurations>
-    <Nullable>enable</Nullable>
   </PropertyGroup>
-  <PropertyGroup Condition=" '$(Configuration)|$(Platform)' == 'Debug|AnyCPU' ">
-    <DocumentationFile></DocumentationFile>
-  </PropertyGroup>
-  <PropertyGroup Condition="'$(Configuration)|$(Platform)' == 'Development|AnyCPU'">
-    <DocumentationFile></DocumentationFile>
-    <Optimize>true</Optimize>
-  </PropertyGroup>
-  <PropertyGroup Condition="'$(Configuration)|$(Platform)' == 'Release|AnyCPU'">
-    <DocumentationFile></DocumentationFile>
-    <Optimize>true</Optimize>
-  </PropertyGroup>
-  <PropertyGroup Condition=" '$(Configuration)|$(Platform)' == 'Analyze|AnyCPU' ">
-    <RunAnalyzersDuringBuild>True</RunAnalyzersDuringBuild>
-    <TreatWarningsAsErrors>False</TreatWarningsAsErrors>
-    <EnforceCodeStyleInBuild>True</EnforceCodeStyleInBuild>
-    <GenerateDocumentationFile>True</GenerateDocumentationFile>
-    <DocumentationFile></DocumentationFile>
-  </PropertyGroup>
-  <ItemGroup>
-    <Compile Include="..\..\Shared\MetaData.cs">
-      <Link>Properties\MetaData.cs</Link>
-    </Compile>
-  </ItemGroup>
 
   <!-- Reference file is generated by AutomationTool. It contains a list of project references that will be compiled with 'dotnet build' -->
   <Import
     Project="$(EngineDir)\Intermediate\Build\EpicGames.ScriptBuild.props"
     Condition="Exists('$(EngineDir)\Intermediate\Build\EpicGames.ScriptBuild.props') And '$(EpicGamesMsBuild)' == 'true'" />
-</Project>
\ No newline at end of file
+</Project>
diff --git a/Engine/Source/Programs/Shared/EpicGames.Serialization.Tests/EpicGames.Serialization.Tests.csproj b/Engine/Source/Programs/Shared/EpicGames.Serialization.Tests/EpicGames.Serialization.Tests.csproj
--- a/Engine/Source/Programs/Shared/EpicGames.Serialization.Tests/EpicGames.Serialization.Tests.csproj
+++ b/Engine/Source/Programs/Shared/EpicGames.Serialization.Tests/EpicGames.Serialization.Tests.csproj
@@ -1,49 +1,9 @@
 <Project Sdk="Microsoft.NET.Sdk">
-  <Import Project="..\UnrealEngine.csproj.props" />
-
   <PropertyGroup>
     <TargetFramework>net8.0</TargetFramework>
-    <Configuration Condition=" '$(Configuration)' == '' ">Development</Configuration>
-    <OutputType>Library</OutputType>
-    <Deterministic>true</Deterministic>
-    <RootNamespace>EpicGames.Serialization.Tests</RootNamespace>
-    <GenerateAssemblyCompanyAttribute>false</GenerateAssemblyCompanyAttribute>
-    <GenerateAssemblyInformationalVersionAttribute>false</GenerateAssemblyInformationalVersionAttribute>
-    <GenerateAssemblyProductAttribute>false</GenerateAssemblyProductAttribute>
-    <GenerateAssemblyVersionAttribute>false</GenerateAssemblyVersionAttribute>
-    <GenerateAssemblyFileVersionAttribute>false</GenerateAssemblyFileVersionAttribute>
-    <Configurations>Debug;Release;Development;Analyze</Configurations>
-    <IsPackable>false</IsPackable>
-    <Nullable>enable</Nullable>
-    <NoWarn>$(NoWarn);1701;1702;1591</NoWarn>
-  </PropertyGroup>
-
-  <PropertyGroup Condition=" '$(Configuration)|$(Platform)' == 'Analyze|AnyCPU' ">
-    <RunAnalyzersDuringBuild>True</RunAnalyzersDuringBuild>
-    <TreatWarningsAsErrors>false</TreatWarningsAsErrors>
-    <EnforceCodeStyleInBuild>True</EnforceCodeStyleInBuild>
-    <GenerateDocumentationFile>True</GenerateDocumentationFile>
-    <DocumentationFile></DocumentationFile>
+    <NoWarn>$(NoWarn);1701;1702</NoWarn>
   </PropertyGroup>
 
-  <ItemGroup>
-    <Compile Include="..\..\Shared\MetaData.cs">
-      <Link>Properties\MetaData.cs</Link>
-    </Compile>
-  </ItemGroup>
-
-  <ItemGroup>
-    <PackageReference Include="Microsoft.NET.Test.Sdk" Version="17.14.1" />
-    <PackageReference Include="MSTest.TestAdapter" Version="3.6.4">
-      <TreatAsUsed>true</TreatAsUsed>
-    </PackageReference>
-    <PackageReference Include="MSTest.TestFramework" Version="3.6.4" />
-    <PackageReference Include="coverlet.collector" Version="6.0.4">
-      <PrivateAssets>all</PrivateAssets>
-      <IncludeAssets>runtime; build; native; contentfiles; analyzers; buildtransitive</IncludeAssets>
-    </PackageReference>
-  </ItemGroup>
-
   <ItemGroup>
     <ProjectReference Include="..\EpicGames.IoHash\EpicGames.IoHash.csproj" />
     <ProjectReference Include="..\EpicGames.Serialization\EpicGames.Serialization.csproj" />
diff --git a/Engine/Source/Programs/Shared/EpicGames.Serialization/EpicGames.Serialization.csproj b/Engine/Source/Programs/Shared/EpicGames.Serialization/EpicGames.Serialization.csproj
--- a/Engine/Source/Programs/Shared/EpicGames.Serialization/EpicGames.Serialization.csproj
+++ b/Engine/Source/Programs/Shared/EpicGames.Serialization/EpicGames.Serialization.csproj
@@ -1,44 +1,14 @@
-<Project Sdk="Microsoft.NET.Sdk">
-  <Import Project="..\UnrealEngine.csproj.props" />
-
+﻿<Project Sdk="Microsoft.NET.Sdk">
   <PropertyGroup>
     <TargetFramework>net8.0</TargetFramework>
-    <Configuration Condition=" '$(Configuration)' == '' ">Development</Configuration>
-    <OutputType>Library</OutputType>
-    <Deterministic>true</Deterministic>
-    <RootNamespace>EpicGames.Serialization</RootNamespace>
-    <GenerateAssemblyInfo>false</GenerateAssemblyInfo>
-    <GenerateTargetFrameworkAttribute>false</GenerateTargetFrameworkAttribute>
-    <Configurations>Debug;Release;Development;Analyze</Configurations>
-    <Nullable>enable</Nullable>
-    <RunAnalyzersDuringBuild>false</RunAnalyzersDuringBuild>
-  </PropertyGroup>
-  <PropertyGroup Condition=" '$(Configuration)|$(Platform)' == 'Debug|AnyCPU' ">
-    <DocumentationFile></DocumentationFile>
-  </PropertyGroup>
-  <PropertyGroup Condition="'$(Configuration)|$(Platform)' == 'Development|AnyCPU'">
-    <DocumentationFile></DocumentationFile>
-    <Optimize>true</Optimize>
-  </PropertyGroup>
-  <PropertyGroup Condition="'$(Configuration)|$(Platform)' == 'Release|AnyCPU'">
-    <DocumentationFile></DocumentationFile>
-    <Optimize>true</Optimize>
-  </PropertyGroup>
-  <PropertyGroup Condition=" '$(Configuration)|$(Platform)' == 'Analyze|AnyCPU' ">
-    <DocumentationFile></DocumentationFile>
-    <RunAnalyzersDuringBuild>True</RunAnalyzersDuringBuild>
-    <GenerateDocumentationFile>True</GenerateDocumentationFile>
   </PropertyGroup>
-  <ItemGroup>
-    <Compile Include="..\..\Shared\MetaData.cs">
-      <Link>Properties\MetaData.cs</Link>
-    </Compile>
-  </ItemGroup>
+
   <ItemGroup>
     <PackageReference Include="System.IO.Hashing" Version="8.0.0" />
   </ItemGroup>
+
   <ItemGroup>
     <ProjectReference Include="..\EpicGames.IoHash\EpicGames.IoHash.csproj" />
     <ProjectReference Include="..\EpicGames.Oodle\EpicGames.Oodle.csproj" />
   </ItemGroup>
-</Project>
\ No newline at end of file
+</Project>
diff --git a/Engine/Source/Programs/Shared/EpicGames.Slack/EpicGames.Slack.csproj b/Engine/Source/Programs/Shared/EpicGames.Slack/EpicGames.Slack.csproj
--- a/Engine/Source/Programs/Shared/EpicGames.Slack/EpicGames.Slack.csproj
+++ b/Engine/Source/Programs/Shared/EpicGames.Slack/EpicGames.Slack.csproj
@@ -1,47 +1,13 @@
-<Project Sdk="Microsoft.NET.Sdk">
-  <Import Project="..\UnrealEngine.csproj.props" />
-
+﻿<Project Sdk="Microsoft.NET.Sdk">
   <PropertyGroup>
     <TargetFramework>net8.0</TargetFramework>
-    <Configuration Condition=" '$(Configuration)' == '' ">Development</Configuration>
-    <OutputType>Library</OutputType>
-    <Deterministic>true</Deterministic>
-    <RootNamespace>EpicGames.Slack</RootNamespace>
-    <GenerateAssemblyInfo>false</GenerateAssemblyInfo>
-    <GenerateTargetFrameworkAttribute>false</GenerateTargetFrameworkAttribute>
-    <Configurations>Debug;Release;Development;Analyze</Configurations>
-    <Nullable>enable</Nullable>
     <DebugType>pdbonly</DebugType>
     <RunAnalyzersDuringBuild>true</RunAnalyzersDuringBuild>
   </PropertyGroup>
-  <PropertyGroup Condition=" '$(Configuration)|$(Platform)' == 'Debug|AnyCPU' ">
-    <DocumentationFile></DocumentationFile>
-  </PropertyGroup>
-  <PropertyGroup Condition="'$(Configuration)|$(Platform)' == 'Development|AnyCPU'">
-    <DocumentationFile></DocumentationFile>
-    <Optimize>true</Optimize>
-  </PropertyGroup>
-  <PropertyGroup Condition="'$(Configuration)|$(Platform)' == 'Release|AnyCPU'">
-    <DocumentationFile></DocumentationFile>
-    <Optimize>true</Optimize>
-  </PropertyGroup>
-  <PropertyGroup Condition=" '$(Configuration)|$(Platform)' == 'Analyze|AnyCPU' ">
-    <DocumentationFile></DocumentationFile>
-    <EnforceCodeStyleInBuild>True</EnforceCodeStyleInBuild>
-    <RunAnalyzersDuringBuild>True</RunAnalyzersDuringBuild>
-    <GenerateDocumentationFile>True</GenerateDocumentationFile>
-  </PropertyGroup>
-  <ItemGroup>
-    <Compile Include="..\..\Shared\MetaData.cs">
-      <Link>Properties\MetaData.cs</Link>
-    </Compile>
-  </ItemGroup>
-  <ItemGroup>
-    <Folder Include="Blocks\" />
-  </ItemGroup>
+
   <ItemGroup>
     <PackageReference Include="Microsoft.Extensions.Http" Version="8.0.1" />
     <PackageReference Include="Microsoft.Extensions.Http.Polly" Version="8.0.20" />
     <PackageReference Include="Microsoft.Extensions.Logging.Abstractions" Version="8.0.3" />
   </ItemGroup>
-</Project>
\ No newline at end of file
+</Project>
diff --git a/Engine/Source/Programs/Shared/EpicGames.Tracing.Tests/EpicGames.Tracing.Tests.csproj b/Engine/Source/Programs/Shared/EpicGames.Tracing.Tests/EpicGames.Tracing.Tests.csproj
--- a/Engine/Source/Programs/Shared/EpicGames.Tracing.Tests/EpicGames.Tracing.Tests.csproj
+++ b/Engine/Source/Programs/Shared/EpicGames.Tracing.Tests/EpicGames.Tracing.Tests.csproj
@@ -1,45 +1,8 @@
 <Project Sdk="Microsoft.NET.Sdk">
-  <Import Project="..\UnrealEngine.csproj.props" />
-
   <PropertyGroup>
     <TargetFramework>net8.0</TargetFramework>
-    <Configuration Condition=" '$(Configuration)' == '' ">Development</Configuration>
     <OutputType>Library</OutputType>
-    <Deterministic>true</Deterministic>
-    <RootNamespace>EpicGames.Tracing.Tests</RootNamespace>
-    <GenerateAssemblyCompanyAttribute>false</GenerateAssemblyCompanyAttribute>
-    <GenerateAssemblyInformationalVersionAttribute>false</GenerateAssemblyInformationalVersionAttribute>
-    <GenerateAssemblyProductAttribute>false</GenerateAssemblyProductAttribute>
-    <GenerateAssemblyVersionAttribute>false</GenerateAssemblyVersionAttribute>
-    <GenerateAssemblyFileVersionAttribute>false</GenerateAssemblyFileVersionAttribute>
-    <Configurations>Debug;Release;Development;Analyze</Configurations>
-    <IsPackable>false</IsPackable>
-    <Nullable>enable</Nullable>
-    <NoWarn>$(NoWarn);1591</NoWarn>
-  </PropertyGroup>
-  <PropertyGroup Condition=" '$(Configuration)|$(Platform)' == 'Analyze|AnyCPU' ">
-    <RunAnalyzersDuringBuild>True</RunAnalyzersDuringBuild>
-    <TreatWarningsAsErrors>false</TreatWarningsAsErrors>
-    <GenerateDocumentationFile>True</GenerateDocumentationFile>
-    <DocumentationFile></DocumentationFile>
   </PropertyGroup>
-  <ItemGroup>
-    <Compile Include="..\..\Shared\MetaData.cs">
-      <Link>Properties\MetaData.cs</Link>
-    </Compile>
-  </ItemGroup>
-
-  <ItemGroup>
-    <PackageReference Include="Microsoft.NET.Test.Sdk" Version="17.14.1" />
-    <PackageReference Include="MSTest.TestAdapter" Version="3.6.4">
-      <TreatAsUsed>true</TreatAsUsed>
-    </PackageReference>
-    <PackageReference Include="MSTest.TestFramework" Version="3.6.4" />
-    <PackageReference Include="coverlet.collector" Version="6.0.4">
-      <PrivateAssets>all</PrivateAssets>
-      <IncludeAssets>runtime; build; native; contentfiles; analyzers; buildtransitive</IncludeAssets>
-    </PackageReference>
-  </ItemGroup>
 
   <ItemGroup>
     <ProjectReference Include="..\EpicGames.Tracing\EpicGames.Tracing.csproj" />
diff --git a/Engine/Source/Programs/Shared/EpicGames.Tracing/EpicGames.Tracing.csproj b/Engine/Source/Programs/Shared/EpicGames.Tracing/EpicGames.Tracing.csproj
--- a/Engine/Source/Programs/Shared/EpicGames.Tracing/EpicGames.Tracing.csproj
+++ b/Engine/Source/Programs/Shared/EpicGames.Tracing/EpicGames.Tracing.csproj
@@ -1,38 +1,6 @@
-<Project Sdk="Microsoft.NET.Sdk">
-  <Import Project="..\UnrealEngine.csproj.props" />
-
+﻿<Project Sdk="Microsoft.NET.Sdk">
   <PropertyGroup>
     <TargetFramework>net8.0</TargetFramework>
-    <Configuration Condition=" '$(Configuration)' == '' ">Development</Configuration>
-    <OutputType>Library</OutputType>
-    <Deterministic>true</Deterministic>
-    <RootNamespace>EpicGames.Tracing</RootNamespace>
-    <GenerateAssemblyInfo>false</GenerateAssemblyInfo>
-    <GenerateTargetFrameworkAttribute>false</GenerateTargetFrameworkAttribute>
-    <Configurations>Debug;Release;Development;Analyze</Configurations>
-    <Nullable>enable</Nullable>
     <NoWarn>$(NoWarn);1591</NoWarn>
   </PropertyGroup>
-  <PropertyGroup Condition=" '$(Configuration)|$(Platform)' == 'Debug|AnyCPU' ">
-    <DocumentationFile></DocumentationFile>
-  </PropertyGroup>
-  <PropertyGroup Condition="'$(Configuration)|$(Platform)' == 'Development|AnyCPU'">
-    <DocumentationFile></DocumentationFile>
-    <Optimize>true</Optimize>
-  </PropertyGroup>
-  <PropertyGroup Condition="'$(Configuration)|$(Platform)' == 'Release|AnyCPU'">
-    <DocumentationFile></DocumentationFile>
-    <Optimize>true</Optimize>
-  </PropertyGroup>
-  <PropertyGroup Condition=" '$(Configuration)|$(Platform)' == 'Analyze|AnyCPU' ">
-    <RunAnalyzersDuringBuild>True</RunAnalyzersDuringBuild>
-    <TreatWarningsAsErrors>false</TreatWarningsAsErrors>
-    <GenerateDocumentationFile>True</GenerateDocumentationFile>
-    <DocumentationFile></DocumentationFile>
-  </PropertyGroup>
-  <ItemGroup>
-    <Compile Include="..\..\Shared\MetaData.cs">
-      <Link>Properties\MetaData.cs</Link>
-    </Compile>
-  </ItemGroup>
-</Project>
\ No newline at end of file
+</Project>
diff --git a/Engine/Source/Programs/Shared/EpicGames.UBA/EpicGames.UBA.csproj b/Engine/Source/Programs/Shared/EpicGames.UBA/EpicGames.UBA.csproj
--- a/Engine/Source/Programs/Shared/EpicGames.UBA/EpicGames.UBA.csproj
+++ b/Engine/Source/Programs/Shared/EpicGames.UBA/EpicGames.UBA.csproj
@@ -1,39 +1,11 @@
-<Project Sdk="Microsoft.NET.Sdk">
-  <Import Project="..\UnrealEngine.csproj.props" />
-
+﻿<Project Sdk="Microsoft.NET.Sdk">
   <PropertyGroup>
     <TargetFramework>net8.0</TargetFramework>
-    <Configuration Condition=" '$(Configuration)' == '' ">Development</Configuration>
-    <OutputType>Library</OutputType>
-    <RootNamespace>EpicGames.UBA</RootNamespace>
-    <GenerateAssemblyInfo>false</GenerateAssemblyInfo>
-    <GenerateTargetFrameworkAttribute>false</GenerateTargetFrameworkAttribute>
-    <Configurations>Debug;Release;Development;Analyze</Configurations>
-    <Nullable>enable</Nullable>
     <AllowUnsafeBlocks>true</AllowUnsafeBlocks>
     <GenerateDocumentationFile>True</GenerateDocumentationFile>
-    <EnforceCodeStyleInBuild>False</EnforceCodeStyleInBuild>
     <ProduceReferenceAssemblyInOutDir>true</ProduceReferenceAssemblyInOutDir>
-    <RunAnalyzersDuringBuild>false</RunAnalyzersDuringBuild>
     <PackageReadmeFile>README.md</PackageReadmeFile>
   </PropertyGroup>
-  <PropertyGroup Condition=" '$(Configuration)|$(Platform)' == 'Debug|AnyCPU' ">
-    <DocumentationFile></DocumentationFile>
-  </PropertyGroup>
-  <PropertyGroup Condition="'$(Configuration)|$(Platform)' == 'Development|AnyCPU'">
-    <DocumentationFile></DocumentationFile>
-    <Optimize>true</Optimize>
-  </PropertyGroup>
-  <PropertyGroup Condition="'$(Configuration)|$(Platform)' == 'Release|AnyCPU'">
-    <DocumentationFile></DocumentationFile>
-    <Optimize>true</Optimize>
-  </PropertyGroup>
-  <PropertyGroup Condition=" '$(Configuration)|$(Platform)' == 'Analyze|AnyCPU' ">
-    <RunAnalyzersDuringBuild>True</RunAnalyzersDuringBuild>
-    <TreatWarningsAsErrors>false</TreatWarningsAsErrors>
-    <GenerateDocumentationFile>True</GenerateDocumentationFile>
-    <DocumentationFile></DocumentationFile>
-  </PropertyGroup>
 
   <ItemGroup>
     <Compile Include="../../../../Restricted/*/Source/Programs/Shared/EpicGames.UBA/**/*.cs">
@@ -42,18 +14,17 @@
   </ItemGroup>
 
   <ItemGroup>
-    <Compile Include="..\MetaData.cs">
-      <Link>Properties\MetaData.cs</Link>
-    </Compile>
     <None Include="docs\README.md" Pack="true" PackagePath="\" />
   </ItemGroup>
 
   <Import Project="Library.props" />
+
   <ItemGroup>
     <PackageReference Include="Microsoft.CSharp" Version="4.7.0" />
     <PackageReference Include="Microsoft.Extensions.Logging" Version="8.0.1" />
   </ItemGroup>
+
   <ItemGroup>
     <ProjectReference Include="..\EpicGames.Core\EpicGames.Core.csproj" />
   </ItemGroup>
-</Project>
\ No newline at end of file
+</Project>
diff --git a/Engine/Source/Programs/Shared/EpicGames.UHT/EpicGames.UHT.csproj b/Engine/Source/Programs/Shared/EpicGames.UHT/EpicGames.UHT.csproj
--- a/Engine/Source/Programs/Shared/EpicGames.UHT/EpicGames.UHT.csproj
+++ b/Engine/Source/Programs/Shared/EpicGames.UHT/EpicGames.UHT.csproj
@@ -1,49 +1,18 @@
 ﻿<Project Sdk="Microsoft.NET.Sdk">
-  <Import Project="..\UnrealEngine.csproj.props" />
-  
   <PropertyGroup>
     <TargetFramework>net8.0</TargetFramework>
-    <Configuration Condition=" '$(Configuration)' == '' ">Development</Configuration>
-    <OutputType>Library</OutputType>
-    <RootNamespace>EpicGames.UHT</RootNamespace>
-    <GenerateAssemblyInfo>false</GenerateAssemblyInfo>
-    <GenerateTargetFrameworkAttribute>false</GenerateTargetFrameworkAttribute>
-    <Configurations>Debug;Release;Development;Analyze</Configurations>
-    <Nullable>enable</Nullable>
     <AllowUnsafeBlocks>true</AllowUnsafeBlocks>
     <GenerateDocumentationFile>True</GenerateDocumentationFile>
-    <EnforceCodeStyleInBuild>False</EnforceCodeStyleInBuild>
     <ProduceReferenceAssemblyInOutDir>true</ProduceReferenceAssemblyInOutDir>
-    <RunAnalyzersDuringBuild>false</RunAnalyzersDuringBuild>
-  </PropertyGroup>
-  <PropertyGroup Condition=" '$(Configuration)|$(Platform)' == 'Debug|AnyCPU' ">
-    <DocumentationFile></DocumentationFile>
-  </PropertyGroup>
-  <PropertyGroup Condition="'$(Configuration)|$(Platform)' == 'Development|AnyCPU'">
-    <DocumentationFile></DocumentationFile>
-    <Optimize>true</Optimize>
-  </PropertyGroup>
-  <PropertyGroup Condition="'$(Configuration)|$(Platform)' == 'Release|AnyCPU'">
-    <DocumentationFile></DocumentationFile>
-    <Optimize>true</Optimize>
-  </PropertyGroup>
-  <PropertyGroup Condition=" '$(Configuration)|$(Platform)' == 'Analyze|AnyCPU' ">
-    <RunAnalyzersDuringBuild>True</RunAnalyzersDuringBuild>
-    <TreatWarningsAsErrors>false</TreatWarningsAsErrors>
-    <GenerateDocumentationFile>True</GenerateDocumentationFile>
-    <DocumentationFile></DocumentationFile>
+    <EnforceCodeStyleInBuild>False</EnforceCodeStyleInBuild>
   </PropertyGroup>
-  
-  <ItemGroup>
-    <Compile Include="..\MetaData.cs">
-      <Link>Properties\MetaData.cs</Link>
-    </Compile>
-  </ItemGroup>
+
   <ItemGroup>
     <ProjectReference Include="..\..\Shared\EpicGames.Core\EpicGames.Core.csproj" />
     <ProjectReference Include="..\EpicGames.Build\EpicGames.Build.csproj" />
   </ItemGroup>
+
   <ItemGroup>
     <PackageReference Include="Microsoft.CSharp" Version="4.7.0" />
   </ItemGroup>
-</Project>
\ No newline at end of file
+</Project>
diff --git a/Engine/Source/Programs/Shared/UnrealBuildTool.Tests/UnrealBuildTool.Tests.csproj b/Engine/Source/Programs/Shared/UnrealBuildTool.Tests/UnrealBuildTool.Tests.csproj
--- a/Engine/Source/Programs/Shared/UnrealBuildTool.Tests/UnrealBuildTool.Tests.csproj
+++ b/Engine/Source/Programs/Shared/UnrealBuildTool.Tests/UnrealBuildTool.Tests.csproj
@@ -1,35 +1,10 @@
 <Project Sdk="Microsoft.NET.Sdk">
-  <Import Project="..\UnrealEngine.csproj.props" />
-
   <PropertyGroup>
     <TargetFramework>net8.0</TargetFramework>
-    <Configurations>Debug;Release;Development;Analyze</Configurations>
-    <Nullable>enable</Nullable>
-    <IsPackable>false</IsPackable>
-    <NoWarn>$(NoWarn);1591</NoWarn>
   </PropertyGroup>
-  <PropertyGroup Condition=" '$(Configuration)|$(Platform)' == 'Analyze|AnyCPU' ">
-    <RunAnalyzersDuringBuild>True</RunAnalyzersDuringBuild>
-    <TreatWarningsAsErrors>false</TreatWarningsAsErrors>
-    <GenerateDocumentationFile>True</GenerateDocumentationFile>
-    <DocumentationFile></DocumentationFile>
-  </PropertyGroup>  
-
-  <ItemGroup>
-    <PackageReference Include="Microsoft.NET.Test.Sdk" Version="17.14.1" />
-    <PackageReference Include="MSTest.TestAdapter" Version="3.6.4">
-      <TreatAsUsed>true</TreatAsUsed>
-    </PackageReference>
-    <PackageReference Include="MSTest.TestFramework" Version="3.6.4" />
-    <PackageReference Include="coverlet.collector" Version="6.0.4">
-      <PrivateAssets>all</PrivateAssets>
-      <IncludeAssets>runtime; build; native; contentfiles; analyzers; buildtransitive</IncludeAssets>
-    </PackageReference>
-  </ItemGroup>
 
   <ItemGroup>
     <ProjectReference Include="..\EpicGames.Horde\EpicGames.Horde.csproj" />
     <ProjectReference Include="..\..\UnrealBuildTool\UnrealBuildTool.csproj" />
   </ItemGroup>
-
 </Project>
