Cloud Storage - Added another longer wait (this time for 60 minutes before retrying). We have seen extreme cases were issues take about an hour to settle and a retry in this scenario would have recovered the build rather then forcing us to rerun the build which takes much longer then this 1 hour wait time.

#Synced-CL 49272873

diff --git a/Engine/Source/Programs/AutomationTool/BuildGraph/Tasks/ZenExportSnapshotTask.cs b/Engine/Source/Programs/AutomationTool/BuildGraph/Tasks/ZenExportSnapshotTask.cs
--- a/Engine/Source/Programs/AutomationTool/BuildGraph/Tasks/ZenExportSnapshotTask.cs
+++ b/Engine/Source/Programs/AutomationTool/BuildGraph/Tasks/ZenExportSnapshotTask.cs
@@ -440,10 +440,10 @@ namespace AutomationTool.Tasks
 
 		private static bool TryExportOplogCommand(string app, string commandLine)
 		{
-			int attemptLimit = 2;
 			int attempt = 0;
 			// increasingly long time to wait before retrying
-			int[] waitDurationWebSeconds = [60, 300, 900];
+			int[] waitDurationWebSeconds = [60, 300, 900, 3600];
+			int attemptLimit = waitDurationWebSeconds.Length;
 
 			while (attempt < attemptLimit)
 			{
diff --git a/Engine/Source/Programs/Shared/EpicGames.ProjectStore/CloudStorage/CloudStorage.cs b/Engine/Source/Programs/Shared/EpicGames.ProjectStore/CloudStorage/CloudStorage.cs
--- a/Engine/Source/Programs/Shared/EpicGames.ProjectStore/CloudStorage/CloudStorage.cs
+++ b/Engine/Source/Programs/Shared/EpicGames.ProjectStore/CloudStorage/CloudStorage.cs
@@ -560,21 +560,22 @@ namespace EpicGames.ProjectStore
 				zenLogger = optionsOverride?.EchoLogger;
 			}
 
-			const int MaxRetryAttempts = 3;
 			const int WarnAtAttempts = 2;
-			for (int retryAttempt = 0; retryAttempt < MaxRetryAttempts; retryAttempt++)
+			// increasingly long time to wait before retrying
+			int[] waitDurationSeconds = [60, 300, 900, 3600];
+			int maxRetryAttempts = waitDurationSeconds.Length;
+
+			for (int retryAttempt = 0; retryAttempt < maxRetryAttempts; retryAttempt++)
 			{
 				int exitCode = await RunProcessAsync(zenExe.FullName, cmdline, env: env, echoStdout: _options.Value.EchoZenCliToStdout, echoLogger: zenLogger, progress: progress, cancellationToken: cancellationToken);
 
-				// increasingly long time to wait before retrying
-				int[] waitDurationSeconds = [60, 300, 900];
 				int waitDuration = waitDurationSeconds[retryAttempt];
 
 				ZenExitCodes zenExitCode = AsZenExitCodeValue(exitCode);
 				// retry on zen exit codes that indicate transient issues
 				if (zenExitCode is ZenExitCodes.HttpServiceUnavailable or ZenExitCodes.HttpBadGateway or ZenExitCodes.HttpGatewayTimeout)
 				{
-					if (retryAttempt + 1 == MaxRetryAttempts)
+					if (retryAttempt + 1 == maxRetryAttempts)
 					{
 						// give up
 						_logger.LogError("Zen returned http status code {ExitCode}({ExitCodeMessage}) after {RetryAttempts} retries. Giving up.", exitCode, AsZenExitCodeValue(exitCode), retryAttempt);
@@ -599,7 +600,7 @@ namespace EpicGames.ProjectStore
 				return exitCode;
 			}
 
-			throw new CloudStorageException($"Failed to run zen cli, gave up after {MaxRetryAttempts} attempts");
+			throw new CloudStorageException($"Failed to run zen cli, gave up after {maxRetryAttempts} attempts");
 		}
 
 		private async Task<int> RunProcessAsync(string cmd, string arguments, string? workingDirectory = null, Dictionary<string, string>? env = null, IProgress<(float, string)>? progress = null, bool echoStdout = false, ILogger? echoLogger = null, CancellationToken cancellationToken = default)
