Fix cloud build query format:
- redundant '{'
- boolean parameter in quotes

#rnx
#Synced-CL 49369561

diff --git a/Engine/Source/Programs/Shared/EpicGames.ProjectStore/ProjectStoreTypes.cs b/Engine/Source/Programs/Shared/EpicGames.ProjectStore/ProjectStoreTypes.cs
--- a/Engine/Source/Programs/Shared/EpicGames.ProjectStore/ProjectStoreTypes.cs
+++ b/Engine/Source/Programs/Shared/EpicGames.ProjectStore/ProjectStoreTypes.cs
@@ -1013,11 +1013,11 @@ namespace EpicGames.ProjectStore
 
 					if (changelist > 0)
 					{
-						request.Content = new StringContent(String.Format("{{\"query\": {{\"isPreflight\": {{\"$eq\": \"false\"}}, {{\"runtime\": {{\"$eq\": \"{0}\"}}, \"changelist\": {{\"$lte\": {1}}}, \"branch\": {{\"$eq\": \"{2}\"}}}}}}", runtime, changelist, branch), MediaTypeHeaderValue.Parse("application/json"));
+						request.Content = new StringContent(String.Format("{{\"query\": {{\"isPreflight\": {{\"$eq\": false}}, \"runtime\": {{\"$eq\": \"{0}\"}}, \"changelist\": {{\"$lte\": {1}}}, \"branch\": {{\"$eq\": \"{2}\"}}}}}}", runtime, changelist, branch), MediaTypeHeaderValue.Parse("application/json"));
 					}
 					else
 					{
-						request.Content = new StringContent(String.Format("{{\"query\": {{\"isPreflight\": {{\"$eq\": \"false\"}}, {{\"runtime\": {{\"$eq\": \"{0}\"}}, \"branch\": {{\"$eq\": \"{1}\"}}}}}}", runtime, branch), MediaTypeHeaderValue.Parse("application/json"));
+						request.Content = new StringContent(String.Format("{{\"query\": {{\"isPreflight\": {{\"$eq\": false}}, \"runtime\": {{\"$eq\": \"{0}\"}}, \"branch\": {{\"$eq\": \"{1}\"}}}}}}", runtime, branch), MediaTypeHeaderValue.Parse("application/json"));
 					}
 
 					request.Headers.Authorization = new AuthenticationHeaderValue("Bearer", _accessToken);
