UhtSpecifierValueType.Legacy is now obsolete.  UhtSpecifierValueType.None should be used in its place.

#Synced-CL 49587850

diff --git a/Engine/Source/Programs/Shared/EpicGames.UHT/Parsers/UhtSpecifierParser.cs b/Engine/Source/Programs/Shared/EpicGames.UHT/Parsers/UhtSpecifierParser.cs
--- a/Engine/Source/Programs/Shared/EpicGames.UHT/Parsers/UhtSpecifierParser.cs
+++ b/Engine/Source/Programs/Shared/EpicGames.UHT/Parsers/UhtSpecifierParser.cs
@@ -243,6 +243,7 @@ namespace EpicGames.UHT.Parsers
 		{
 			value = null;
 
+#pragma warning disable CS0618 // Type or member is obsolete
 			switch (valueType)
 			{
 				case UhtSpecifierValueType.NotSet:
@@ -292,7 +293,9 @@ namespace EpicGames.UHT.Parsers
 					value = ReadValueList();
 					return true;
 
+#pragma warning disable CS0618 // Type or member is obsolete
 				case UhtSpecifierValueType.Legacy:
+#pragma warning restore CS0618 // Type or member is obsolete
 					value = ReadValueList();
 					return true;
 
@@ -344,6 +347,7 @@ namespace EpicGames.UHT.Parsers
 				default:
 					throw new UhtIceException("Unknown value type");
 			}
+#pragma warning restore CS0618 // Type or member is obsolete
 		}
 
 		private KeyValuePair<StringView, StringView> ReadKVP()
diff --git a/Engine/Source/Programs/Shared/EpicGames.UHT/Specifiers/UhtClassSpecifiers.cs b/Engine/Source/Programs/Shared/EpicGames.UHT/Specifiers/UhtClassSpecifiers.cs
--- a/Engine/Source/Programs/Shared/EpicGames.UHT/Specifiers/UhtClassSpecifiers.cs
+++ b/Engine/Source/Programs/Shared/EpicGames.UHT/Specifiers/UhtClassSpecifiers.cs
@@ -16,21 +16,21 @@ namespace EpicGames.UHT.Parsers
 	[SuppressMessage("CodeQuality", "IDE0051:Remove unused private members", Justification = "Attribute accessed method")]
 	public static class UhtClassSpecifiers
 	{
-		[UhtSpecifier(Extends = UhtTableNames.Class, ValueType = UhtSpecifierValueType.Legacy)]
+		[UhtSpecifier(Extends = UhtTableNames.Class, ValueType = UhtSpecifierValueType.None)]
 		private static void NoExportSpecifier(UhtSpecifierContext specifierContext)
 		{
 			UhtClass classObj = (UhtClass)specifierContext.Type;
 			classObj.ClassExportFlags |= UhtClassExportFlags.NoExport;
 		}
 
-		[UhtSpecifier(Extends = UhtTableNames.Class, ValueType = UhtSpecifierValueType.Legacy)]
+		[UhtSpecifier(Extends = UhtTableNames.Class, ValueType = UhtSpecifierValueType.None)]
 		private static void IntrinsicSpecifier(UhtSpecifierContext specifierContext)
 		{
 			UhtClass classObj = (UhtClass)specifierContext.Type;
 			classObj.AddClassFlags(EClassFlags.Intrinsic);
 		}
 
-		[UhtSpecifier(Extends = UhtTableNames.Class, ValueType = UhtSpecifierValueType.Legacy)]
+		[UhtSpecifier(Extends = UhtTableNames.Class, ValueType = UhtSpecifierValueType.None)]
 		private static void ComponentWrapperClassSpecifier(UhtSpecifierContext specifierContext)
 		{
 			specifierContext.MetaData.Add(UhtNames.IgnoreCategoryKeywordsInSubclasses, true);
@@ -43,49 +43,49 @@ namespace EpicGames.UHT.Parsers
 			classObj.ClassWithinIdentifier = value.ToString();
 		}
 
-		[UhtSpecifier(Extends = UhtTableNames.Class, ValueType = UhtSpecifierValueType.Legacy)]
+		[UhtSpecifier(Extends = UhtTableNames.Class, ValueType = UhtSpecifierValueType.None)]
 		private static void EditInlineNewSpecifier(UhtSpecifierContext specifierContext)
 		{
 			UhtClass classObj = (UhtClass)specifierContext.Type;
 			classObj.AddClassFlags(EClassFlags.EditInlineNew);
 		}
 
-		[UhtSpecifier(Extends = UhtTableNames.Class, ValueType = UhtSpecifierValueType.Legacy)]
+		[UhtSpecifier(Extends = UhtTableNames.Class, ValueType = UhtSpecifierValueType.None)]
 		private static void NotEditInlineNewSpecifier(UhtSpecifierContext specifierContext)
 		{
 			UhtClass classObj = (UhtClass)specifierContext.Type;
 			classObj.RemoveClassFlags(EClassFlags.EditInlineNew);
 		}
 
-		[UhtSpecifier(Extends = UhtTableNames.Class, ValueType = UhtSpecifierValueType.Legacy)]
+		[UhtSpecifier(Extends = UhtTableNames.Class, ValueType = UhtSpecifierValueType.None)]
 		private static void PlaceableSpecifier(UhtSpecifierContext specifierContext)
 		{
 			UhtClass classObj = (UhtClass)specifierContext.Type;
 			classObj.RemoveClassFlags(EClassFlags.NotPlaceable);
 		}
 
-		[UhtSpecifier(Extends = UhtTableNames.Class, ValueType = UhtSpecifierValueType.Legacy)]
+		[UhtSpecifier(Extends = UhtTableNames.Class, ValueType = UhtSpecifierValueType.None)]
 		private static void NotPlaceableSpecifier(UhtSpecifierContext specifierContext)
 		{
 			UhtClass classObj = (UhtClass)specifierContext.Type;
 			classObj.AddClassFlags(EClassFlags.NotPlaceable);
 		}
 
-		[UhtSpecifier(Extends = UhtTableNames.Class, ValueType = UhtSpecifierValueType.Legacy)]
+		[UhtSpecifier(Extends = UhtTableNames.Class, ValueType = UhtSpecifierValueType.None)]
 		private static void DefaultToInstancedSpecifier(UhtSpecifierContext specifierContext)
 		{
 			UhtClass classObj = (UhtClass)specifierContext.Type;
 			classObj.AddClassFlags(EClassFlags.DefaultToInstanced);
 		}
 
-		[UhtSpecifier(Extends = UhtTableNames.Class, ValueType = UhtSpecifierValueType.Legacy)]
+		[UhtSpecifier(Extends = UhtTableNames.Class, ValueType = UhtSpecifierValueType.None)]
 		private static void HideDropdownSpecifier(UhtSpecifierContext specifierContext)
 		{
 			UhtClass classObj = (UhtClass)specifierContext.Type;
 			classObj.AddClassFlags(EClassFlags.HideDropDown);
 		}
 
-		[UhtSpecifier(Extends = UhtTableNames.Class, ValueType = UhtSpecifierValueType.Legacy)]
+		[UhtSpecifier(Extends = UhtTableNames.Class, ValueType = UhtSpecifierValueType.None)]
 		private static void HiddenSpecifier(UhtSpecifierContext specifierContext)
 		{
 			// Prevents class from appearing in the editor class browser and edit inline menus.
@@ -93,76 +93,76 @@ namespace EpicGames.UHT.Parsers
 			classObj.AddClassFlags(EClassFlags.Hidden);
 		}
 
-		[UhtSpecifier(Extends = UhtTableNames.Class, ValueType = UhtSpecifierValueType.Legacy)]
+		[UhtSpecifier(Extends = UhtTableNames.Class, ValueType = UhtSpecifierValueType.None)]
 		private static void DependsOnSpecifier(UhtSpecifierContext specifierContext)
 		{
 			specifierContext.MessageSite.LogError("The dependsOn specifier is deprecated. Please use #include \"ClassHeaderFilename.h\" instead.");
 		}
 
-		[UhtSpecifier(Extends = UhtTableNames.Class, ValueType = UhtSpecifierValueType.Legacy)]
+		[UhtSpecifier(Extends = UhtTableNames.Class, ValueType = UhtSpecifierValueType.None)]
 		private static void MinimalAPISpecifier(UhtSpecifierContext specifierContext)
 		{
 			UhtClass classObj = (UhtClass)specifierContext.Type;
 			classObj.AddClassFlags(EClassFlags.MinimalAPI);
 		}
 
-		[UhtSpecifier(Extends = UhtTableNames.Class, ValueType = UhtSpecifierValueType.Legacy)]
+		[UhtSpecifier(Extends = UhtTableNames.Class, ValueType = UhtSpecifierValueType.None)]
 		private static void ConstSpecifier(UhtSpecifierContext specifierContext)
 		{
 			UhtClass classObj = (UhtClass)specifierContext.Type;
 			classObj.AddClassFlags(EClassFlags.Const);
 		}
 
-		[UhtSpecifier(Extends = UhtTableNames.Class, ValueType = UhtSpecifierValueType.Legacy)]
+		[UhtSpecifier(Extends = UhtTableNames.Class, ValueType = UhtSpecifierValueType.None)]
 		private static void PerObjectConfigSpecifier(UhtSpecifierContext specifierContext)
 		{
 			UhtClass classObj = (UhtClass)specifierContext.Type;
 			classObj.AddClassFlags(EClassFlags.PerObjectConfig);
 		}
 
-		[UhtSpecifier(Extends = UhtTableNames.Class, ValueType = UhtSpecifierValueType.Legacy)]
+		[UhtSpecifier(Extends = UhtTableNames.Class, ValueType = UhtSpecifierValueType.None)]
 		private static void PerPlatformConfigSpecifier(UhtSpecifierContext specifierContext)
 		{
 			UhtClass classObj = (UhtClass)specifierContext.Type;
 			classObj.AddClassFlags(EClassFlags.PerPlatformConfig);
 		}
 
-		[UhtSpecifier(Extends = UhtTableNames.Class, ValueType = UhtSpecifierValueType.Legacy)]
+		[UhtSpecifier(Extends = UhtTableNames.Class, ValueType = UhtSpecifierValueType.None)]
 		private static void ConfigDoNotCheckDefaultsSpecifier(UhtSpecifierContext specifierContext)
 		{
 			UhtClass classObj = (UhtClass)specifierContext.Type;
 			classObj.AddClassFlags(EClassFlags.ConfigDoNotCheckDefaults);
 		}
 
-		[UhtSpecifier(Extends = UhtTableNames.Class, ValueType = UhtSpecifierValueType.Legacy)]
+		[UhtSpecifier(Extends = UhtTableNames.Class, ValueType = UhtSpecifierValueType.None)]
 		private static void AbstractSpecifier(UhtSpecifierContext specifierContext)
 		{
 			UhtClass classObj = (UhtClass)specifierContext.Type;
 			classObj.AddClassFlags(EClassFlags.Abstract);
 		}
 
-		[UhtSpecifier(Extends = UhtTableNames.Class, ValueType = UhtSpecifierValueType.Legacy)]
+		[UhtSpecifier(Extends = UhtTableNames.Class, ValueType = UhtSpecifierValueType.None)]
 		private static void DeprecatedSpecifier(UhtSpecifierContext specifierContext)
 		{
 			UhtClass classObj = (UhtClass)specifierContext.Type;
 			classObj.AddClassFlags(EClassFlags.Deprecated | EClassFlags.NotPlaceable);
 		}
 
-		[UhtSpecifier(Extends = UhtTableNames.Class, ValueType = UhtSpecifierValueType.Legacy)]
+		[UhtSpecifier(Extends = UhtTableNames.Class, ValueType = UhtSpecifierValueType.None)]
 		private static void TransientSpecifier(UhtSpecifierContext specifierContext)
 		{
 			UhtClass classObj = (UhtClass)specifierContext.Type;
 			classObj.AddClassFlags(EClassFlags.Transient);
 		}
 
-		[UhtSpecifier(Extends = UhtTableNames.Class, ValueType = UhtSpecifierValueType.Legacy)]
+		[UhtSpecifier(Extends = UhtTableNames.Class, ValueType = UhtSpecifierValueType.None)]
 		private static void NonTransientSpecifier(UhtSpecifierContext specifierContext)
 		{
 			UhtClass classObj = (UhtClass)specifierContext.Type;
 			classObj.RemoveClassFlags(EClassFlags.Transient);
 		}
 
-		[UhtSpecifier(Extends = UhtTableNames.Class, ValueType = UhtSpecifierValueType.Legacy)]
+		[UhtSpecifier(Extends = UhtTableNames.Class, ValueType = UhtSpecifierValueType.None)]
 		private static void OptionalSpecifier(UhtSpecifierContext specifierContext)
 		{
 			// Optional class
@@ -170,7 +170,7 @@ namespace EpicGames.UHT.Parsers
 			classObj.AddClassFlags(EClassFlags.Optional);
 		}
 
-		[UhtSpecifier(Extends = UhtTableNames.Class, ValueType = UhtSpecifierValueType.Legacy)]
+		[UhtSpecifier(Extends = UhtTableNames.Class, ValueType = UhtSpecifierValueType.None)]
 		private static void CustomConstructorSpecifier(UhtSpecifierContext specifierContext)
 		{
 			// we will not export a constructor for this class, assuming it is in the CPP block
@@ -185,7 +185,7 @@ namespace EpicGames.UHT.Parsers
 			classObj.Config = value.ToString();
 		}
 
-		[UhtSpecifier(Extends = UhtTableNames.Class, ValueType = UhtSpecifierValueType.Legacy)]
+		[UhtSpecifier(Extends = UhtTableNames.Class, ValueType = UhtSpecifierValueType.None)]
 		private static void DefaultConfigSpecifier(UhtSpecifierContext specifierContext)
 		{
 			// Save object config only to Default INIs, never to local INIs.
@@ -193,7 +193,7 @@ namespace EpicGames.UHT.Parsers
 			classObj.AddClassFlags(EClassFlags.DefaultConfig);
 		}
 
-		[UhtSpecifier(Extends = UhtTableNames.Class, ValueType = UhtSpecifierValueType.Legacy)]
+		[UhtSpecifier(Extends = UhtTableNames.Class, ValueType = UhtSpecifierValueType.None)]
 		private static void GlobalUserConfigSpecifier(UhtSpecifierContext specifierContext)
 		{
 			// Save object config only to global user overrides, never to local INIs
@@ -201,7 +201,7 @@ namespace EpicGames.UHT.Parsers
 			classObj.AddClassFlags(EClassFlags.GlobalUserConfig);
 		}
 
-		[UhtSpecifier(Extends = UhtTableNames.Class, ValueType = UhtSpecifierValueType.Legacy)]
+		[UhtSpecifier(Extends = UhtTableNames.Class, ValueType = UhtSpecifierValueType.None)]
 		private static void ProjectUserConfigSpecifier(UhtSpecifierContext specifierContext)
 		{
 			// Save object config only to project user overrides, never to INIs that are checked in
@@ -290,7 +290,7 @@ namespace EpicGames.UHT.Parsers
 			classObj.AutoCollapseCategories.RemoveSwapRange(value);
 		}
 
-		[UhtSpecifier(Extends = UhtTableNames.Class, ValueType = UhtSpecifierValueType.Legacy)]
+		[UhtSpecifier(Extends = UhtTableNames.Class, ValueType = UhtSpecifierValueType.None)]
 		private static void CollapseCategoriesSpecifier(UhtSpecifierContext specifierContext)
 		{
 			// Class' properties should not be shown categorized in the editor.
@@ -298,7 +298,7 @@ namespace EpicGames.UHT.Parsers
 			classObj.AddClassFlags(EClassFlags.CollapseCategories);
 		}
 
-		[UhtSpecifier(Extends = UhtTableNames.Class, ValueType = UhtSpecifierValueType.Legacy)]
+		[UhtSpecifier(Extends = UhtTableNames.Class, ValueType = UhtSpecifierValueType.None)]
 		private static void DontCollapseCategoriesSpecifier(UhtSpecifierContext specifierContext)
 		{
 			// Class' properties should be shown categorized in the editor.
@@ -306,7 +306,7 @@ namespace EpicGames.UHT.Parsers
 			classObj.RemoveClassFlags(EClassFlags.CollapseCategories);
 		}
 
-		[UhtSpecifier(Extends = UhtTableNames.Class, ValueType = UhtSpecifierValueType.Legacy)]
+		[UhtSpecifier(Extends = UhtTableNames.Class, ValueType = UhtSpecifierValueType.None)]
 		private static void AdvancedClassDisplaySpecifier(UhtSpecifierContext specifierContext)
 		{
 			// By default the class properties are shown in advanced sections in UI
@@ -314,21 +314,21 @@ namespace EpicGames.UHT.Parsers
 			classObj.MetaData.Add(UhtNames.AdvancedClassDisplay, "true");
 		}
 
-		[UhtSpecifier(Extends = UhtTableNames.Class, ValueType = UhtSpecifierValueType.Legacy)]
+		[UhtSpecifier(Extends = UhtTableNames.Class, ValueType = UhtSpecifierValueType.None)]
 		private static void ConversionRootSpecifier(UhtSpecifierContext specifierContext)
 		{
 			UhtClass classObj = (UhtClass)specifierContext.Type;
 			classObj.MetaData.Add(UhtNames.IsConversionRoot, "true");
 		}
 
-		[UhtSpecifier(Extends = UhtTableNames.Class, ValueType = UhtSpecifierValueType.Legacy)]
+		[UhtSpecifier(Extends = UhtTableNames.Class, ValueType = UhtSpecifierValueType.None)]
 		private static void NeedsDeferredDependencyLoadingSpecifier(UhtSpecifierContext specifierContext)
 		{
 			UhtClass classObj = (UhtClass)specifierContext.Type;
 			classObj.AddClassFlags(EClassFlags.NeedsDeferredDependencyLoading);
 		}
 
-		[UhtSpecifier(Extends = UhtTableNames.Class, ValueType = UhtSpecifierValueType.Legacy)]
+		[UhtSpecifier(Extends = UhtTableNames.Class, ValueType = UhtSpecifierValueType.None)]
 		private static void MatchedSerializersSpecifier(UhtSpecifierContext specifierContext)
 		{
 			if (!specifierContext.Type.HeaderFile.IsNoExportTypes)
@@ -339,7 +339,7 @@ namespace EpicGames.UHT.Parsers
 			classObj.AddClassFlags(EClassFlags.MatchedSerializers);
 		}
 
-		[UhtSpecifier(Extends = UhtTableNames.Class, ValueType = UhtSpecifierValueType.Legacy)]
+		[UhtSpecifier(Extends = UhtTableNames.Class, ValueType = UhtSpecifierValueType.None)]
 		private static void InterfaceSpecifier(UhtSpecifierContext specifierContext)
 		{
 			if (!specifierContext.Type.HeaderFile.IsNoExportTypes)
diff --git a/Engine/Source/Programs/Shared/EpicGames.UHT/Specifiers/UhtEnumSpecifiers.cs b/Engine/Source/Programs/Shared/EpicGames.UHT/Specifiers/UhtEnumSpecifiers.cs
--- a/Engine/Source/Programs/Shared/EpicGames.UHT/Specifiers/UhtEnumSpecifiers.cs
+++ b/Engine/Source/Programs/Shared/EpicGames.UHT/Specifiers/UhtEnumSpecifiers.cs
@@ -16,14 +16,14 @@ namespace EpicGames.UHT.Parsers
 	[SuppressMessage("CodeQuality", "IDE0051:Remove unused private members", Justification = "Attribute accessed method")]
 	public static class UhtEnumSpecifiers
 	{
-		[UhtSpecifier(Extends = UhtTableNames.Enum, ValueType = UhtSpecifierValueType.Legacy)]
+		[UhtSpecifier(Extends = UhtTableNames.Enum, ValueType = UhtSpecifierValueType.None)]
 		private static void FlagsSpecifier(UhtSpecifierContext specifierContext)
 		{
 			UhtEnum enumObj = (UhtEnum)specifierContext.Type;
 			enumObj.EnumFlags |= EEnumFlags.Flags;
 		}
 
-		[UhtSpecifier(Extends = UhtTableNames.Enum, ValueType = UhtSpecifierValueType.Legacy)]
+		[UhtSpecifier(Extends = UhtTableNames.Enum, ValueType = UhtSpecifierValueType.None)]
 		private static void MinimalAPISpecifier(UhtSpecifierContext specifierContext)
 		{
 			UhtEnum enumObj = (UhtEnum)specifierContext.Type;
diff --git a/Engine/Source/Programs/Shared/EpicGames.UHT/Specifiers/UhtFunctionSpecifiers.cs b/Engine/Source/Programs/Shared/EpicGames.UHT/Specifiers/UhtFunctionSpecifiers.cs
--- a/Engine/Source/Programs/Shared/EpicGames.UHT/Specifiers/UhtFunctionSpecifiers.cs
+++ b/Engine/Source/Programs/Shared/EpicGames.UHT/Specifiers/UhtFunctionSpecifiers.cs
@@ -19,28 +19,28 @@ namespace EpicGames.UHT.Parsers
 	[SuppressMessage("Style", "IDE0060:Remove unused parameter", Justification = "Attribute accessed method")]
 	public static class UhtFunctionSpecifiers
 	{
-		[UhtSpecifier(Extends = UhtTableNames.Function, ValueType = UhtSpecifierValueType.Legacy)]
+		[UhtSpecifier(Extends = UhtTableNames.Function, ValueType = UhtSpecifierValueType.None)]
 		private static void BlueprintAuthorityOnlySpecifier(UhtSpecifierContext specifierContext)
 		{
 			UhtFunction function = (UhtFunction)specifierContext.Type;
 			function.FunctionFlags |= EFunctionFlags.BlueprintAuthorityOnly;
 		}
 
-		[UhtSpecifier(Extends = UhtTableNames.Function, ValueType = UhtSpecifierValueType.Legacy)]
+		[UhtSpecifier(Extends = UhtTableNames.Function, ValueType = UhtSpecifierValueType.None)]
 		private static void BlueprintCallableSpecifier(UhtSpecifierContext specifierContext)
 		{
 			UhtFunction function = (UhtFunction)specifierContext.Type;
 			function.FunctionFlags |= EFunctionFlags.BlueprintCallable;
 		}
 
-		[UhtSpecifier(Extends = UhtTableNames.Function, ValueType = UhtSpecifierValueType.Legacy)]
+		[UhtSpecifier(Extends = UhtTableNames.Function, ValueType = UhtSpecifierValueType.None)]
 		private static void BlueprintCosmeticSpecifier(UhtSpecifierContext specifierContext)
 		{
 			UhtFunction function = (UhtFunction)specifierContext.Type;
 			function.FunctionFlags |= EFunctionFlags.BlueprintCosmetic;
 		}
 
-		[UhtSpecifier(Extends = UhtTableNames.Function, ValueType = UhtSpecifierValueType.Legacy)]
+		[UhtSpecifier(Extends = UhtTableNames.Function, ValueType = UhtSpecifierValueType.None)]
 		private static void BlueprintGetterSpecifier(UhtSpecifierContext specifierContext)
 		{
 			UhtFunction function = (UhtFunction)specifierContext.Type;
@@ -56,7 +56,7 @@ namespace EpicGames.UHT.Parsers
 			function.MetaData.Add(UhtNames.BlueprintGetter, "");
 		}
 
-		[UhtSpecifier(Extends = UhtTableNames.Function, ValueType = UhtSpecifierValueType.Legacy)]
+		[UhtSpecifier(Extends = UhtTableNames.Function, ValueType = UhtSpecifierValueType.None)]
 		private static void BlueprintImplementableEventSpecifier(UhtSpecifierContext specifierContext)
 		{
 			UhtFunction function = (UhtFunction)specifierContext.Type;
@@ -84,7 +84,7 @@ namespace EpicGames.UHT.Parsers
 			function.FunctionFlags &= ~EFunctionFlags.Native;
 		}
 
-		[UhtSpecifier(Extends = UhtTableNames.Function, ValueType = UhtSpecifierValueType.Legacy)]
+		[UhtSpecifier(Extends = UhtTableNames.Function, ValueType = UhtSpecifierValueType.None)]
 		private static void BlueprintNativeEventSpecifier(UhtSpecifierContext specifierContext)
 		{
 			UhtFunction function = (UhtFunction)specifierContext.Type;
@@ -129,7 +129,7 @@ namespace EpicGames.UHT.Parsers
 			}
 		}
 
-		[UhtSpecifier(Extends = UhtTableNames.Function, ValueType = UhtSpecifierValueType.Legacy)]
+		[UhtSpecifier(Extends = UhtTableNames.Function, ValueType = UhtSpecifierValueType.None)]
 		private static void BlueprintSetterSpecifier(UhtSpecifierContext specifierContext)
 		{
 			UhtFunction function = (UhtFunction)specifierContext.Type;
@@ -168,14 +168,14 @@ namespace EpicGames.UHT.Parsers
 			}
 		}
 
-		[UhtSpecifier(Extends = UhtTableNames.Function, ValueType = UhtSpecifierValueType.Legacy)]
+		[UhtSpecifier(Extends = UhtTableNames.Function, ValueType = UhtSpecifierValueType.None)]
 		private static void CustomThunkSpecifier(UhtSpecifierContext specifierContext)
 		{
 			UhtFunction function = (UhtFunction)specifierContext.Type;
 			function.FunctionExportFlags |= UhtFunctionExportFlags.CustomThunk;
 		}
 
-		[UhtSpecifier(Extends = UhtTableNames.Function, ValueType = UhtSpecifierValueType.Legacy)]
+		[UhtSpecifier(Extends = UhtTableNames.Function, ValueType = UhtSpecifierValueType.None)]
 		private static void ExecSpecifier(UhtSpecifierContext specifierContext)
 		{
 			UhtFunction function = (UhtFunction)specifierContext.Type;
@@ -186,7 +186,7 @@ namespace EpicGames.UHT.Parsers
 			}
 		}
 
-		[UhtSpecifier(Extends = UhtTableNames.Function, ValueType = UhtSpecifierValueType.Legacy)]
+		[UhtSpecifier(Extends = UhtTableNames.Function, ValueType = UhtSpecifierValueType.None)]
 		private static void NetMulticastSpecifier(UhtSpecifierContext specifierContext)
 		{
 			UhtFunction function = (UhtFunction)specifierContext.Type;
@@ -199,14 +199,14 @@ namespace EpicGames.UHT.Parsers
 			function.FunctionFlags |= EFunctionFlags.NetMulticast;
 		}
 
-		[UhtSpecifier(Extends = UhtTableNames.Function, ValueType = UhtSpecifierValueType.Legacy)]
+		[UhtSpecifier(Extends = UhtTableNames.Function, ValueType = UhtSpecifierValueType.None)]
 		private static void ReliableSpecifier(UhtSpecifierContext specifierContext)
 		{
 			UhtFunction function = (UhtFunction)specifierContext.Type;
 			function.FunctionFlags |= EFunctionFlags.NetReliable;
 		}
 
-		[UhtSpecifier(Extends = UhtTableNames.Function, ValueType = UhtSpecifierValueType.Legacy)]
+		[UhtSpecifier(Extends = UhtTableNames.Function, ValueType = UhtSpecifierValueType.None)]
 		private static void SealedEventSpecifier(UhtSpecifierContext specifierContext)
 		{
 			UhtFunction function = (UhtFunction)specifierContext.Type;
@@ -306,7 +306,7 @@ namespace EpicGames.UHT.Parsers
 			}
 		}
 
-		[UhtSpecifier(Extends = UhtTableNames.Function, ValueType = UhtSpecifierValueType.Legacy)]
+		[UhtSpecifier(Extends = UhtTableNames.Function, ValueType = UhtSpecifierValueType.None)]
 		private static void UnreliableSpecifier(UhtSpecifierContext specifierContext)
 		{
 			UhtFunction function = (UhtFunction)specifierContext.Type;
diff --git a/Engine/Source/Programs/Shared/EpicGames.UHT/Specifiers/UhtInterfaceClassSpecifiers.cs b/Engine/Source/Programs/Shared/EpicGames.UHT/Specifiers/UhtInterfaceClassSpecifiers.cs
--- a/Engine/Source/Programs/Shared/EpicGames.UHT/Specifiers/UhtInterfaceClassSpecifiers.cs
+++ b/Engine/Source/Programs/Shared/EpicGames.UHT/Specifiers/UhtInterfaceClassSpecifiers.cs
@@ -12,20 +12,20 @@ namespace EpicGames.UHT.Parsers
 	[SuppressMessage("CodeQuality", "IDE0051:Remove unused private members", Justification = "Attribute accessed method")]
 	sealed class UhtInterfaceClassSpecifiers
 	{
-		[UhtSpecifier(Extends = UhtTableNames.Interface, ValueType = UhtSpecifierValueType.Legacy)]
+		[UhtSpecifier(Extends = UhtTableNames.Interface, ValueType = UhtSpecifierValueType.None)]
 		private static void DependsOnSpecifier(UhtSpecifierContext specifierContext)
 		{
 			throw new UhtException(specifierContext.MessageSite, $"The dependsOn specifier is deprecated. Please use #include \"ClassHeaderFilename.h\" instead.");
 		}
 
-		[UhtSpecifier(Extends = UhtTableNames.Interface, ValueType = UhtSpecifierValueType.Legacy)]
+		[UhtSpecifier(Extends = UhtTableNames.Interface, ValueType = UhtSpecifierValueType.None)]
 		private static void MinimalAPISpecifier(UhtSpecifierContext specifierContext)
 		{
 			UhtClass clasObj = (UhtClass)specifierContext.Type;
 			clasObj.ClassFlags |= EClassFlags.MinimalAPI;
 		}
 
-		[UhtSpecifier(Extends = UhtTableNames.Interface, ValueType = UhtSpecifierValueType.Legacy)]
+		[UhtSpecifier(Extends = UhtTableNames.Interface, ValueType = UhtSpecifierValueType.None)]
 		private static void ConversionRootSpecifier(UhtSpecifierContext specifierContext)
 		{
 			UhtClass classObj = (UhtClass)specifierContext.Type;
diff --git a/Engine/Source/Programs/Shared/EpicGames.UHT/Specifiers/UhtPropertyArgumentSpecifiers.cs b/Engine/Source/Programs/Shared/EpicGames.UHT/Specifiers/UhtPropertyArgumentSpecifiers.cs
--- a/Engine/Source/Programs/Shared/EpicGames.UHT/Specifiers/UhtPropertyArgumentSpecifiers.cs
+++ b/Engine/Source/Programs/Shared/EpicGames.UHT/Specifiers/UhtPropertyArgumentSpecifiers.cs
@@ -19,21 +19,21 @@ namespace EpicGames.UHT.Parsers
 	public static class UhtPropertyArgumentSpecifiers
 	{
 		#region Argument Property Specifiers
-		[UhtSpecifier(Extends = UhtTableNames.PropertyArgument, ValueType = UhtSpecifierValueType.Legacy)]
+		[UhtSpecifier(Extends = UhtTableNames.PropertyArgument, ValueType = UhtSpecifierValueType.None)]
 		private static void ConstSpecifier(UhtSpecifierContext specifierContext)
 		{
 			UhtPropertySpecifierContext context = (UhtPropertySpecifierContext)specifierContext;
 			context.PropertySettings.PropertyFlags |= EPropertyFlags.ConstParm;
 		}
 
-		[UhtSpecifier(Extends = UhtTableNames.PropertyArgument, ValueType = UhtSpecifierValueType.Legacy)]
+		[UhtSpecifier(Extends = UhtTableNames.PropertyArgument, ValueType = UhtSpecifierValueType.None)]
 		private static void RefSpecifier(UhtSpecifierContext specifierContext)
 		{
 			UhtPropertySpecifierContext context = (UhtPropertySpecifierContext)specifierContext;
 			context.PropertySettings.PropertyFlags |= EPropertyFlags.OutParm | EPropertyFlags.ReferenceParm;
 		}
 
-		[UhtSpecifier(Extends = UhtTableNames.PropertyArgument, ValueType = UhtSpecifierValueType.Legacy)]
+		[UhtSpecifier(Extends = UhtTableNames.PropertyArgument, ValueType = UhtSpecifierValueType.None)]
 		private static void NotReplicatedSpecifier(UhtSpecifierContext specifierContext)
 		{
 			UhtPropertySpecifierContext context = (UhtPropertySpecifierContext)specifierContext;
diff --git a/Engine/Source/Programs/Shared/EpicGames.UHT/Specifiers/UhtPropertyMemberSpecifiers.cs b/Engine/Source/Programs/Shared/EpicGames.UHT/Specifiers/UhtPropertyMemberSpecifiers.cs
--- a/Engine/Source/Programs/Shared/EpicGames.UHT/Specifiers/UhtPropertyMemberSpecifiers.cs
+++ b/Engine/Source/Programs/Shared/EpicGames.UHT/Specifiers/UhtPropertyMemberSpecifiers.cs
@@ -17,7 +17,7 @@ namespace EpicGames.UHT.Parsers
 	[SuppressMessage("CodeQuality", "IDE0051:Remove unused private members", Justification = "Attribute accessed method")]
 	public static class UhtPropertyMemberSpecifiers
 	{
-		[UhtSpecifier(Extends = UhtTableNames.PropertyMember, ValueType = UhtSpecifierValueType.Legacy)]
+		[UhtSpecifier(Extends = UhtTableNames.PropertyMember, ValueType = UhtSpecifierValueType.None)]
 		private static void EditAnywhereSpecifier(UhtSpecifierContext specifierContext)
 		{
 			UhtPropertySpecifierContext context = (UhtPropertySpecifierContext)specifierContext;
@@ -29,7 +29,7 @@ namespace EpicGames.UHT.Parsers
 			context.SeenEditSpecifier = true;
 		}
 
-		[UhtSpecifier(Extends = UhtTableNames.PropertyMember, ValueType = UhtSpecifierValueType.Legacy)]
+		[UhtSpecifier(Extends = UhtTableNames.PropertyMember, ValueType = UhtSpecifierValueType.None)]
 		private static void EditInstanceOnlySpecifier(UhtSpecifierContext specifierContext)
 		{
 			UhtPropertySpecifierContext context = (UhtPropertySpecifierContext)specifierContext;
@@ -41,7 +41,7 @@ namespace EpicGames.UHT.Parsers
 			context.SeenEditSpecifier = true;
 		}
 
-		[UhtSpecifier(Extends = UhtTableNames.PropertyMember, ValueType = UhtSpecifierValueType.Legacy)]
+		[UhtSpecifier(Extends = UhtTableNames.PropertyMember, ValueType = UhtSpecifierValueType.None)]
 		private static void EditDefaultsOnlySpecifier(UhtSpecifierContext specifierContext)
 		{
 			UhtPropertySpecifierContext context = (UhtPropertySpecifierContext)specifierContext;
@@ -53,7 +53,7 @@ namespace EpicGames.UHT.Parsers
 			context.SeenEditSpecifier = true;
 		}
 
-		[UhtSpecifier(Extends = UhtTableNames.PropertyMember, ValueType = UhtSpecifierValueType.Legacy)]
+		[UhtSpecifier(Extends = UhtTableNames.PropertyMember, ValueType = UhtSpecifierValueType.None)]
 		private static void VisibleAnywhereSpecifier(UhtSpecifierContext specifierContext)
 		{
 			UhtPropertySpecifierContext context = (UhtPropertySpecifierContext)specifierContext;
@@ -65,7 +65,7 @@ namespace EpicGames.UHT.Parsers
 			context.SeenEditSpecifier = true;
 		}
 
-		[UhtSpecifier(Extends = UhtTableNames.PropertyMember, ValueType = UhtSpecifierValueType.Legacy)]
+		[UhtSpecifier(Extends = UhtTableNames.PropertyMember, ValueType = UhtSpecifierValueType.None)]
 		private static void VisibleInstanceOnlySpecifier(UhtSpecifierContext specifierContext)
 		{
 			UhtPropertySpecifierContext context = (UhtPropertySpecifierContext)specifierContext;
@@ -77,7 +77,7 @@ namespace EpicGames.UHT.Parsers
 			context.SeenEditSpecifier = true;
 		}
 
-		[UhtSpecifier(Extends = UhtTableNames.PropertyMember, ValueType = UhtSpecifierValueType.Legacy)]
+		[UhtSpecifier(Extends = UhtTableNames.PropertyMember, ValueType = UhtSpecifierValueType.None)]
 		private static void VisibleDefaultsOnlySpecifier(UhtSpecifierContext specifierContext)
 		{
 			UhtPropertySpecifierContext context = (UhtPropertySpecifierContext)specifierContext;
@@ -89,7 +89,7 @@ namespace EpicGames.UHT.Parsers
 			context.SeenEditSpecifier = true;
 		}
 
-		[UhtSpecifier(Extends = UhtTableNames.PropertyMember, ValueType = UhtSpecifierValueType.Legacy)]
+		[UhtSpecifier(Extends = UhtTableNames.PropertyMember, ValueType = UhtSpecifierValueType.None)]
 		private static void BlueprintReadWriteSpecifier(UhtSpecifierContext specifierContext)
 		{
 			UhtPropertySpecifierContext context = (UhtPropertySpecifierContext)specifierContext;
@@ -113,7 +113,7 @@ namespace EpicGames.UHT.Parsers
 			context.SeenBlueprintWriteSpecifier = true;
 		}
 
-		[UhtSpecifier(Extends = UhtTableNames.PropertyMember, ValueType = UhtSpecifierValueType.Legacy)]
+		[UhtSpecifier(Extends = UhtTableNames.PropertyMember, ValueType = UhtSpecifierValueType.None)]
 		private static void BlueprintReadOnlySpecifier(UhtSpecifierContext specifierContext)
 		{
 			UhtPropertySpecifierContext context = (UhtPropertySpecifierContext)specifierContext;
@@ -172,49 +172,49 @@ namespace EpicGames.UHT.Parsers
 			context.SeenBlueprintGetterSpecifier = true;
 		}
 
-		[UhtSpecifier(Extends = UhtTableNames.PropertyMember, ValueType = UhtSpecifierValueType.Legacy)]
+		[UhtSpecifier(Extends = UhtTableNames.PropertyMember, ValueType = UhtSpecifierValueType.None)]
 		private static void ConfigSpecifier(UhtSpecifierContext specifierContext)
 		{
 			UhtPropertySpecifierContext context = (UhtPropertySpecifierContext)specifierContext;
 			context.PropertySettings.PropertyFlags |= EPropertyFlags.Config;
 		}
 
-		[UhtSpecifier(Extends = UhtTableNames.PropertyMember, ValueType = UhtSpecifierValueType.Legacy)]
+		[UhtSpecifier(Extends = UhtTableNames.PropertyMember, ValueType = UhtSpecifierValueType.None)]
 		private static void GlobalConfigSpecifier(UhtSpecifierContext specifierContext)
 		{
 			UhtPropertySpecifierContext context = (UhtPropertySpecifierContext)specifierContext;
 			context.PropertySettings.PropertyFlags |= EPropertyFlags.GlobalConfig | EPropertyFlags.Config;
 		}
 
-		[UhtSpecifier(Extends = UhtTableNames.PropertyMember, ValueType = UhtSpecifierValueType.Legacy)]
+		[UhtSpecifier(Extends = UhtTableNames.PropertyMember, ValueType = UhtSpecifierValueType.None)]
 		private static void LocalizedSpecifier(UhtSpecifierContext specifierContext)
 		{
 			UhtPropertySpecifierContext context = (UhtPropertySpecifierContext)specifierContext;
 			context.MessageSite.LogError("The Localized specifier is deprecated");
 		}
 
-		[UhtSpecifier(Extends = UhtTableNames.PropertyMember, ValueType = UhtSpecifierValueType.Legacy)]
+		[UhtSpecifier(Extends = UhtTableNames.PropertyMember, ValueType = UhtSpecifierValueType.None)]
 		private static void TransientSpecifier(UhtSpecifierContext specifierContext)
 		{
 			UhtPropertySpecifierContext context = (UhtPropertySpecifierContext)specifierContext;
 			context.PropertySettings.PropertyFlags |= EPropertyFlags.Transient;
 		}
 
-		[UhtSpecifier(Extends = UhtTableNames.PropertyMember, ValueType = UhtSpecifierValueType.Legacy)]
+		[UhtSpecifier(Extends = UhtTableNames.PropertyMember, ValueType = UhtSpecifierValueType.None)]
 		private static void DuplicateTransientSpecifier(UhtSpecifierContext specifierContext)
 		{
 			UhtPropertySpecifierContext context = (UhtPropertySpecifierContext)specifierContext;
 			context.PropertySettings.PropertyFlags |= EPropertyFlags.DuplicateTransient;
 		}
 
-		[UhtSpecifier(Extends = UhtTableNames.PropertyMember, ValueType = UhtSpecifierValueType.Legacy)]
+		[UhtSpecifier(Extends = UhtTableNames.PropertyMember, ValueType = UhtSpecifierValueType.None)]
 		private static void TextExportTransientSpecifier(UhtSpecifierContext specifierContext)
 		{
 			UhtPropertySpecifierContext context = (UhtPropertySpecifierContext)specifierContext;
 			context.PropertySettings.PropertyFlags |= EPropertyFlags.TextExportTransient;
 		}
 
-		[UhtSpecifier(Extends = UhtTableNames.PropertyMember, ValueType = UhtSpecifierValueType.Legacy)]
+		[UhtSpecifier(Extends = UhtTableNames.PropertyMember, ValueType = UhtSpecifierValueType.None)]
 		private static void NonPIETransientSpecifier(UhtSpecifierContext specifierContext)
 		{
 			UhtPropertySpecifierContext context = (UhtPropertySpecifierContext)specifierContext;
@@ -222,42 +222,42 @@ namespace EpicGames.UHT.Parsers
 			context.PropertySettings.PropertyFlags |= EPropertyFlags.NonPIEDuplicateTransient;
 		}
 
-		[UhtSpecifier(Extends = UhtTableNames.PropertyMember, ValueType = UhtSpecifierValueType.Legacy)]
+		[UhtSpecifier(Extends = UhtTableNames.PropertyMember, ValueType = UhtSpecifierValueType.None)]
 		private static void NonPIEDuplicateTransientSpecifier(UhtSpecifierContext specifierContext)
 		{
 			UhtPropertySpecifierContext context = (UhtPropertySpecifierContext)specifierContext;
 			context.PropertySettings.PropertyFlags |= EPropertyFlags.NonPIEDuplicateTransient;
 		}
 
-		[UhtSpecifier(Extends = UhtTableNames.PropertyMember, ValueType = UhtSpecifierValueType.Legacy)]
+		[UhtSpecifier(Extends = UhtTableNames.PropertyMember, ValueType = UhtSpecifierValueType.None)]
 		private static void ExportSpecifier(UhtSpecifierContext specifierContext)
 		{
 			UhtPropertySpecifierContext context = (UhtPropertySpecifierContext)specifierContext;
 			context.PropertySettings.PropertyFlags |= EPropertyFlags.ExportObject;
 		}
 
-		[UhtSpecifier(Extends = UhtTableNames.PropertyMember, ValueType = UhtSpecifierValueType.Legacy)]
+		[UhtSpecifier(Extends = UhtTableNames.PropertyMember, ValueType = UhtSpecifierValueType.None)]
 		private static void EditInlineSpecifier(UhtSpecifierContext specifierContext)
 		{
 			UhtPropertySpecifierContext context = (UhtPropertySpecifierContext)specifierContext;
 			context.MessageSite.LogError("EditInline is deprecated. Remove it, or use Instanced instead.");
 		}
 
-		[UhtSpecifier(Extends = UhtTableNames.PropertyMember, ValueType = UhtSpecifierValueType.Legacy)]
+		[UhtSpecifier(Extends = UhtTableNames.PropertyMember, ValueType = UhtSpecifierValueType.None)]
 		private static void NoClearSpecifier(UhtSpecifierContext specifierContext)
 		{
 			UhtPropertySpecifierContext context = (UhtPropertySpecifierContext)specifierContext;
 			context.PropertySettings.PropertyFlags |= EPropertyFlags.NoClear;
 		}
 
-		[UhtSpecifier(Extends = UhtTableNames.PropertyMember, ValueType = UhtSpecifierValueType.Legacy)]
+		[UhtSpecifier(Extends = UhtTableNames.PropertyMember, ValueType = UhtSpecifierValueType.None)]
 		private static void EditFixedSizeSpecifier(UhtSpecifierContext specifierContext)
 		{
 			UhtPropertySpecifierContext context = (UhtPropertySpecifierContext)specifierContext;
 			context.PropertySettings.PropertyFlags |= EPropertyFlags.EditFixedSize;
 		}
 
-		[UhtSpecifier(Extends = UhtTableNames.PropertyMember, ValueType = UhtSpecifierValueType.Legacy)]
+		[UhtSpecifier(Extends = UhtTableNames.PropertyMember, ValueType = UhtSpecifierValueType.None)]
 		private static void ReplicatedSpecifier(UhtSpecifierContext specifierContext)
 		{
 			UhtPropertySpecifierContext context = (UhtPropertySpecifierContext)specifierContext;
@@ -286,7 +286,7 @@ namespace EpicGames.UHT.Parsers
 			context.PropertySettings.PropertyFlags |= EPropertyFlags.RepNotify;
 		}
 
-		[UhtSpecifier(Extends = UhtTableNames.PropertyMember, ValueType = UhtSpecifierValueType.Legacy)]
+		[UhtSpecifier(Extends = UhtTableNames.PropertyMember, ValueType = UhtSpecifierValueType.None)]
 		private static void NotReplicatedSpecifier(UhtSpecifierContext specifierContext)
 		{
 			UhtPropertySpecifierContext context = (UhtPropertySpecifierContext)specifierContext;
@@ -297,28 +297,28 @@ namespace EpicGames.UHT.Parsers
 			context.PropertySettings.PropertyFlags |= EPropertyFlags.RepSkip;
 		}
 
-		[UhtSpecifier(Extends = UhtTableNames.PropertyMember, ValueType = UhtSpecifierValueType.Legacy)]
+		[UhtSpecifier(Extends = UhtTableNames.PropertyMember, ValueType = UhtSpecifierValueType.None)]
 		private static void RepRetrySpecifier(UhtSpecifierContext specifierContext)
 		{
 			UhtPropertySpecifierContext context = (UhtPropertySpecifierContext)specifierContext;
 			context.MessageSite.LogError("'RepRetry' is deprecated.");
 		}
 
-		[UhtSpecifier(Extends = UhtTableNames.PropertyMember, ValueType = UhtSpecifierValueType.Legacy)]
+		[UhtSpecifier(Extends = UhtTableNames.PropertyMember, ValueType = UhtSpecifierValueType.None)]
 		private static void InterpSpecifier(UhtSpecifierContext specifierContext)
 		{
 			UhtPropertySpecifierContext context = (UhtPropertySpecifierContext)specifierContext;
 			context.PropertySettings.PropertyFlags |= EPropertyFlags.Edit | EPropertyFlags.BlueprintVisible | EPropertyFlags.Interp;
 		}
 
-		[UhtSpecifier(Extends = UhtTableNames.PropertyMember, ValueType = UhtSpecifierValueType.Legacy)]
+		[UhtSpecifier(Extends = UhtTableNames.PropertyMember, ValueType = UhtSpecifierValueType.None)]
 		private static void NonTransactionalSpecifier(UhtSpecifierContext specifierContext)
 		{
 			UhtPropertySpecifierContext context = (UhtPropertySpecifierContext)specifierContext;
 			context.PropertySettings.PropertyFlags |= EPropertyFlags.NonTransactional;
 		}
 
-		[UhtSpecifier(Extends = UhtTableNames.PropertyMember, ValueType = UhtSpecifierValueType.Legacy)]
+		[UhtSpecifier(Extends = UhtTableNames.PropertyMember, ValueType = UhtSpecifierValueType.None)]
 		private static void InstancedSpecifier(UhtSpecifierContext specifierContext)
 		{
 			UhtPropertySpecifierContext context = (UhtPropertySpecifierContext)specifierContext;
@@ -326,84 +326,84 @@ namespace EpicGames.UHT.Parsers
 			context.MetaData.Add(UhtNames.EditInline, true);
 		}
 
-		[UhtSpecifier(Extends = UhtTableNames.PropertyMember, ValueType = UhtSpecifierValueType.Legacy)]
+		[UhtSpecifier(Extends = UhtTableNames.PropertyMember, ValueType = UhtSpecifierValueType.None)]
 		private static void BlueprintAssignableSpecifier(UhtSpecifierContext specifierContext)
 		{
 			UhtPropertySpecifierContext context = (UhtPropertySpecifierContext)specifierContext;
 			context.PropertySettings.PropertyFlags |= EPropertyFlags.BlueprintAssignable;
 		}
 
-		[UhtSpecifier(Extends = UhtTableNames.PropertyMember, ValueType = UhtSpecifierValueType.Legacy)]
+		[UhtSpecifier(Extends = UhtTableNames.PropertyMember, ValueType = UhtSpecifierValueType.None)]
 		private static void BlueprintCallableSpecifier(UhtSpecifierContext specifierContext)
 		{
 			UhtPropertySpecifierContext context = (UhtPropertySpecifierContext)specifierContext;
 			context.PropertySettings.PropertyFlags |= EPropertyFlags.BlueprintCallable;
 		}
 
-		[UhtSpecifier(Extends = UhtTableNames.PropertyMember, ValueType = UhtSpecifierValueType.Legacy)]
+		[UhtSpecifier(Extends = UhtTableNames.PropertyMember, ValueType = UhtSpecifierValueType.None)]
 		private static void BlueprintAuthorityOnlySpecifier(UhtSpecifierContext specifierContext)
 		{
 			UhtPropertySpecifierContext context = (UhtPropertySpecifierContext)specifierContext;
 			context.PropertySettings.PropertyFlags |= EPropertyFlags.BlueprintAuthorityOnly;
 		}
 
-		[UhtSpecifier(Extends = UhtTableNames.PropertyMember, ValueType = UhtSpecifierValueType.Legacy)]
+		[UhtSpecifier(Extends = UhtTableNames.PropertyMember, ValueType = UhtSpecifierValueType.None)]
 		private static void AssetRegistrySearchableSpecifier(UhtSpecifierContext specifierContext)
 		{
 			UhtPropertySpecifierContext context = (UhtPropertySpecifierContext)specifierContext;
 			context.PropertySettings.PropertyFlags |= EPropertyFlags.AssetRegistrySearchable;
 		}
 
-		[UhtSpecifier(Extends = UhtTableNames.PropertyMember, ValueType = UhtSpecifierValueType.Legacy)]
+		[UhtSpecifier(Extends = UhtTableNames.PropertyMember, ValueType = UhtSpecifierValueType.None)]
 		private static void SimpleDisplaySpecifier(UhtSpecifierContext specifierContext)
 		{
 			UhtPropertySpecifierContext context = (UhtPropertySpecifierContext)specifierContext;
 			context.PropertySettings.PropertyFlags |= EPropertyFlags.SimpleDisplay;
 		}
 
-		[UhtSpecifier(Extends = UhtTableNames.PropertyMember, ValueType = UhtSpecifierValueType.Legacy)]
+		[UhtSpecifier(Extends = UhtTableNames.PropertyMember, ValueType = UhtSpecifierValueType.None)]
 		private static void AdvancedDisplaySpecifier(UhtSpecifierContext specifierContext)
 		{
 			UhtPropertySpecifierContext context = (UhtPropertySpecifierContext)specifierContext;
 			context.PropertySettings.PropertyFlags |= EPropertyFlags.AdvancedDisplay;
 		}
 
-		[UhtSpecifier(Extends = UhtTableNames.PropertyMember, ValueType = UhtSpecifierValueType.Legacy)]
+		[UhtSpecifier(Extends = UhtTableNames.PropertyMember, ValueType = UhtSpecifierValueType.None)]
 		private static void SaveGameSpecifier(UhtSpecifierContext specifierContext)
 		{
 			UhtPropertySpecifierContext context = (UhtPropertySpecifierContext)specifierContext;
 			context.PropertySettings.PropertyFlags |= EPropertyFlags.SaveGame;
 		}
 
-		[UhtSpecifier(Extends = UhtTableNames.PropertyMember, ValueType = UhtSpecifierValueType.Legacy)]
+		[UhtSpecifier(Extends = UhtTableNames.PropertyMember, ValueType = UhtSpecifierValueType.None)]
 		private static void SkipSerializationSpecifier(UhtSpecifierContext specifierContext)
 		{
 			UhtPropertySpecifierContext context = (UhtPropertySpecifierContext)specifierContext;
 			context.PropertySettings.PropertyFlags |= EPropertyFlags.SkipSerialization;
 		}
 
-		[UhtSpecifier(Extends = UhtTableNames.PropertyMember, ValueType = UhtSpecifierValueType.Legacy)]
+		[UhtSpecifier(Extends = UhtTableNames.PropertyMember, ValueType = UhtSpecifierValueType.None)]
 		private static void ExperimentalOverridableLogicSpecifier(UhtSpecifierContext specifierContext)
 		{
 			UhtPropertySpecifierContext context = (UhtPropertySpecifierContext)specifierContext;
 			context.PropertySettings.PropertyFlags |= EPropertyFlags.ExperimentalOverridableLogic;
 		}
 
-		[UhtSpecifier(Extends = UhtTableNames.PropertyMember, ValueType = UhtSpecifierValueType.Legacy)]
+		[UhtSpecifier(Extends = UhtTableNames.PropertyMember, ValueType = UhtSpecifierValueType.None)]
 		private static void ExperimentalAlwaysOverridenSpecifier(UhtSpecifierContext specifierContext)
 		{
 			UhtPropertySpecifierContext context = (UhtPropertySpecifierContext)specifierContext;
 			context.PropertySettings.PropertyFlags |= EPropertyFlags.ExperimentalAlwaysOverriden;
 		}
 
-		[UhtSpecifier(Extends = UhtTableNames.PropertyMember, ValueType = UhtSpecifierValueType.Legacy)]
+		[UhtSpecifier(Extends = UhtTableNames.PropertyMember, ValueType = UhtSpecifierValueType.None)]
 		private static void ExperimentalNeverOverridenSpecifier(UhtSpecifierContext specifierContext)
 		{
 			UhtPropertySpecifierContext context = (UhtPropertySpecifierContext)specifierContext;
 			context.PropertySettings.PropertyFlags |= EPropertyFlags.ExperimentalNeverOverriden;
 		}
 
-		[UhtSpecifier(Extends = UhtTableNames.PropertyMember, ValueType = UhtSpecifierValueType.Legacy)]
+		[UhtSpecifier(Extends = UhtTableNames.PropertyMember, ValueType = UhtSpecifierValueType.None)]
 		private static void ExperimentalExternalObjectsSpecifier(UhtSpecifierContext specifierContext)
 		{
 			UhtPropertySpecifierContext context = (UhtPropertySpecifierContext)specifierContext;
diff --git a/Engine/Source/Programs/Shared/EpicGames.UHT/Specifiers/UhtScriptStructSpecifiers.cs b/Engine/Source/Programs/Shared/EpicGames.UHT/Specifiers/UhtScriptStructSpecifiers.cs
--- a/Engine/Source/Programs/Shared/EpicGames.UHT/Specifiers/UhtScriptStructSpecifiers.cs
+++ b/Engine/Source/Programs/Shared/EpicGames.UHT/Specifiers/UhtScriptStructSpecifiers.cs
@@ -15,7 +15,7 @@ namespace EpicGames.UHT.Parsers
 	[SuppressMessage("CodeQuality", "IDE0051:Remove unused private members", Justification = "Attribute accessed method")]
 	public static class UhtScriptStructSpecifiers
 	{
-		[UhtSpecifier(Extends = UhtTableNames.ScriptStruct, ValueType = UhtSpecifierValueType.Legacy)]
+		[UhtSpecifier(Extends = UhtTableNames.ScriptStruct, ValueType = UhtSpecifierValueType.None)]
 		private static void NoExportSpecifier(UhtSpecifierContext specifierContext)
 		{
 			UhtScriptStruct scriptStruct = (UhtScriptStruct)specifierContext.Type;
@@ -23,14 +23,14 @@ namespace EpicGames.UHT.Parsers
 			scriptStruct.ScriptStructFlags &= ~EStructFlags.Native;
 		}
 
-		[UhtSpecifier(Extends = UhtTableNames.ScriptStruct, ValueType = UhtSpecifierValueType.Legacy)]
+		[UhtSpecifier(Extends = UhtTableNames.ScriptStruct, ValueType = UhtSpecifierValueType.None)]
 		private static void AtomicSpecifier(UhtSpecifierContext specifierContext)
 		{
 			UhtScriptStruct scriptStruct = (UhtScriptStruct)specifierContext.Type;
 			scriptStruct.ScriptStructFlags |= EStructFlags.Atomic;
 		}
 
-		[UhtSpecifier(Extends = UhtTableNames.ScriptStruct, ValueType = UhtSpecifierValueType.Legacy)]
+		[UhtSpecifier(Extends = UhtTableNames.ScriptStruct, ValueType = UhtSpecifierValueType.None)]
 		[SuppressMessage("CodeQuality", "IDE0051:Remove unused private members", Justification = "Attribute accessed method")]
 		private static void ImmutableSpecifier(UhtSpecifierContext specifierContext)
 		{
@@ -38,7 +38,7 @@ namespace EpicGames.UHT.Parsers
 			scriptStruct.ScriptStructFlags |= EStructFlags.Atomic | EStructFlags.Immutable;
 		}
 
-		[UhtSpecifier(Extends = UhtTableNames.ScriptStruct, ValueType = UhtSpecifierValueType.Legacy)]
+		[UhtSpecifier(Extends = UhtTableNames.ScriptStruct, ValueType = UhtSpecifierValueType.None)]
 		private static void HasDefaultsSpecifier(UhtSpecifierContext specifierContext)
 		{
 			if (!specifierContext.Type.HeaderFile.IsNoExportTypes)
@@ -49,7 +49,7 @@ namespace EpicGames.UHT.Parsers
 			scriptStruct.ScriptStructExportFlags |= UhtScriptStructExportFlags.HasDefaults;
 		}
 
-		[UhtSpecifier(Extends = UhtTableNames.ScriptStruct, ValueType = UhtSpecifierValueType.Legacy)]
+		[UhtSpecifier(Extends = UhtTableNames.ScriptStruct, ValueType = UhtSpecifierValueType.None)]
 		private static void HasNoOpConstructorSpecifier(UhtSpecifierContext specifierContext)
 		{
 			if (!specifierContext.Type.HeaderFile.IsNoExportTypes)
@@ -60,7 +60,7 @@ namespace EpicGames.UHT.Parsers
 			scriptStruct.ScriptStructExportFlags |= UhtScriptStructExportFlags.HasNoOpConstructor;
 		}
 
-		[UhtSpecifier(Extends = UhtTableNames.ScriptStruct, ValueType = UhtSpecifierValueType.Legacy)]
+		[UhtSpecifier(Extends = UhtTableNames.ScriptStruct, ValueType = UhtSpecifierValueType.None)]
 		private static void IsAlwaysAccessibleSpecifier(UhtSpecifierContext specifierContext)
 		{
 			if (!specifierContext.Type.HeaderFile.IsNoExportTypes)
@@ -71,7 +71,7 @@ namespace EpicGames.UHT.Parsers
 			scriptStruct.ScriptStructExportFlags |= UhtScriptStructExportFlags.IsAlwaysAccessible;
 		}
 
-		[UhtSpecifier(Extends = UhtTableNames.ScriptStruct, ValueType = UhtSpecifierValueType.Legacy)]
+		[UhtSpecifier(Extends = UhtTableNames.ScriptStruct, ValueType = UhtSpecifierValueType.None)]
 		private static void IsCoreTypeSpecifier(UhtSpecifierContext specifierContext)
 		{
 			if (!specifierContext.Type.HeaderFile.IsNoExportTypes)
@@ -82,7 +82,7 @@ namespace EpicGames.UHT.Parsers
 			scriptStruct.ScriptStructExportFlags |= UhtScriptStructExportFlags.IsCoreType;
 		}
 
-		[UhtSpecifier(Extends = UhtTableNames.ScriptStruct, ValueType = UhtSpecifierValueType.Legacy)]
+		[UhtSpecifier(Extends = UhtTableNames.ScriptStruct, ValueType = UhtSpecifierValueType.None)]
 		private static void MinimalAPISpecifier(UhtSpecifierContext specifierContext)
 		{
 			UhtScriptStruct scriptStruct = (UhtScriptStruct)specifierContext.Type;
diff --git a/Engine/Source/Programs/Shared/EpicGames.UHT/Tables/UhtSpecifierTable.cs b/Engine/Source/Programs/Shared/EpicGames.UHT/Tables/UhtSpecifierTable.cs
--- a/Engine/Source/Programs/Shared/EpicGames.UHT/Tables/UhtSpecifierTable.cs
+++ b/Engine/Source/Programs/Shared/EpicGames.UHT/Tables/UhtSpecifierTable.cs
@@ -65,6 +65,7 @@ namespace EpicGames.UHT.Tables
 		/// <summary>
 		/// Accepts a string list but the value is ignored by the specifier and is automatically deferred.  This is for legacy UHT support.
 		/// </summary>
+		[Obsolete("Legacy specifier mode is obsolete, please use None")]
 		Legacy,
 	}
 
@@ -620,7 +621,9 @@ namespace EpicGames.UHT.Tables
 				case UhtSpecifierValueType.NonEmptyStringList:
 					table.Add(new UhtSpecifierNonEmptyStringList(name, specifierAttribute.When, (UhtSpecifierNonEmptyStringListDelegate)Delegate.CreateDelegate(typeof(UhtSpecifierNonEmptyStringListDelegate), methodInfo)));
 					break;
+#pragma warning disable CS0618 // Type or member is obsolete
 				case UhtSpecifierValueType.Legacy:
+#pragma warning restore CS0618 // Type or member is obsolete
 					table.Add(new UhtSpecifierLegacy(name, (UhtSpecifierLegacyDelegate)Delegate.CreateDelegate(typeof(UhtSpecifierLegacyDelegate), methodInfo)));
 					break;
 			}
