Add support for dynamically disabling named log event matchers by sending log lines of a specific format to the log event parser.

#Synced-CL 49616476

diff --git a/Engine/Source/Programs/Shared/EpicGames.Core.Tests/LogEventParserTests.cs b/Engine/Source/Programs/Shared/EpicGames.Core.Tests/LogEventParserTests.cs
new file mode 100644
--- /dev/null
+++ b/Engine/Source/Programs/Shared/EpicGames.Core.Tests/LogEventParserTests.cs
@@ -0,0 +1,165 @@
+// Copyright Epic Games, Inc. All Rights Reserved.
+
+using System;
+using System.Collections.Generic;
+using System.Linq;
+using Microsoft.Extensions.Logging;
+using Microsoft.VisualStudio.TestTools.UnitTesting;
+
+namespace EpicGames.Core.Tests;
+
+[TestClass]
+public class LogEventParserTests
+{
+	[TestMethod]
+	public void TestAllMatchersAreTried()
+	{
+		LoggerCapture capture = new();
+		TestMatcher highest = new(LogEventPriority.Highest, LogLevel.Critical);
+		TestMatcher lowest = new(LogEventPriority.Lowest, LogLevel.Debug);
+
+		const string LineToLog = "hello";
+
+		using (LogEventParser parser = new(capture))
+		{
+			parser.AddMatchers([highest, lowest]);
+			parser.WriteLine(LineToLog);
+		}
+
+		Assert.AreEqual(1, highest.MatchCalls);
+		Assert.AreEqual(1, lowest.MatchCalls);
+	}
+
+	[TestMethod]
+	public void TestHighestPriorityMatcherWins()
+	{
+		LoggerCapture capture = new();
+		TestMatcher highest = new(LogEventPriority.Highest, LogLevel.Critical);
+		TestMatcher lowest = new(LogEventPriority.Lowest, LogLevel.Debug);
+
+		const string LineToLog = "hello";
+
+		using (LogEventParser parser = new(capture))
+		{
+			parser.AddMatchers([highest, lowest]);
+			parser.WriteLine(LineToLog);
+		}
+
+		Assert.AreEqual(1, capture.Events.Count);
+		Assert.AreEqual(LineToLog, capture.Events[0].Message);
+		Assert.AreEqual(highest.Level, capture.Events[0].Level);
+	}
+
+	[TestMethod]
+	public void TestTogglingMatchers()
+	{
+		LoggerCapture capture = new();
+		TestMatcher highest = new(LogEventPriority.Highest, LogLevel.Critical);
+		TestMatcher lowest = new(LogEventPriority.Lowest, LogLevel.Debug);
+
+		const string LineToLog = "hello";
+
+		using LogEventParser parser = new(capture);
+		parser.AddMatchers([highest, lowest]);
+
+		parser.WriteLine(LineToLog);
+
+		Assert.AreEqual(1, highest.MatchCalls);
+		Assert.AreEqual(1, lowest.MatchCalls);
+
+		parser.WriteLine("<-- logeventparser suspend highest -->");
+		parser.WriteLine(LineToLog);
+
+		// Highest should not have gotten another call.
+		Assert.AreEqual(1, highest.MatchCalls);
+		Assert.AreEqual(2, lowest.MatchCalls);
+
+		Assert.AreEqual(2, capture.Events.Count);
+
+		Assert.AreEqual(LineToLog, capture.Events[^1].Message);
+		// Lowest should have matched the event.
+		Assert.AreEqual(lowest.Level, capture.Events[^1].Level);
+
+		parser.WriteLine("<-- logeventparser resume highest -->");
+		parser.WriteLine(LineToLog);
+
+		// Highest should have gotten another call after being restored.
+		Assert.AreEqual(2, highest.MatchCalls);
+		Assert.AreEqual(3, lowest.MatchCalls);
+
+		Assert.AreEqual(3, capture.Events.Count);
+
+		Assert.AreEqual(LineToLog, capture.Events[^1].Message);
+		// Highest should have matched the event.
+		Assert.AreEqual(highest.Level, capture.Events[^1].Level);
+	}
+
+	[TestMethod]
+	public void TestTogglingAllMatchers()
+	{
+		LoggerCapture capture = new();
+		TestMatcher highest = new(LogEventPriority.Highest, LogLevel.Critical);
+		TestMatcher lowest = new(LogEventPriority.Lowest, LogLevel.Debug);
+
+		const string LineToLog = "hello";
+
+		using LogEventParser parser = new(capture);
+		parser.AddMatchers([highest, lowest]);
+
+		parser.WriteLine("<-- logeventparser suspend highest -->");
+		parser.WriteLine("<-- logeventparser suspend lowest -->");
+		parser.WriteLine(LineToLog);
+
+		// Neither should have been called.
+		Assert.AreEqual(0, highest.MatchCalls);
+		Assert.AreEqual(0, lowest.MatchCalls);
+
+		// We should still have one event.
+		Assert.AreEqual(1, capture.Events.Count);
+		// The event should have a zero ID.
+		Assert.AreEqual(new EventId(0), capture.Events[0].Id);
+	}
+
+	private const string LogLineProperty = "LogLine";
+
+	private class LoggerCapture : ILogger
+	{
+		public IDisposable? BeginScope<TState>(TState state) where TState : notnull => null!;
+
+		public bool IsEnabled(LogLevel logLevel) => true;
+
+		public void Log<TState>(LogLevel logLevel, EventId eventId, TState state, Exception? exception, Func<TState, Exception?, string> formatter)
+		{
+			LogEvent logEvent = LogEvent.Read(JsonLogEvent.FromLoggerState(logLevel, eventId, state, exception, formatter).Data.Span);
+			KeyValuePair<string, object?>[] items = [new KeyValuePair<string, object?>(LogLineProperty, _logLineIndex)];
+			logEvent.Properties = (logEvent.Properties == null) ? items : Enumerable.Concat(logEvent.Properties, items);
+			Events.Add(logEvent);
+
+			_logLineIndex++;
+		}
+
+		public List<LogEvent> Events { get; } = [];
+
+		private int _logLineIndex;
+	}
+
+	private class TestMatcher(LogEventPriority priority, LogLevel level) : ILogEventMatcher
+	{
+		public LogEventMatch? Match(ILogCursor cursor)
+		{
+			++MatchCalls;
+
+			string? line = cursor.CurrentLine;
+			Assert.IsNotNull(line);
+
+			LogEventBuilder builder = new(cursor);
+			return builder.ToMatch(LogEventPriority.Highest, Level, new EventId(999));
+		}
+
+		public LogEventPriority Priority { get; } = priority;
+		public LogLevel Level { get; } = level;
+		public int MatchCalls { get; private set; } = 0;
+
+		public string GetMatcherName() => Priority.ToString();
+	}
+}
diff --git a/Engine/Source/Programs/Shared/EpicGames.Core/LogEventParser.cs b/Engine/Source/Programs/Shared/EpicGames.Core/LogEventParser.cs
--- a/Engine/Source/Programs/Shared/EpicGames.Core/LogEventParser.cs
+++ b/Engine/Source/Programs/Shared/EpicGames.Core/LogEventParser.cs
@@ -103,23 +103,53 @@ namespace EpicGames.Core
 		/// <param name="cursor">The input buffer</param>
 		/// <returns>Information about the error that was matched, or null if an error was not matched</returns>
 		LogEventMatch? Match(ILogCursor cursor);
+
+		/// <summary>
+		/// Get a name for this matcher, used to support dynamic matcher toggling. Names must not contain whitespace characters.
+		/// </summary>
+		/// <remarks>
+		/// This is an interface method with a default interface to allow for edge-cases like multiple matchers of the same type
+		/// but differing behaviour, matchers that want to be enabled or disabled together, inheritance etc.
+		/// </remarks>
+		string GetMatcherName() => GetType().Name;
 	}
 
 	/// <summary>
 	/// Turns raw text output into structured logging events
 	/// </summary>
-	public class LogEventParser : IDisposable
+	/// <remarks>
+	/// Log events matching the expected format will be processed by the parser, taking the specified action if possible,
+	/// and may not be output to users. The format is specified in <see cref="MatcherToggleRegex"/>, in addition to the commands
+	/// <c>&lt;-- Suspend Log Parsing --&gt;</c> and <c>&lt;-- Resume Log Parsing --&gt;</c>.
+	/// </remarks>
+	public partial class LogEventParser : IDisposable
 	{
 		/// <summary>
-		/// List of event matchers for this parser
+		/// A list of enabled event matchers for this parser.
 		/// </summary>
-		public List<ILogEventMatcher> Matchers { get; } = [];
+		public IEnumerable<ILogEventMatcher> EnabledMatchers => [.. _enabledMatchers];
 
 		/// <summary>
 		/// List of patterns to ignore
 		/// </summary>
 		public List<Regex> IgnorePatterns { get; } = [];
 
+		/// <summary>
+		/// A list of matchers all matchers, including disabled matchers.
+		/// </summary>
+		private readonly List<ILogEventMatcher> _matchers = [];
+
+		/// <summary>
+		/// A list of enabled matchers, to allow quicker iteration when parsing log lines.
+		/// It's expected that we'll parse significantly more log lines than we will toggle matcher states.
+		/// </summary>
+		private readonly List<ILogEventMatcher> _enabledMatchers = [];
+
+		/// <summary>
+		/// A mapping of matcher names to whether they're enabled. Items not present are enabled by default.
+		/// </summary>
+		private readonly Dictionary<string, bool> _matcherNameToEnabled = new(StringComparer.OrdinalIgnoreCase);
+
 		/// <summary>
 		/// Buffer of input lines
 		/// </summary>
@@ -169,6 +199,12 @@ namespace EpicGames.Core
 			set => _logger = value;
 		}
 
+		/// <summary>
+		/// The pattern matcher toggle commands must conform to.
+		/// </summary>
+		[GeneratedRegex("^<-- logeventparser (?<Action>suspend|resume) (?<MatcherName>\\S+) -->$", RegexOptions.IgnoreCase | RegexOptions.ExplicitCapture)]
+		public static partial Regex MatcherToggleRegex();
+
 		/// <summary>
 		/// Constructor
 		/// </summary>
@@ -210,7 +246,10 @@ namespace EpicGames.Core
 		/// <param name="matchers">The matchers to add</param>
 		public void AddMatchers(IEnumerable<ILogEventMatcher> matchers)
 		{
-			Matchers.AddRange(matchers.Where(m => m is not null));
+			foreach (ILogEventMatcher matcher in matchers)
+			{
+				AddMatcher(matcher);
+			}
 		}
 
 		/// <summary>
@@ -224,7 +263,7 @@ namespace EpicGames.Core
 				if (type.IsClass && typeof(ILogEventMatcher).IsAssignableFrom(type))
 				{
 					ILogEventMatcher matcher = (ILogEventMatcher)Activator.CreateInstance(type)!;
-					Matchers.Add(matcher);
+					AddMatcher(matcher);
 				}
 			}
 		}
@@ -403,6 +442,53 @@ namespace EpicGames.Core
 			ProcessData(true);
 		}
 
+		/// <summary>
+		/// Add the new matcher, making sure its enabled state is correct.
+		/// </summary>
+		private void AddMatcher(ILogEventMatcher matcher)
+		{
+			if (matcher is null)
+			{
+				return;
+			}
+
+			bool enabled;
+			if (!_matcherNameToEnabled.TryGetValue(matcher.GetMatcherName(), out enabled))
+			{
+				enabled = true;
+			}
+
+			_matchers.Add(matcher);
+			if (enabled)
+			{
+				_enabledMatchers.Add(matcher);
+			}
+		}
+
+		/// <summary>
+		/// Set the given matcher to the given enabled state.
+		/// </summary>
+		private void ToggleMatcher(string matcherName, bool enabled)
+		{
+			_matcherNameToEnabled[matcherName] = enabled;
+			RecalculateEnabledMatchers();
+		}
+
+		/// <summary>
+		/// Clear and repopulate the list of enabled matchers.
+		/// </summary>
+		private void RecalculateEnabledMatchers()
+		{
+			_enabledMatchers.Clear();
+			foreach (ILogEventMatcher matcher in _matchers)
+			{
+				if (!_matcherNameToEnabled.TryGetValue(matcher.GetMatcherName(), out bool enabled) || enabled)
+				{
+					_enabledMatchers.Add(matcher);
+				}
+			}
+		}
+
 		/// <summary>
 		/// Adds a raw utf-8 string to the buffer
 		/// </summary>
@@ -453,6 +539,16 @@ namespace EpicGames.Core
 				{
 					line = "";
 				}
+
+				Match match = MatcherToggleRegex().Match(line);
+				if (match.Success)
+				{
+					bool enabled = match.Groups["Action"].Value.Equals("resume", StringComparison.OrdinalIgnoreCase);
+					ToggleMatcher(match.Groups["MatcherName"].Value, enabled);
+					_buffer.MoveNext();
+					continue;
+				}
+
 				if (Regex.IsMatch(line, "<-- Suspend Log Parsing -->", RegexOptions.IgnoreCase))
 				{
 					_matchingEnabled--;
@@ -527,7 +623,7 @@ namespace EpicGames.Core
 		private List<LogEvent>? MatchEvent()
 		{
 			LogEventMatch? currentMatch = null;
-			foreach (ILogEventMatcher matcher in Matchers)
+			foreach (ILogEventMatcher matcher in _enabledMatchers)
 			{
 				LogEventMatch? match = null;
 				try
diff --git a/Engine/Source/Programs/Shared/EpicGames.MsBuild/CsProjBuilder.cs b/Engine/Source/Programs/Shared/EpicGames.MsBuild/CsProjBuilder.cs
--- a/Engine/Source/Programs/Shared/EpicGames.MsBuild/CsProjBuilder.cs
+++ b/Engine/Source/Programs/Shared/EpicGames.MsBuild/CsProjBuilder.cs
@@ -435,7 +435,7 @@ namespace EpicGames.MsBuild
 			startInfo.EnvironmentVariables["DOTNET_MULTILEVEL_LOOKUP"] = "0"; // use only the bundled dotnet installation - ignore any other/system dotnet install
 
 			using LogEventParser parser = new LogEventParser(logger);
-			parser.Matchers.AddRange(Log.EventParser.Matchers);
+			parser.AddMatchers(Log.EventParser.EnabledMatchers);
 
 			logger.LogDebug("Running: {Application} {Arguments}", startInfo.FileName, startInfo.Arguments);
 			using Process dotnetProcess = new Process();
