EpicGames.MsBuild: Prevent reimporting UnrealEngine.CSharp.targets

#rnx
#Synced-CL 49631673

diff --git a/Engine/Source/Programs/Shared/EpicGames.MsBuild/CsProjBuilder.cs b/Engine/Source/Programs/Shared/EpicGames.MsBuild/CsProjBuilder.cs
--- a/Engine/Source/Programs/Shared/EpicGames.MsBuild/CsProjBuilder.cs
+++ b/Engine/Source/Programs/Shared/EpicGames.MsBuild/CsProjBuilder.cs
@@ -90,6 +90,7 @@ namespace EpicGames.MsBuild
 				{ "EnginePath", hook.EngineDirectory.FullName },
 				{ "EpicGamesMsBuild", "true" },
 				{ "CustomAfterMicrosoftCommonProps", csharpTargetsFile.FullName },
+				{ "UnrealEngineCSharpTargetsImported", "true" },
 				{ "NoWarn", "$(NoWarn);MSB3026;NETSDK1206".Replace(";", "%3B", StringComparison.Ordinal) },
 #if DEBUG
 				{ "Configuration", "Debug" },
