Remove some unused elements in UHT

#rnx

#Synced-CL 49619778

diff --git a/Engine/Source/Programs/Shared/EpicGames.UHT/Exporters/CodeGen/UhtHeaderCodeGenerator.cs b/Engine/Source/Programs/Shared/EpicGames.UHT/Exporters/CodeGen/UhtHeaderCodeGenerator.cs
--- a/Engine/Source/Programs/Shared/EpicGames.UHT/Exporters/CodeGen/UhtHeaderCodeGenerator.cs
+++ b/Engine/Source/Programs/Shared/EpicGames.UHT/Exporters/CodeGen/UhtHeaderCodeGenerator.cs
@@ -13,7 +13,6 @@ namespace EpicGames.UHT.Exporters.CodeGen
 	class UhtHeaderCodeGenerator : UhtPackageCodeGenerator
 	{
 		#region Define block macro names
-		public const string EventParamsMacroSuffix = "EVENT_PARMS";
 		public const string CallbackWrappersMacroSuffix = "CALLBACK_WRAPPERS";
 		public const string SparseDataMacroSuffix = "SPARSE_DATA";
 		public const string SparseDataPropertyAccessorsMacroSuffix = "SPARSE_DATA_PROPERTY_ACCESSORS";
@@ -33,7 +32,6 @@ namespace EpicGames.UHT.Exporters.CodeGen
 		public const string InClassIInterfaceMacroSuffix = "INCLASS_IINTERFACE";
 		public const string InClassIInterfaceNoPureDeclsMacroSuffix = "INCLASS_IINTERFACE_NO_PURE_DECLS";
 		public const string PrologMacroSuffix = "PROLOG";
-		public const string DelegateMacroSuffix = "DELEGATE";
 		public const string AutoGettersSettersMacroSuffix = "AUTOGETTERSETTER_DECLS";
 		public const string VInterfacesMacroSuffix = "VINTERFACES";
 		#endregion
@@ -602,27 +600,6 @@ namespace EpicGames.UHT.Exporters.CodeGen
 				}
 			}
 		}
-
-		protected static string GetDelegateFunctionExportName(UhtFunction function)
-		{
-			const string DelegatePrefix = "delegate";
-
-			if (!function.FunctionFlags.HasAnyFlags(EFunctionFlags.Delegate))
-			{
-				throw new UhtIceException("Attempt to export a function that isn't a delegate as a delegate");
-			}
-			if (!function.MarshalAndCallName.StartsWith(DelegatePrefix, StringComparison.Ordinal))
-			{
-				throw new UhtIceException("Marshal and call name must begin with 'delegate'");
-			}
-			return $"F{function.MarshalAndCallName[DelegatePrefix.Length..]}_DelegateWrapper";
-		}
-
-		protected static string GetDelegateFunctionExtraParameter(UhtFunction function)
-		{
-			string returnType = function.FunctionFlags.HasAnyFlags(EFunctionFlags.MulticastDelegate) ? "FMulticastScriptDelegate" : "FScriptDelegate";
-			return $"const {returnType}& {function.SourceName[1..]}";
-		}
 		#endregion
 
 		#region Field notify support
diff --git a/Engine/Source/Programs/Shared/EpicGames.UHT/Exporters/CodeGen/UhtHeaderCodeGeneratorCppFile.cs b/Engine/Source/Programs/Shared/EpicGames.UHT/Exporters/CodeGen/UhtHeaderCodeGeneratorCppFile.cs
--- a/Engine/Source/Programs/Shared/EpicGames.UHT/Exporters/CodeGen/UhtHeaderCodeGeneratorCppFile.cs
+++ b/Engine/Source/Programs/Shared/EpicGames.UHT/Exporters/CodeGen/UhtHeaderCodeGeneratorCppFile.cs
@@ -352,7 +352,7 @@ namespace EpicGames.UHT.Exporters.CodeGen
 					{
 						using UhtCodeBlockComment blockComment = new(builder, field);
 						using UhtMacroBlockEmitter macroBlockEmitter = new(builder, UhtDefineScopeNames.Standard, function.DefineScope);
-						AppendDelegate(builder, function);
+						AppendFunction(builder, function, false, false);
 					}
 					else if (field is UhtClass classObj)
 					{
@@ -1278,11 +1278,6 @@ namespace EpicGames.UHT.Exporters.CodeGen
 			return builder;
 		}
 
-		private StringBuilder AppendDelegate(StringBuilder builder, UhtFunction function)
-		{
-			return AppendFunction(builder, function, false, false);
-		}
-
 		private StringBuilder AppendFunction(StringBuilder builder, UhtFunction function, bool isNoExport, bool isVerseNativeCallable)
 		{
 			const string MetaDataParamsName = "Function_MetaDataParams";
