AutomationTool: Force rebuild of script projects if build configuration or dotnet runtime changes

#rnx

#Synced-CL 49665286

diff --git a/Engine/Source/Programs/Shared/EpicGames.Build/System/CompileScriptModules.cs b/Engine/Source/Programs/Shared/EpicGames.Build/System/CompileScriptModules.cs
--- a/Engine/Source/Programs/Shared/EpicGames.Build/System/CompileScriptModules.cs
+++ b/Engine/Source/Programs/Shared/EpicGames.Build/System/CompileScriptModules.cs
@@ -526,6 +526,20 @@ namespace UnrealBuildBase
 				return false;
 			}
 
+			if (BuildRecord.RuntimeVersion != Environment.Version)
+			{
+				ValidationFailureMessage =
+					$"runtime version does not match: build record has version {BuildRecord.RuntimeVersion}; current version is {Environment.Version}";
+				return false;
+			}
+
+			if (BuildRecord.TargetConfiguration != CsProjBuildRecord.CurrentConfiguration)
+			{
+				ValidationFailureMessage =
+					$"recorded target configuration ({BuildRecord.TargetConfiguration}) does not match ({CsProjBuildRecord.CurrentConfiguration})";
+				return false;
+			}
+
 			string TargetRelativePath = Path.GetRelativePath(Unreal.EngineDirectory.FullName, BuildRecord.TargetPath!);
 			string NormalizedPath = Path.GetFullPath(BuildRecord.TargetPath!, ProjectDirectory.FullName);
 			if (!File.Exists(NormalizedPath))
diff --git a/Engine/Source/Programs/Shared/EpicGames.Core/CsProjBuildRecord.cs b/Engine/Source/Programs/Shared/EpicGames.Core/CsProjBuildRecord.cs
--- a/Engine/Source/Programs/Shared/EpicGames.Core/CsProjBuildRecord.cs
+++ b/Engine/Source/Programs/Shared/EpicGames.Core/CsProjBuildRecord.cs
@@ -36,13 +36,28 @@ namespace EpicGames.Core
 		/// <summary>
 		/// Version number making it possible to quickly invalidate written records.
 		/// </summary>
-		public static readonly int CurrentVersion = 6;
+		public static readonly int CurrentVersion = 7;
+
+		/// <summary>
+		/// Build configuration for valid written records
+		/// </summary>
+		public static readonly string CurrentConfiguration
+#if DEBUG
+		= "Debug";
+#else
+		= "Development";
+#endif
 
 		/// <summary>
 		/// Version number for this record
 		/// </summary>
 		public int Version { get; set; } = -1;
 
+		/// <summary>
+		/// .NET runtime version when the target assembly was built
+		/// </summary>
+		public Version RuntimeVersion { get; set; } = Environment.Version;
+
 		/// <summary>
 		/// Path to the .csproj project file, relative to the location of the build record .json file
 		/// </summary>
@@ -53,6 +68,11 @@ namespace EpicGames.Core
 		/// </summary>
 		public DateTime TargetBuildTime { get; set; }
 
+		/// <summary>
+		/// The build configuration for the target assembly when built
+		/// </summary>
+		public string TargetConfiguration { get; set; } = CurrentConfiguration;
+
 		// all following paths are relative to the project directory, the directory containing ProjectPath 
 
 		/// <summary>
