Horde: Add GetJobsAsync to HordeHttpClient

diff --git a/Engine/Source/Programs/Shared/EpicGames.Horde/HordeHttpClient.cs b/Engine/Source/Programs/Shared/EpicGames.Horde/HordeHttpClient.cs
--- a/Engine/Source/Programs/Shared/EpicGames.Horde/HordeHttpClient.cs
+++ b/Engine/Source/Programs/Shared/EpicGames.Horde/HordeHttpClient.cs
@@ -8,6 +8,7 @@ using System.Net;
 using System.Net.Http;
 using System.Net.Http.Headers;
 using System.Net.Http.Json;
+using System.Text;
 using System.Text.Json;
 using System.Text.RegularExpressions;
 using System.Threading;
@@ -20,6 +21,7 @@ using EpicGames.Horde.Dashboard;
 using EpicGames.Horde.Devices;
 using EpicGames.Horde.Jobs;
 using EpicGames.Horde.Jobs.Graphs;
+using EpicGames.Horde.Jobs.Templates;
 using EpicGames.Horde.Logs;
 using EpicGames.Horde.Projects;
 using EpicGames.Horde.Secrets;
@@ -583,6 +585,71 @@ namespace EpicGames.Horde
 			return PutAsync(_httpClient, $"api/v1/jobs/{id}/metadata", new PutJobMetadataRequest() { JobMetaData = jobMetaData?.ToList(), StepMetaData = stepMetaData }, cancellationToken);
 		}
 
+		/// <summary>
+		/// Get a list of job information for given stream ID and template IDs. Returns maximum count results or fewer.
+		/// </summary>
+		/// <param name="streamId">Stream ID to filter jobs</param>
+		/// <param name="templateIds">List of job template IDs to filter by</param>
+		/// <param name="filters">A list of filter fields to return for each Job in the results. Defaults to id, name, change, preflightChange, createTime</param>
+		/// <param name="minChange">Minimum commit number</param>
+		/// <param name="maxChange">Maximum commit number</param>
+		/// <param name="minCreateTime">Minimum job creation time</param>
+		/// <param name="maxCreateTime">Maximum job creation time</param>
+		/// <param name="includePreflight">Whether to include preflight jobs</param>		
+		/// <param name="count">Maximum number of results to return</param>
+		/// <param name="cancellationToken">Cancellation token for the operation</param>
+		/// <returns>GetJobResponse List</returns>
+		public Task<List<GetJobResponse>> GetJobsAsync(
+			StreamId streamId,
+			IEnumerable<TemplateId>? templateIds = null,
+			List<string>? filters = null,
+			CommitId? minChange = null,
+			CommitId? maxChange = null,
+			DateTimeOffset? minCreateTime = null,
+			DateTimeOffset? maxCreateTime = null,
+			bool includePreflight = false,			
+			int count = 100,
+			CancellationToken cancellationToken = default)
+		{
+			StringBuilder queryBuilder = new StringBuilder($"/api/v1/jobs?streamId={streamId}&count={count}&includePreflight={(includePreflight ? "true" : "false")}");
+
+			if (templateIds != null)
+			{
+				foreach (TemplateId templateId in templateIds)
+				{
+					queryBuilder.Append($"&t={templateId}");
+				}
+			}
+
+			if (minChange != null)
+			{
+				queryBuilder.Append($"&minChange={minChange}");
+			}
+
+			if (maxChange != null)
+			{
+				queryBuilder.Append($"&maxChange={maxChange}");
+			}
+
+			if (minCreateTime != null)
+			{
+				queryBuilder.Append($"&minCreateTime={Uri.EscapeDataString(minCreateTime.Value.ToString("o"))}");
+			}
+
+			if (maxCreateTime != null)
+			{
+				queryBuilder.Append($"&maxCreateTime={Uri.EscapeDataString(maxCreateTime.Value.ToString("o"))}");
+			}
+
+			filters = filters ?? ["id", "name", "change", "preflightChange", "createTime"];
+			if (filters.Count > 0)
+			{
+				queryBuilder.Append($"&filter={Uri.EscapeDataString(String.Join(",", filters))}");
+			}
+
+			return GetAsync<List<GetJobResponse>>(_httpClient, queryBuilder.ToString(), cancellationToken);
+		}
+		
 		#endregion
 
 		#region Log
