Add support for layered matcher toggles to LogEventParser, to support the use-case of disabling all-but-one matcher.

#Synced-CL 49723634

diff --git a/Engine/Source/Programs/Shared/EpicGames.Core.Tests/LogEventParserTests.cs b/Engine/Source/Programs/Shared/EpicGames.Core.Tests/LogEventParserTests.cs
--- a/Engine/Source/Programs/Shared/EpicGames.Core.Tests/LogEventParserTests.cs
+++ b/Engine/Source/Programs/Shared/EpicGames.Core.Tests/LogEventParserTests.cs
@@ -18,8 +18,6 @@ public class LogEventParserTests
 		TestMatcher highest = new(LogEventPriority.Highest, LogLevel.Critical);
 		TestMatcher lowest = new(LogEventPriority.Lowest, LogLevel.Debug);
 
-		const string LineToLog = "hello";
-
 		using (LogEventParser parser = new(capture))
 		{
 			parser.AddMatchers([highest, lowest]);
@@ -37,8 +35,6 @@ public class LogEventParserTests
 		TestMatcher highest = new(LogEventPriority.Highest, LogLevel.Critical);
 		TestMatcher lowest = new(LogEventPriority.Lowest, LogLevel.Debug);
 
-		const string LineToLog = "hello";
-
 		using (LogEventParser parser = new(capture))
 		{
 			parser.AddMatchers([highest, lowest]);
@@ -51,7 +47,7 @@ public class LogEventParserTests
 	}
 
 	[TestMethod]
-	public void TestTogglingMatchers()
+	public void TestTogglingOneMatcherOnce()
 	{
 		LoggerCapture capture = new();
 		TestMatcher highest = new(LogEventPriority.Highest, LogLevel.Critical);
@@ -70,6 +66,7 @@ public class LogEventParserTests
 		parser.WriteLine("<-- logeventparser suspend highest -->");
 		parser.WriteLine(LineToLog);
 
+		CollectionAssert.AreEquivalent<ILogEventMatcher>(parser.EnabledMatchers, [lowest], ReferenceEqualityComparer.Instance);
 		// Highest should not have gotten another call.
 		Assert.AreEqual(1, highest.MatchCalls);
 		Assert.AreEqual(2, lowest.MatchCalls);
@@ -95,14 +92,12 @@ public class LogEventParserTests
 	}
 
 	[TestMethod]
-	public void TestTogglingAllMatchers()
+	public void TestTogglingAllMatchersIndividuallyOnce()
 	{
 		LoggerCapture capture = new();
 		TestMatcher highest = new(LogEventPriority.Highest, LogLevel.Critical);
 		TestMatcher lowest = new(LogEventPriority.Lowest, LogLevel.Debug);
 
-		const string LineToLog = "hello";
-
 		using LogEventParser parser = new(capture);
 		parser.AddMatchers([highest, lowest]);
 
@@ -110,6 +105,30 @@ public class LogEventParserTests
 		parser.WriteLine("<-- logeventparser suspend lowest -->");
 		parser.WriteLine(LineToLog);
 
+		Assert.AreEqual(0, parser.EnabledMatchers.Count());
+		// Neither should have been called.
+		Assert.AreEqual(0, highest.MatchCalls);
+		Assert.AreEqual(0, lowest.MatchCalls);
+
+		// We should still have one event.
+		Assert.AreEqual(1, capture.Events.Count);
+		// The event should have a zero ID.
+		Assert.AreEqual(new EventId(0), capture.Events[0].Id);
+	}
+
+	[TestMethod]
+	public void TestTogglingAllWildcardAndReenablingOneMatcher()
+	{
+		LoggerCapture capture = new();
+		TestMatcher highest = new(LogEventPriority.Highest, LogLevel.Critical);
+		TestMatcher lowest = new(LogEventPriority.Lowest, LogLevel.Debug);
+
+		using LogEventParser parser = new(capture);
+		parser.AddMatchers([highest, lowest]);
+
+		parser.WriteLine("<-- logeventparser suspend * -->");
+		parser.WriteLine(LineToLog);
+
 		// Neither should have been called.
 		Assert.AreEqual(0, highest.MatchCalls);
 		Assert.AreEqual(0, lowest.MatchCalls);
@@ -118,9 +137,144 @@ public class LogEventParserTests
 		Assert.AreEqual(1, capture.Events.Count);
 		// The event should have a zero ID.
 		Assert.AreEqual(new EventId(0), capture.Events[0].Id);
+
+		parser.WriteLine("<-- logeventparser resume lowest -->");
+		parser.WriteLine(LineToLog);
+
+		// Lowest should have been called.
+		Assert.AreEqual(0, highest.MatchCalls);
+		Assert.AreEqual(1, lowest.MatchCalls);
+
+		Assert.AreEqual(LineToLog, capture.Events[^1].Message);
+		// Lowest should have matched the event.
+		Assert.AreEqual(lowest.Level, capture.Events[^1].Level);
+	}
+
+	[TestMethod]
+	public void TestPoppingUndoesToggles()
+	{
+		LoggerCapture capture = new();
+		TestMatcher highest = new(LogEventPriority.Highest, LogLevel.Critical);
+		TestMatcher lowest = new(LogEventPriority.Lowest, LogLevel.Debug);
+
+		using LogEventParser parser = new(capture);
+		parser.AddMatchers([highest, lowest]);
+
+		parser.WriteLine("<-- logeventparser suspend * -->");
+		parser.WriteLine("<-- logeventparser resume lowest -->");
+		parser.WriteLine("<-- logeventparser pop -->");
+		parser.WriteLine(LineToLog);
+
+		// Neither should have been called after we popped the resume.
+		Assert.AreEqual(0, highest.MatchCalls);
+		Assert.AreEqual(0, lowest.MatchCalls);
+
+		parser.WriteLine("<-- logeventparser pop -->");
+		parser.WriteLine(LineToLog);
+
+		// Both should have been called after the pop.
+		Assert.AreEqual(1, highest.MatchCalls);
+		Assert.AreEqual(1, lowest.MatchCalls);
+	}
+
+	[TestMethod]
+	public void TestLayersStackCorrectly()
+	{
+		LoggerCapture capture = new();
+		TestMatcher highest = new(LogEventPriority.Highest, LogLevel.Critical);
+		TestMatcher lowest = new(LogEventPriority.Lowest, LogLevel.Debug);
+
+		using LogEventParser parser = new(capture);
+		parser.AddMatchers([highest, lowest]);
+
+		parser.WriteLine("<-- logeventparser suspend * -->");
+		parser.WriteLine("<-- logeventparser suspend * -->");
+		parser.WriteLine("<-- logeventparser resume lowest -->");
+		parser.WriteLine(LineToLog);
+
+		// Neither should have been called after we only resumed once but suspended twice.
+		Assert.AreEqual(0, highest.MatchCalls);
+		Assert.AreEqual(0, lowest.MatchCalls);
+
+		parser.WriteLine("<-- logeventparser pop -->");
+		parser.WriteLine("<-- logeventparser pop -->");
+		parser.WriteLine(LineToLog);
+
+		// Neither should have been called as we should still have a wildcard disable as the top layer.
+		Assert.AreEqual(0, highest.MatchCalls);
+		Assert.AreEqual(0, lowest.MatchCalls);
+
+		parser.WriteLine("<-- logeventparser pop -->");
+		parser.WriteLine(LineToLog);
+
+		// Now we've popped the disable, the matchers should be enabled again.
+		Assert.AreEqual(1, highest.MatchCalls);
+		Assert.AreEqual(1, lowest.MatchCalls);
+	}
+
+	[TestMethod]
+	public void TestAllowsEnablingBeforeDisabling()
+	{
+		LoggerCapture capture = new();
+		TestMatcher highest = new(LogEventPriority.Highest, LogLevel.Critical);
+		TestMatcher lowest = new(LogEventPriority.Lowest, LogLevel.Debug);
+
+		using LogEventParser parser = new(capture);
+		parser.AddMatchers([highest, lowest]);
+
+		parser.WriteLine("<-- logeventparser resume lowest -->");
+		parser.WriteLine("<-- logeventparser suspend * -->");
+		parser.WriteLine(LineToLog);
+
+		// Lowest should be enabled because it was explicitly enabled as many times as the wildcard was disabled.
+		Assert.AreEqual(0, highest.MatchCalls);
+		Assert.AreEqual(1, lowest.MatchCalls);
+
+		parser.WriteLine("<-- logeventparser pop -->");
+		parser.WriteLine(LineToLog);
+
+		// Both should now be enabled.
+		Assert.AreEqual(1, highest.MatchCalls);
+		Assert.AreEqual(2, lowest.MatchCalls);
+	}
+
+	[TestMethod]
+	public void TestMismatchedLayerNamesWarns()
+	{
+		LoggerCapture capture = new();
+		TestMatcher highest = new(LogEventPriority.Highest, LogLevel.Critical);
+		TestMatcher lowest = new(LogEventPriority.Lowest, LogLevel.Debug);
+
+		using LogEventParser parser = new(capture);
+		parser.AddMatchers([highest, lowest]);
+
+		parser.WriteLine("<-- logeventparser layer=thing1 suspend * -->");
+		parser.WriteLine("<-- logeventparser layer=thing2 suspend * -->");
+		parser.WriteLine("<-- logeventparser layer=thing1 pop -->");
+
+		// Neither should have been called.
+		Assert.AreEqual(0, highest.MatchCalls);
+		Assert.AreEqual(0, lowest.MatchCalls);
+
+		// We should have one event, the warning issued.
+		Assert.AreEqual(1, capture.Events.Count);
+		Assert.AreEqual("[LogEventParser] Given layer 'thing1' is not the same as last pushed layer 'thing2', please check your logic.", capture.Events[^1].Message);
+		Assert.AreEqual(LogLevel.Warning, capture.Events[^1].Level);
+
+		parser.WriteLine("<-- logeventparser pop -->");
+
+		// An anonymous pop of a non-anonymous layer is fine and should not warn again.
+		Assert.AreEqual(1, capture.Events.Count);
+
+		parser.WriteLine(LineToLog);
+
+		// The pop that warned should still have gone through, so matchers should be enabled again.
+		Assert.AreEqual(1, highest.MatchCalls);
+		Assert.AreEqual(1, lowest.MatchCalls);
 	}
 
 	private const string LogLineProperty = "LogLine";
+	private const string LineToLog = "hello";
 
 	private class LoggerCapture : ILogger
 	{
@@ -131,6 +285,11 @@ public class LogEventParserTests
 		public void Log<TState>(LogLevel logLevel, EventId eventId, TState state, Exception? exception, Func<TState, Exception?, string> formatter)
 		{
 			LogEvent logEvent = LogEvent.Read(JsonLogEvent.FromLoggerState(logLevel, eventId, state, exception, formatter).Data.Span);
+			if (logEvent.Message.StartsWith("[LogEventParser]", StringComparison.Ordinal) && logLevel == LogLevel.Trace)
+			{
+				return;
+			}
+
 			KeyValuePair<string, object?>[] items = [new KeyValuePair<string, object?>(LogLineProperty, _logLineIndex)];
 			logEvent.Properties = (logEvent.Properties == null) ? items : Enumerable.Concat(logEvent.Properties, items);
 			Events.Add(logEvent);
diff --git a/Engine/Source/Programs/Shared/EpicGames.Core/LogEventParser.cs b/Engine/Source/Programs/Shared/EpicGames.Core/LogEventParser.cs
--- a/Engine/Source/Programs/Shared/EpicGames.Core/LogEventParser.cs
+++ b/Engine/Source/Programs/Shared/EpicGames.Core/LogEventParser.cs
@@ -146,9 +146,15 @@ namespace EpicGames.Core
 		private readonly List<ILogEventMatcher> _enabledMatchers = [];
 
 		/// <summary>
-		/// A mapping of matcher names to whether they're enabled. Items not present are enabled by default.
+		/// A mapping of matcher names to how many times they've been enabled / disabled. A negative value indicates a status of disabled.
+		/// Items not present are enabled by default.
 		/// </summary>
-		private readonly Dictionary<string, bool> _matcherNameToEnabled = new(StringComparer.OrdinalIgnoreCase);
+		private readonly Dictionary<string, int> _matcherNameToEnabled = new(StringComparer.OrdinalIgnoreCase);
+
+		/// <summary>
+		/// The layers of enabling / disabling matchers, so we can undo applications.
+		/// </summary>
+		private readonly Stack<(string MatcherName, string? GroupName, bool Enabled)> _matcherToggleLayers = [];
 
 		/// <summary>
 		/// Buffer of input lines
@@ -202,7 +208,12 @@ namespace EpicGames.Core
 		/// <summary>
 		/// The pattern matcher toggle commands must conform to.
 		/// </summary>
-		[GeneratedRegex("^<-- logeventparser (?<Action>suspend|resume) (?<MatcherName>\\S+) -->$", RegexOptions.IgnoreCase | RegexOptions.ExplicitCapture)]
+		/// <remarks>
+		/// <c>LayerName</c> is an arbitrary string to allow validation that users are popping the layer they expect to be - if not, a warning will be issued.
+		/// 
+		/// <c>MatcherName</c> is either the literal name of a matcher, or the string <c>*</c> for all matchers.
+		/// </remarks>
+		[GeneratedRegex("^<-- logeventparser (layer=(?<LayerName>\\S+)? )?((?<Action>suspend|resume) (?<MatcherName>\\S+)|(?<Action>pop)) -->$", RegexOptions.IgnoreCase | RegexOptions.ExplicitCapture)]
 		public static partial Regex MatcherToggleRegex();
 
 		/// <summary>
@@ -452,25 +463,65 @@ namespace EpicGames.Core
 				return;
 			}
 
-			bool enabled;
-			if (!_matcherNameToEnabled.TryGetValue(matcher.GetMatcherName(), out enabled))
-			{
-				enabled = true;
-			}
-
 			_matchers.Add(matcher);
-			if (enabled)
-			{
-				_enabledMatchers.Add(matcher);
-			}
+			// This is a little inefficient, but this function shouldn't be called often.
+			RecalculateEnabledMatchers();
 		}
 
 		/// <summary>
 		/// Set the given matcher to the given enabled state.
 		/// </summary>
-		private void ToggleMatcher(string matcherName, bool enabled)
+		private void PushToggle(string matcherName, string layerName, bool enabled)
+		{
+			if (_matcherToggleLayers.Count > 1024)
+			{
+				_logger.LogWarning(
+					"[LogEventParser] Over 1024 log event parser toggle layers have been pushed, no longer processing any more to avoid unbounded memory usage. Most recent toggle: matcher='{MatcherName}', layer='{LayerName}', enabled={Enabled}",
+					matcherName,
+					layerName,
+					enabled);
+				return;
+			}
+
+			// Allow debugging if all log matchers get disabled and this isn't popped.
+			_logger.LogTrace(
+				"[LogEventParser] Most recent toggle: matcher='{MatcherName}', layer='{LayerName}', enabled={Enabled}",
+				matcherName,
+				layerName,
+				enabled);
+
+			_matcherToggleLayers.Push((matcherName, layerName, enabled));
+			_matcherNameToEnabled.TryGetValue(matcherName, out int enabledCount);
+			enabledCount += enabled ? 1 : -1;
+			_matcherNameToEnabled[matcherName] = enabledCount;
+
+			RecalculateEnabledMatchers();
+		}
+
+		private void PopToggle(string layerName)
 		{
-			_matcherNameToEnabled[matcherName] = enabled;
+			_logger.LogTrace("[LogEventParser] Popping most recent toggle: layer='{LayerName}'", layerName);
+
+			if (_matcherToggleLayers.Count == 0)
+			{
+				_logger.LogWarning("[LogEventParser] Popped empty stack, layer='{LayerName}'", layerName);
+				return;
+			}
+
+			(string matcherName, string? layerNameFromPop, bool enabled) = _matcherToggleLayers.Pop();
+
+			if (layerName.Length != 0 && !String.Equals(layerName, layerNameFromPop, StringComparison.OrdinalIgnoreCase))
+			{
+				_logger.LogWarning(
+					"[LogEventParser] Given layer '{LayerName}' is not the same as last pushed layer '{LayerNameFromPop}', please check your logic.",
+					layerName,
+					layerNameFromPop);
+				// Arguable if we should pop anyway, but I think it's slightly safer to do that.
+			}
+
+			// We're popping a layer, so if it was enabled, we're adding a disable and vice versa - hence why the sign is flipped compared to PushToggle.
+			_matcherNameToEnabled[matcherName] += enabled ? -1 : 1;
+
 			RecalculateEnabledMatchers();
 		}
 
@@ -480,9 +531,15 @@ namespace EpicGames.Core
 		private void RecalculateEnabledMatchers()
 		{
 			_enabledMatchers.Clear();
+
+			// To allow disabling all loggers then enabling specific ones, we add the value for each logger to the value for "*".
+			_matcherNameToEnabled.TryGetValue("*", out int allCount);
+
 			foreach (ILogEventMatcher matcher in _matchers)
 			{
-				if (!_matcherNameToEnabled.TryGetValue(matcher.GetMatcherName(), out bool enabled) || enabled)
+				_matcherNameToEnabled.TryGetValue(matcher.GetMatcherName(), out int enabledCount);
+
+				if (enabledCount + allCount >= 0)
 				{
 					_enabledMatchers.Add(matcher);
 				}
@@ -521,6 +578,38 @@ namespace EpicGames.Core
 			_partialLine.Clear();
 		}
 
+		/// <summary>
+		/// Check for and handle a log event parser command.
+		/// </summary>
+		/// <param name="line">The line to check for a log event parser command.</param>
+		/// <returns>Whether a log event parser command was found and handled.</returns>
+		private bool HandleLogEventParserCommand(string line)
+		{
+			Match match = MatcherToggleRegex().Match(line);
+			if (match.Success)
+			{
+				string action = match.Groups["Action"].Value;
+				string matcherName = match.Groups["MatcherName"].Value;
+				string layerName = match.Groups["LayerName"].Value;
+				switch (action.ToUpperInvariant())
+				{
+					case "RESUME":
+						PushToggle(matcherName, layerName, true);
+						break;
+					case "SUSPEND":
+						PushToggle(matcherName, layerName, false);
+						break;
+					case "POP":
+						PopToggle(layerName);
+						break;
+				}
+				_buffer.MoveNext();
+				return true;
+			}
+
+			return false;
+		}
+
 		/// <summary>
 		/// Process any data in the buffer
 		/// </summary>
@@ -540,12 +629,8 @@ namespace EpicGames.Core
 					line = "";
 				}
 
-				Match match = MatcherToggleRegex().Match(line);
-				if (match.Success)
+				if (HandleLogEventParserCommand(line))
 				{
-					bool enabled = match.Groups["Action"].Value.Equals("resume", StringComparison.OrdinalIgnoreCase);
-					ToggleMatcher(match.Groups["MatcherName"].Value, enabled);
-					_buffer.MoveNext();
 					continue;
 				}
 
