SDK json changes:
- can include other json files
- json file references can specify $(EngineDir)
- platform sdks can override the min/max/main key names

#rnx

#Synced-CL 49766262

diff --git a/Engine/Config/Microsoft/GDK_SDK.json b/Engine/Config/Microsoft/GDK_SDK.json
new file mode 100644
--- /dev/null
+++ b/Engine/Config/Microsoft/GDK_SDK.json
@@ -0,0 +1,11 @@
+{
+	"MainGDKVersion": "240605",
+	"MinGDKVersion": "240604",
+	"MaxGDKVersion": "240605",
+
+	"//GRDK_1": "240305 has a byte-ordering problem with PlayFab that causes compile errors",
+	"BannedGDKVersions": [
+		"240305"
+	]
+
+}
diff --git a/Engine/Config/Windows/Windows_SDK.json b/Engine/Config/Windows/Windows_SDK.json
--- a/Engine/Config/Windows/Windows_SDK.json
+++ b/Engine/Config/Windows/Windows_SDK.json
@@ -81,5 +81,7 @@
 	"//05.1": "X64 -- 0xA00 = Win10",
 	"//05.2": "Arm64 -- 0xA00 = Win10",
 	"MinimumWindowsX64TargetVersion": "0xA00",
-	"MinimumWindowsArm64TargetVersion": "0xA00"
+	"MinimumWindowsArm64TargetVersion": "0xA00",
+
+	"IncludeSDKFile": "$(EngineDir)/Config/Microsoft/GDK_SDK.json"
 }
diff --git a/Engine/Source/Programs/Shared/EpicGames.Build/Platform/UEBuildPlatformSDK.cs b/Engine/Source/Programs/Shared/EpicGames.Build/Platform/UEBuildPlatformSDK.cs
--- a/Engine/Source/Programs/Shared/EpicGames.Build/Platform/UEBuildPlatformSDK.cs
+++ b/Engine/Source/Programs/Shared/EpicGames.Build/Platform/UEBuildPlatformSDK.cs
@@ -286,6 +286,12 @@ namespace EpicGames.Core
 
 		#region Platform Overrides
 
+		protected virtual string MainVersionKey => "MainVersion";
+		protected virtual string MinVersionKey => "MinVersion";
+		protected virtual string MaxVersionKey => "MaxVersion";
+
+		protected virtual string BannedSdkVersionsKey => "BannedSdkVersions";
+
 		/// <summary>
 		/// Return the SDK version that the platform wants to use (AutoSDK dir must match this, full SDKs can be in a valid range)
 		/// </summary>
@@ -293,7 +299,7 @@ namespace EpicGames.Core
 		public virtual string GetMainVersion()
 		{
 			// by default, look up in SDK config value
-			return GetRequiredVersionFromConfig("MainVersion");
+			return GetRequiredVersionFromConfig(MainVersionKey);
 		}
 
 		/// <summary>
@@ -304,8 +310,8 @@ namespace EpicGames.Core
 		protected virtual void GetValidVersionRange(out string? minVersion, out string? maxVersion)
 		{
 			// by default, use SDK config file values
-			minVersion = GetVersionFromConfig("MinVersion");
-			maxVersion = GetVersionFromConfig("MaxVersion");
+			minVersion = GetVersionFromConfig(MinVersionKey);
+			maxVersion = GetVersionFromConfig(MaxVersionKey);
 		}
 
 		protected virtual SDKCollection GetValidVersions()
@@ -405,7 +411,7 @@ namespace EpicGames.Core
 			if (TryConvertVersionToInt(minVersion, out ulong minVersionInt, null) && TryConvertVersionToInt(maxVersion, out ulong maxVersionInt, null))
 			{
 				List<Tuple<ulong, ulong>> bannedSdkVersionRanges = [];
-				foreach (string sdkVersionPair in GetStringArrayFromConfig("BannedSdkVersions"))
+				foreach (string sdkVersionPair in GetStringArrayFromConfig(BannedSdkVersionsKey))
 				{
 					string[] versions = sdkVersionPair.Split("-");
 					if (versions.Length > 0 && versions.Length <= 2)
@@ -831,7 +837,6 @@ namespace EpicGames.Core
 					}
 
 					// currently only care about MainVersion in the overrides
-					const string MainVersionKey = "MainVersion";
 					if (overrideConfigSDKVersions.TryGetValue(MainVersionKey, out string? overrideMainVersion))
 					{
 						// if different from default, then remmber it, and mark that it's been overridden
@@ -871,8 +876,10 @@ namespace EpicGames.Core
 			{
 				if (jsonObject.TryGetStringField(key, out string? stringValue))
 				{
-					if (key.Equals("ParentSDKFile", StringComparison.OrdinalIgnoreCase))
+					if (key.Equals("ParentSDKFile", StringComparison.OrdinalIgnoreCase) ||
+						key.Equals("IncludeSDKFile", StringComparison.OrdinalIgnoreCase))
 					{
+						stringValue = stringValue.Replace("$(EngineDir)", Unreal.EngineDirectory.ToString(), StringComparison.OrdinalIgnoreCase);
 						parents.Add(stringValue);
 					}
 					// add this key if it's not already in the (case-insensitive) dictionary 
