Shared: Do not use insecure randomness (CA5394)

#rnx
#Synced-CL 49789638

diff --git a/Engine/Source/Programs/Shared/.editorconfig b/Engine/Source/Programs/Shared/.editorconfig
--- a/Engine/Source/Programs/Shared/.editorconfig
+++ b/Engine/Source/Programs/Shared/.editorconfig
@@ -1061,7 +1061,7 @@ dotnet_diagnostic.CA5392.severity = none
 dotnet_diagnostic.CA5393.severity = warning
 
 # CA5394: Do not use insecure randomness
-dotnet_diagnostic.CA5394.severity = none
+dotnet_diagnostic.CA5394.severity = warning
 
 # CA5395: Miss HttpVerb attribute for action methods
 dotnet_diagnostic.CA5395.severity = warning
@@ -1162,4 +1162,7 @@ dotnet_diagnostic.MA0008.severity = none
 [*Test*.cs]
 
 # CA1707: Identifiers should not contain underscores
-dotnet_diagnostic.CA1707.severity = none
\ No newline at end of file
+dotnet_diagnostic.CA1707.severity = none
+
+# CA5394: Do not use insecure randomness
+dotnet_diagnostic.CA5394.severity = none
diff --git a/Engine/Source/Programs/Shared/EpicGames.Benchmarks/HashBenchmarks.cs b/Engine/Source/Programs/Shared/EpicGames.Benchmarks/HashBenchmarks.cs
--- a/Engine/Source/Programs/Shared/EpicGames.Benchmarks/HashBenchmarks.cs
+++ b/Engine/Source/Programs/Shared/EpicGames.Benchmarks/HashBenchmarks.cs
@@ -1,6 +1,7 @@
 // Copyright Epic Games, Inc. All Rights Reserved.
 
 using System;
+using System.Diagnostics.CodeAnalysis;
 using BenchmarkDotNet.Attributes;
 using EpicGames.Core;
 using Microsoft.VSDiagnostics;
@@ -18,6 +19,7 @@ public class HashBenchmarks
 	private byte[]? _data;
 
 	[GlobalSetup]
+	[SuppressMessage("Security", "CA5394:Do not use insecure randomness", Justification = "Using Random for deterministic test data generation")]
 	public void Setup()
 	{
 		_data = new byte[DataLength];
diff --git a/Engine/Source/Programs/Shared/EpicGames.Horde/Compute/RequestIdAllocator.cs b/Engine/Source/Programs/Shared/EpicGames.Horde/Compute/RequestIdAllocator.cs
--- a/Engine/Source/Programs/Shared/EpicGames.Horde/Compute/RequestIdAllocator.cs
+++ b/Engine/Source/Programs/Shared/EpicGames.Horde/Compute/RequestIdAllocator.cs
@@ -2,7 +2,7 @@
 
 using System;
 using System.Collections.Generic;
-using System.Linq;
+using System.Security.Cryptography;
 
 namespace EpicGames.Horde.Compute;
 
@@ -76,8 +76,13 @@ public class RequestIdAllocator
 
 	private static string GenerateAllocatorId(int length)
 	{
-		Random random = new();
-		string chars = "ABCDEFGHIJKLMNOPQRSTUVWXYZabcdefghijklmnopqrstuvwxyz0123456789";
-		return new string(Enumerable.Repeat(chars, length).Select(s => s[random.Next(s.Length)]).ToArray());
+		const string Chars = "ABCDEFGHIJKLMNOPQRSTUVWXYZabcdefghijklmnopqrstuvwxyz0123456789";
+		return String.Create(length, Chars, (span, c) =>
+		{
+			for (int i = 0; i < span.Length; i++)
+			{
+				span[i] = c[RandomNumberGenerator.GetInt32(c.Length)];
+			}
+		});
 	}
 }
\ No newline at end of file
