Minor BuildGraph v2 improvements

- Added the ability to specify a nodes name when creating a new BgNode object
- Added the ability to create a FileSet from a single FileReference

#Synced-CL 49759094

diff --git a/Engine/Source/Programs/Shared/EpicGames.BuildGraph/Expressions/BgNode.cs b/Engine/Source/Programs/Shared/EpicGames.BuildGraph/Expressions/BgNode.cs
--- a/Engine/Source/Programs/Shared/EpicGames.BuildGraph/Expressions/BgNode.cs
+++ b/Engine/Source/Programs/Shared/EpicGames.BuildGraph/Expressions/BgNode.cs
@@ -99,10 +99,10 @@ namespace EpicGames.BuildGraph.Expressions
 		/// <summary>
 		/// Constructor
 		/// </summary>
-		public BgNode(BgThunk thunk, BgAgent agent)
+		public BgNode(BgThunk thunk, BgAgent agent, BgString? name = null)
 			: base(BgExprFlags.ForceFragment)
 		{
-			Name = GetDefaultNodeName(thunk);
+			Name = name ?? GetDefaultNodeName(thunk);
 			Thunk = thunk;
 			Agent = agent;
 			DefaultOutput = new BgFileSetFromNodeOutputExpr(this, 0);
@@ -343,8 +343,8 @@ namespace EpicGames.BuildGraph.Expressions
 		/// <summary>
 		/// Constructor
 		/// </summary>
-		public BgNode(BgThunk<T> thunk, BgAgent agent)
-			: base(thunk, agent)
+		public BgNode(BgThunk<T> thunk, BgAgent agent, BgString? name = null)
+			: base(thunk, agent, name)
 		{
 			Output = CreateOutput();
 		}
@@ -433,6 +433,19 @@ namespace EpicGames.BuildGraph.Expressions
 			BgThunk thunk = BgThunk.Create(func);
 			return new BgNode(thunk, agent);
 		}
+		
+		/// <summary>
+		/// Creates a node builder for the given agent
+		/// </summary>
+		/// <param name="agent">Agent to run the node</param>
+		/// <param name="name"></param>
+		/// <param name="func">Function to execute</param>
+		/// <returns>Node builder</returns>
+		public static BgNode AddNode(this BgAgent agent, BgString name, Expression<Func<BgContext, Task>> func)
+		{
+			BgThunk thunk = BgThunk.Create(func);
+			return new BgNode(thunk, agent, name);
+		}
 
 		/// <summary>
 		/// Creates a node builder for the given agent
@@ -445,6 +458,19 @@ namespace EpicGames.BuildGraph.Expressions
 			BgThunk<T> thunk = BgThunk.Create(func);
 			return new BgNode<T>(thunk, agent);
 		}
+		
+		/// <summary>
+		/// Creates a node builder for the given agent
+		/// </summary>
+		/// <param name="agent">Agent to run the node</param>
+		/// <param name="name"></param>
+		/// <param name="func">Function to execute</param>
+		/// <returns>Node builder</returns>
+		public static BgNode<T> AddNode<T>(this BgAgent agent, BgString name, Expression<Func<BgContext, Task<T>>> func)
+		{
+			BgThunk<T> thunk = BgThunk.Create(func);
+			return new BgNode<T>(thunk, agent, name);
+		}
 
 		/// <summary>
 		/// Add dependencies onto other nodes or outputs. Outputs from the given tokens will be copied to the current machine before execution of the node.
diff --git a/Engine/Source/Programs/Shared/EpicGames.Core/FileSet.cs b/Engine/Source/Programs/Shared/EpicGames.Core/FileSet.cs
--- a/Engine/Source/Programs/Shared/EpicGames.Core/FileSet.cs
+++ b/Engine/Source/Programs/Shared/EpicGames.Core/FileSet.cs
@@ -77,6 +77,16 @@ namespace EpicGames.Core
 			return FromFiles(directory, [file]);
 		}
 
+		/// <summary>
+		/// Creates a file tree from a given file
+		/// </summary>
+		/// <param name="file">File to add</param>
+		/// <returns>Tree containing the given files</returns>
+		public static FileSet FromFile(FileReference file)
+		{
+			return FromFiles([(file.GetFileName(), file)]);
+		}
+
 		/// <summary>
 		/// Creates a file tree from a given set of files
 		/// </summary>
@@ -318,7 +328,6 @@ namespace EpicGames.Core
 
 		/// <inheritdoc/>
 		public override IEnumerable<KeyValuePair<string, FileSet>> EnumerateDirectories() => _directoryInfo.EnumerateDirectories().Select(x => KeyValuePair.Create<string, FileSet>(x.Name, new FileSetFromDirectory(x, $"{Path}{x.Name}/")));
-
 	}
 
 	/// <summary>
