Fixing issue where delegates were not properly bound by #if/#endif pairs in their definitions.
Added "Editor" as a distinct define scope seperate from EditorOnly.

#Synced-CL 49804989

diff --git a/Engine/Source/Programs/Shared/EpicGames.UHT/Exporters/CodeGen/UhtHeaderCodeGenerator.cs b/Engine/Source/Programs/Shared/EpicGames.UHT/Exporters/CodeGen/UhtHeaderCodeGenerator.cs
--- a/Engine/Source/Programs/Shared/EpicGames.UHT/Exporters/CodeGen/UhtHeaderCodeGenerator.cs
+++ b/Engine/Source/Programs/Shared/EpicGames.UHT/Exporters/CodeGen/UhtHeaderCodeGenerator.cs
@@ -668,7 +668,9 @@ namespace EpicGames.UHT.Exporters.CodeGen
 			builder.Append(fileId).Append('_').Append(lineNumber).Append('_').Append(macroSuffix);
 			if (includeSuffix)
 			{
-				if (defineScope.HasAnyFlags(UhtDefineScope.EditorOnlyData))
+				// TODO: Both flags should never be set at once, so keeping older codegen here.
+				// Should this list be expanded for the other define scopes??
+				if (defineScope.HasAnyFlags(UhtDefineScope.EditorOnlyData | UhtDefineScope.Editor))
 				{
 					builder.Append("_EOD");
 				}
diff --git a/Engine/Source/Programs/Shared/EpicGames.UHT/Exporters/CodeGen/UhtHeaderCodeGeneratorCppFile.cs b/Engine/Source/Programs/Shared/EpicGames.UHT/Exporters/CodeGen/UhtHeaderCodeGeneratorCppFile.cs
--- a/Engine/Source/Programs/Shared/EpicGames.UHT/Exporters/CodeGen/UhtHeaderCodeGeneratorCppFile.cs
+++ b/Engine/Source/Programs/Shared/EpicGames.UHT/Exporters/CodeGen/UhtHeaderCodeGeneratorCppFile.cs
@@ -330,14 +330,14 @@ namespace EpicGames.UHT.Exporters.CodeGen
 					if (field is UhtEnum enumObj)
 					{
 						using UhtCodeBlockComment blockComment = new(builder, field);
-						using UhtMacroBlockEmitter macroBlockEmitter = new(builder, UhtDefineScopeNames.Standard, enumObj.DefineScope);
+						using UhtMacroBlockEmitter macroBlockEmitter = new(builder, enumObj.DefineScope);
 						AppendEnum(builder, enumObj);
 						GetRegistrations(packageRegistrations, field).Enumerations.Add(enumObj);
 					}
 					else if (field is UhtScriptStruct scriptStruct)
 					{
 						using UhtCodeBlockComment blockComment = new(builder, field);
-						using UhtMacroBlockEmitter macroBlockEmitter = new(builder, UhtDefineScopeNames.Standard, scriptStruct.DefineScope);
+						using UhtMacroBlockEmitter macroBlockEmitter = new(builder, scriptStruct.DefineScope);
 						AppendScriptStruct(builder, scriptStruct);
 						if (scriptStruct.ScriptStructFlags.HasAnyFlags(EStructFlags.Native))
 						{
@@ -351,7 +351,7 @@ namespace EpicGames.UHT.Exporters.CodeGen
 					else if (field is UhtFunction function)
 					{
 						using UhtCodeBlockComment blockComment = new(builder, field);
-						using UhtMacroBlockEmitter macroBlockEmitter = new(builder, UhtDefineScopeNames.Standard, function.DefineScope);
+						using UhtMacroBlockEmitter macroBlockEmitter = new(builder, function.DefineScope);
 						AppendFunction(builder, function, false, false);
 					}
 					else if (field is UhtClass classObj)
@@ -369,7 +369,7 @@ namespace EpicGames.UHT.Exporters.CodeGen
 							}
 
 							using UhtCodeBlockComment blockComment = new(builder, field);
-							using UhtMacroBlockEmitter macroBlockEmitter = new(builder, UhtDefineScopeNames.Standard, classObj.DefineScope);
+							using UhtMacroBlockEmitter macroBlockEmitter = new(builder, classObj.DefineScope);
 							AppendClass(builder, classObj, functions);
 							GetRegistrations(packageRegistrations, field).Classes.Add(classObj);
 
@@ -405,7 +405,7 @@ namespace EpicGames.UHT.Exporters.CodeGen
 						using UhtConditionalMacroBlock block = new(builder, "UE_WITH_CONSTINIT_UOBJECT", Session.IsUsingMultipleCompiledInObjectFormats);
 
 						// Registration for RigVM functions
-						builder.AppendInstances(registrations.RigVMStructs, UhtDefineScopeNames.Standard,
+						builder.AppendInstances(registrations.RigVMStructs, 
 							builder => builder.Append("\tstatic inline constinit FRigVMCompiledInStruct RigVMStructs[] = {\r\n"),
 							(builder, scriptStruct) => builder.Append($"\t\t{{ .Struct = {ConstInitSingletonRef(this, scriptStruct)}, ")
 								 .Append(".Functions = MakeArrayView(").AppendSingletonName(this, scriptStruct, UhtSingletonType.Statics).Append("::RigVMFunctions), ")
@@ -417,7 +417,7 @@ namespace EpicGames.UHT.Exporters.CodeGen
 					{
 						using UhtConditionalMacroBlock block = new(builder, "!UE_WITH_CONSTINIT_UOBJECT", Session.IsUsingMultipleCompiledInObjectFormats);
 						// Non-constinit version
-						builder.AppendInstances(registrations.Enumerations, UhtDefineScopeNames.Standard,
+						builder.AppendInstances(registrations.Enumerations, 
 							builder => builder.Append("\tstatic constexpr FEnumRegisterCompiledInInfo EnumInfo[] = {\r\n"),
 							(builder, enumObj) =>
 							{
@@ -431,7 +431,7 @@ namespace EpicGames.UHT.Exporters.CodeGen
 							},
 							builder => builder.Append("\t};\r\n"));
 
-						builder.AppendInstances(registrations.ScriptStructs, UhtDefineScopeNames.Standard,
+						builder.AppendInstances(registrations.ScriptStructs, 
 							builder => builder.Append("\tstatic constexpr FStructRegisterCompiledInInfo ScriptStructInfo[] = {\r\n"),
 							(builder, scriptStruct) =>
 							{
@@ -446,7 +446,7 @@ namespace EpicGames.UHT.Exporters.CodeGen
 							},
 							builder => builder.Append("\t};\r\n"));
 
-						builder.AppendInstances(registrations.Classes, UhtDefineScopeNames.Standard,
+						builder.AppendInstances(registrations.Classes, 
 							builder => builder.Append("\tstatic constexpr FClassRegisterCompiledInInfo ClassInfo[] = {\r\n"),
 							(builder, classObj) =>
 							{
@@ -470,7 +470,7 @@ namespace EpicGames.UHT.Exporters.CodeGen
 						if (!registrations.RigVMStructs.IsEmpty)
 						{
 							builder.Append($"static FRegisterRigVMStructs RigVM_{name}{{\r\n")
-								.AppendArrayView(registrations.RigVMStructs, UhtDefineScopeNames.Standard, staticsName, "RigVMStructs", 1, "\r\n")
+								.AppendArrayView(registrations.RigVMStructs, staticsName, "RigVMStructs", 1, "\r\n")
 								.Append("};\r\n");
 						}
 					}
@@ -483,9 +483,9 @@ namespace EpicGames.UHT.Exporters.CodeGen
 							.Append("\tTEXT(\"")
 							.Append(package.EngineName)
 							.Append("\"),\r\n")
-							.AppendArrayPtrAndCountLine(registrations.Classes, UhtDefineScopeNames.Standard, staticsName, "ClassInfo", 1, ",\r\n")
-							.AppendArrayPtrAndCountLine(registrations.ScriptStructs, UhtDefineScopeNames.Standard, staticsName, "ScriptStructInfo", 1, ",\r\n")
-							.AppendArrayPtrAndCountLine(registrations.Enumerations, UhtDefineScopeNames.Standard, staticsName, "EnumInfo", 1, ",\r\n")
+							.AppendArrayPtrAndCountLine(registrations.Classes, staticsName, "ClassInfo", 1, ",\r\n")
+							.AppendArrayPtrAndCountLine(registrations.ScriptStructs, staticsName, "ScriptStructInfo", 1, ",\r\n")
+							.AppendArrayPtrAndCountLine(registrations.Enumerations, staticsName, "EnumInfo", 1, ",\r\n")
 							.Append("};\r\n");
 					}
 				}
@@ -1116,8 +1116,8 @@ namespace EpicGames.UHT.Exporters.CodeGen
 						builder.Append("\tnullptr,\r\n");
 					}
 					builder.Append('\t').AppendUTF8LiteralString(scriptStruct.EngineName).Append(",\r\n");
-					builder.AppendArrayPtrLine(properties, UhtDefineScopeNames.Standard, staticsName, "PropPointers", 1, ",\r\n");
-					builder.AppendArrayCountLine(properties, UhtDefineScopeNames.Standard, staticsName, "PropPointers", 1, ",\r\n");
+					builder.AppendArrayPtrLine(properties, staticsName, "PropPointers", 1, ",\r\n");
+					builder.AppendArrayCountLine(properties, staticsName, "PropPointers", 1, ",\r\n");
 					builder.Append("\tsizeof(").Append(scriptStruct.NamespaceExportName).Append(scriptStruct.SourceName).Append("),\r\n");
 					builder.Append("\talignof(").Append(scriptStruct.NamespaceExportName).Append(scriptStruct.SourceName).Append("),\r\n");
 					builder.Append("\tRF_Public|RF_Transient|RF_MarkAsNative,\r\n");
@@ -1366,7 +1366,7 @@ namespace EpicGames.UHT.Exporters.CodeGen
 					builder.Append($"UE_CONSTINIT_UOBJECT_DECL TNoDestroy<U{function.EngineClassName}> {constinitName}{{NoDestroyConstEval,\r\n")
 						.AppendConstInitObjectParams(function, this, 1);
 					builder.Append("\tUE::CodeGen::ConstInit::FUFieldParams{ \r\n")
-						.AppendScopeLink(function.NextFunction, UhtDefineScopeNames.WithEditor, context, 2, ".Next = ", ",\r\n")
+						.AppendScopeLink(function.NextFunction, context, 2, ".Next = ", ",\r\n")
 						.Append("\t},\r\n");
 					builder.Append("\tUE::CodeGen::ConstInit::FStructParams{\r\n")
 						.AppendLine($"\t\t.Super = {ConstInitSingletonRef(this, function.SuperFunction)}, ")
@@ -1421,8 +1421,8 @@ namespace EpicGames.UHT.Exporters.CodeGen
 						.Append("(UObject*(*)())").Append(GetFunctionOuterFunc(function)).Append(", ")
 						.Append(GetSingletonName(function.SuperFunction, UhtSingletonType.Registered)).Append(", ")
 						.AppendUTF8LiteralString(functionEngineName).Append(", ")
-						.AppendArrayPtrLine(properties, UhtDefineScopeNames.Standard, staticsName, "PropPointers", 1, ", \r\n")
-						.AppendArrayCountLine(properties, UhtDefineScopeNames.Standard, staticsName, "PropPointers", 1, ", \r\n")
+						.AppendArrayPtrLine(properties, staticsName, "PropPointers", 1, ", \r\n")
+						.AppendArrayCountLine(properties, staticsName, "PropPointers", 1, ", \r\n")
 						.Append(sizeOfStatic ?? "0").Append(",\r\n")
 						.Append("RF_Public|RF_Transient|RF_MarkAsNative, ")
 						.Append($"(EFunctionFlags)0x{(uint)function.FunctionFlags:X8}, ")
@@ -1522,7 +1522,7 @@ namespace EpicGames.UHT.Exporters.CodeGen
 
 		private static StringBuilder AppendExportProperties(StringBuilder builder, UhtScriptStruct scriptStruct, int tabs)
 		{
-			using (UhtMacroBlockEmitter emitter = new(builder, UhtDefineScopeNames.Standard, UhtDefineScope.None))
+			using (UhtMacroBlockEmitter emitter = new(builder, UhtDefineScope.None))
 			{
 				foreach (UhtProperty property in scriptStruct.Properties)
 				{
@@ -1545,12 +1545,12 @@ namespace EpicGames.UHT.Exporters.CodeGen
 				}
 				else
 				{
-					builder.AppendInstances(properties, UhtDefineScopeNames.Standard,
+					builder.AppendInstances(properties, 
 							(builder, property) => builder.AppendConstInitMemberDecl(property, context, property.EngineName, "", tabs))
 						// declare consteval function to provide first property for constructor uncondtionally - body will be empty if there are no properties
 						.AppendTabs(tabs).Append("static consteval FProperty* GetFirstProperty();\r\n")
 						// for properties which have more than one 'next' property depending on preprocessor definitions, define a function to return the correct property
-						.AppendInstances(properties, UhtDefineScopeNames.Standard, (builder, property) =>
+						.AppendInstances(properties, (builder, property) =>
 						{
 							property.AppendGetNextProperty(builder, context, tabs);
 						});
@@ -1559,7 +1559,7 @@ namespace EpicGames.UHT.Exporters.CodeGen
 			if (Session.IsUsingCompiledInObjectFormat(UhtCompiledInObjectFormat.Params))
 			{
 				using UhtConditionalMacroBlock macroBlock = new(builder, "!UE_WITH_CONSTINIT_UOBJECT", Session.IsUsingMultipleCompiledInObjectFormats);
-				builder.AppendInstances(properties, UhtDefineScopeNames.Standard,
+				builder.AppendInstances(properties,
 					builder => { },
 					(builder, property) => builder.AppendMemberDecl(property, context, property.EngineName, "", tabs),
 					builder => builder.AppendTabs(tabs).Append("static const UECodeGen_Private::FPropertyParamsBase* const PropPointers[];\r\n")
@@ -1580,19 +1580,19 @@ namespace EpicGames.UHT.Exporters.CodeGen
 			if (Session.IsUsingCompiledInObjectFormat(UhtCompiledInObjectFormat.ConstInit))
 			{
 				using UhtConditionalMacroBlock macroBlock = new(builder, "UE_WITH_CONSTINIT_UOBJECT", Session.IsUsingMultipleCompiledInObjectFormats);
-				builder.AppendInstances(properties, UhtDefineScopeNames.Standard,
+				builder.AppendInstances(properties, 
 						(builder, property) => builder.AppendConstInitMemberDef(property, context, property.EngineName, "", null, null, tabs))
 					.AppendTabs(tabs).Append($"consteval FProperty* {context.StaticsName}::GetFirstProperty() {{\r\n")
-					.AppendScopeLink(context.OuterStruct.FirstProperty, UhtDefineScopeNames.Standard, context, tabs + 1, "return ", ";\r\n")
+					.AppendScopeLink(context.OuterStruct.FirstProperty, context, tabs + 1, "return ", ";\r\n")
 					.Append("}\r\n");
 			}
 			if (Session.IsUsingCompiledInObjectFormat(UhtCompiledInObjectFormat.Params))
 			{
 				using UhtConditionalMacroBlock macroBlock = new(builder, "!UE_WITH_CONSTINIT_UOBJECT", Session.IsUsingMultipleCompiledInObjectFormats);
-				builder.AppendInstances(properties, UhtDefineScopeNames.Standard,
+				builder.AppendInstances(properties, 
 					(builder, property) => builder.AppendMemberDef(property, context, property.EngineName, "", null, tabs));
 
-				builder.AppendInstances(properties, UhtDefineScopeNames.Standard,
+				builder.AppendInstances(properties, 
 					builder => builder.AppendTabs(tabs).Append("const UECodeGen_Private::FPropertyParamsBase* const ").Append(context.StaticsName).Append("::PropPointers[] = {\r\n"),
 					(builder, property) => builder.AppendMemberPtr(property, context, property.EngineName, "", tabs + 1),
 					builder =>
@@ -1615,7 +1615,7 @@ namespace EpicGames.UHT.Exporters.CodeGen
 			if (isNotDelegate || isRpc || isCallback)
 			{
 				using UhtCodeBlockComment blockComment = new(builder, classObj, function);
-				using UhtMacroBlockEmitter blockEmitter = new(builder, UhtDefineScopeNames.WithEditor, function.DefineScope);
+				using UhtMacroBlockEmitter blockEmitter = new(builder, function.DefineScope);
 				if (isCallback)
 				{
 					AppendEventParameter(builder, function, function.StrippedFunctionName, UhtPropertyTextType.EventParameterMember, true, 0, "\r\n");
@@ -1667,7 +1667,7 @@ namespace EpicGames.UHT.Exporters.CodeGen
 
 		private static StringBuilder AppendVerseNativeCallableProxyFunctions(StringBuilder builder, UhtClass classObj, UhtUsedDefineScopes<UhtFunction> verseNativeCallableFunctions)
 		{
-			return builder.AppendInstances(verseNativeCallableFunctions, UhtDefineScopeNames.Standard,
+			return builder.AppendInstances(verseNativeCallableFunctions, 
 				(builder, function) =>
 				{
 					AppendVerseNativeCallableFunction(builder, classObj, function, UhtNames.VerseProxySuffix);
@@ -2057,7 +2057,7 @@ namespace EpicGames.UHT.Exporters.CodeGen
 			}
 
 			bool isInterfaceClass = classObj.ClassFlags.HasAnyFlags(EClassFlags.Interface);
-			using UhtMacroBlockEmitter blockEmitter = new(builder, UhtDefineScopeNames.WithEditor, function.DefineScope);
+			using UhtMacroBlockEmitter blockEmitter = new(builder, function.DefineScope);
 
 			if (!isInterfaceClass)
 			{
@@ -2368,14 +2368,14 @@ namespace EpicGames.UHT.Exporters.CodeGen
 			UhtUsedDefineScopes<UhtType> notifyTypes = GetFieldNotifyTypes(classObj);
 
 			//UE_FIELD_NOTIFICATION_DECLARE_FIELD
-			builder.AppendInstances(notifyTypes, UhtDefineScopeNames.Standard,
+			builder.AppendInstances(notifyTypes, 
 				(builder, notifyType) =>
 				{
 					builder.Append($"\tUE_FIELD_NOTIFICATION_IMPLEMENT_FIELD({classObj.SourceName}, {GetNotifyTypeName(notifyType)})\r\n");
 				});
 
 			//UE_FIELD_NOTIFICATION_DECLARE_ENUM_FIELD
-			builder.AppendInstances(notifyTypes, UhtDefineScopeNames.Standard,
+			builder.AppendInstances(notifyTypes,
 				builder => builder.Append($"\tUE_FIELD_NOTIFICATION_IMPLEMENTATION_BEGIN({classObj.SourceName})\r\n"),
 				(builder, notifyType) =>
 				{
@@ -2394,7 +2394,7 @@ namespace EpicGames.UHT.Exporters.CodeGen
 				return builder;
 			}
 
-			return builder.AppendInstances(autoGetterSetterProperties, UhtDefineScopeNames.Standard,
+			return builder.AppendInstances(autoGetterSetterProperties, 
 				(builder, property) =>
 				{
 					if (property.PropertyExportFlags.HasAnyFlags(UhtPropertyExportFlags.GetterSpecifiedAuto))
@@ -2426,7 +2426,7 @@ namespace EpicGames.UHT.Exporters.CodeGen
 					{
 						builder.Append("void ").Append(classObj.Namespace.FullSourceName).Append(classObj.SourceName).Append("::").AppendPropertyGetterWrapperName(property).Append("(const void* Object, void* OutValue)\r\n");
 						builder.Append("{\r\n");
-						using (UhtMacroBlockEmitter blockEmitter = new(builder, UhtDefineScopeNames.Standard, property.DefineScope))
+						using (UhtMacroBlockEmitter blockEmitter = new(builder, property.DefineScope))
 						{
 							builder.Append("\tconst ").Append(classObj.SourceName).Append("* Obj = (const ").Append(classObj.SourceName).Append("*)Object;\r\n");
 							if (property.IsStaticArray)
@@ -2490,7 +2490,7 @@ namespace EpicGames.UHT.Exporters.CodeGen
 					{
 						builder.Append("void ").Append(classObj.Namespace.FullSourceName).Append(classObj.SourceName).Append("::").AppendPropertySetterWrapperName(property).Append("(void* Object, const void* InValue)\r\n");
 						builder.Append("{\r\n");
-						using (UhtMacroBlockEmitter blockEmitter = new(builder, UhtDefineScopeNames.Standard, property.DefineScope))
+						using (UhtMacroBlockEmitter blockEmitter = new(builder, property.DefineScope))
 						{
 							builder.Append('\t').Append(classObj.SourceName).Append("* Obj = (").Append(classObj.SourceName).Append("*)Object;\r\n");
 							if (property.IsStaticArray)
@@ -2568,7 +2568,7 @@ namespace EpicGames.UHT.Exporters.CodeGen
 		{
 			if (!classObj.SerializerArchiveType.HasAnyFlags(serializerType))
 			{
-				builder.AppendScoped(UhtDefineScopeNames.Standard, classObj.SerializerDefineScope,
+				builder.AppendScoped(classObj.SerializerDefineScope,
 					builder => builder.Append(macroText).Append('(').Append(classObj.Namespace.FullSourceName).Append(classObj.SourceName).Append(")\r\n"));
 			}
 			return builder;
@@ -2662,7 +2662,7 @@ namespace EpicGames.UHT.Exporters.CodeGen
 					hasVerseCallableFunctions = nativeFunctions.Instances.Any(x => x.IsVerseField && x.FunctionExportFlags.HasAnyFlags(UhtFunctionExportFlags.VerseCallable));
 					if (hasNativeFunctions)
 					{
-						using UhtMacroBlockEmitter blockEmitter = new(builder, UhtDefineScopeNames.WithEditor, nativeFunctions.SoleScope);
+						using UhtMacroBlockEmitter blockEmitter = new(builder, nativeFunctions.SoleScope);
 						builder.Append("\tstatic constexpr UE::CodeGen::FClassNativeFunction Funcs[] = {\r\n");
 
 						foreach (UhtFunction function in nativeFunctions.Instances)
@@ -2704,7 +2704,7 @@ namespace EpicGames.UHT.Exporters.CodeGen
 					if (hasVerseCallableFunctions)
 					{
 						builder.Append("#if WITH_VERSE_VM\r\n");
-						using UhtMacroBlockEmitter blockEmitter = new(builder, UhtDefineScopeNames.WithEditor, nativeFunctions.SoleScope);
+						using UhtMacroBlockEmitter blockEmitter = new(builder, nativeFunctions.SoleScope);
 						builder.Append("\tstatic constexpr FVerseCallableThunk VerseFuncs[] = {\r\n");
 
 						foreach (UhtFunction function in nativeFunctions.Instances)
@@ -2741,13 +2741,13 @@ namespace EpicGames.UHT.Exporters.CodeGen
 					using UhtConditionalMacroBlock block = new(builder, "UE_WITH_CONSTINIT_UOBJECT", Session.IsUsingMultipleCompiledInObjectFormats);
 					builder.Append("\tstatic consteval UFunction* GetFirstFunction()\r\n");
 					builder.Append("\t{\r\n");
-					builder.AppendScopeLink(classObj.FirstFunction, UhtDefineScopeNames.WithEditor, context, 2, "return ", ";\r\n");
+					builder.AppendScopeLink(classObj.FirstFunction, context, 2, "return ", ";\r\n");
 					builder.Append("\t}\r\n");
 				}
 				if (Session.IsUsingCompiledInObjectFormat(UhtCompiledInObjectFormat.Params))
 				{
 					using UhtConditionalMacroBlock block = new(builder, "!UE_WITH_CONSTINIT_UOBJECT", Session.IsUsingMultipleCompiledInObjectFormats);
-					builder.AppendInstances(functions, UhtDefineScopeNames.WithEditor,
+					builder.AppendInstances(functions, 
 						builder =>
 						{
 							builder.Append("\tstatic constexpr FClassFunctionLinkInfo ").Append("FuncInfo[] = {\r\n");
@@ -2927,7 +2927,7 @@ namespace EpicGames.UHT.Exporters.CodeGen
 						.Append($"\t\t.CppClassTypeInfo = &{staticsName}::StaticCppClassTypeInfo,\r\n");
 					if (hasNativeFunctions)
 					{
-						builder.Append("\t\t.NativeFunctions = ").AppendArrayView(nativeFunctions, UhtDefineScopeNames.WithEditor, staticsName, "Funcs", 2, ",\r\n");
+						builder.Append("\t\t.NativeFunctions = ").AppendArrayView(nativeFunctions, staticsName, "Funcs", 2, ",\r\n");
 					}
 					if (classObj.ClassWithin is not null && classObj != Session.UObject)
 					{
@@ -2977,12 +2977,12 @@ namespace EpicGames.UHT.Exporters.CodeGen
 					}
 					builder.Append("\t&StaticCppClassTypeInfo,\r\n");
 					builder.Append("\tDependentSingletons,\r\n");
-					builder.AppendArrayPtrLine(functions, UhtDefineScopeNames.WithEditor, null, "FuncInfo", 1, ",\r\n");
-					builder.AppendArrayPtrLine(properties, UhtDefineScopeNames.Standard, staticsName, "PropPointers", 1, ",\r\n");
+					builder.AppendArrayPtrLine(functions, null, "FuncInfo", 1, ",\r\n");
+					builder.AppendArrayPtrLine(properties, staticsName, "PropPointers", 1, ",\r\n");
 					builder.Append('\t').Append(hasInterfaces ? "InterfaceParams" : "nullptr").Append(",\r\n");
 					builder.Append("\tUE_ARRAY_COUNT(DependentSingletons),\r\n");
-					builder.AppendArrayCountLine(functions, UhtDefineScopeNames.WithEditor, null, "FuncInfo", 1, ",\r\n");
-					builder.AppendArrayCountLine(properties, UhtDefineScopeNames.Standard, staticsName, "PropPointers", 1, ",\r\n");
+					builder.AppendArrayCountLine(functions, null, "FuncInfo", 1, ",\r\n");
+					builder.AppendArrayCountLine(properties, staticsName, "PropPointers", 1, ",\r\n");
 					builder.Append('\t').Append(hasInterfaces ? "UE_ARRAY_COUNT(InterfaceParams)" : "0").Append(",\r\n");
 					builder.Append($"\t0x{(uint)classFlags:X8}u,\r\n");
 					builder.Append('\t').AppendMetaDataParams(classObj, staticsName, MetaDataParamsName).Append("\r\n");
@@ -3016,7 +3016,7 @@ namespace EpicGames.UHT.Exporters.CodeGen
 				if (hasNativeFunctions)
 				{
 					builder.Append($"\tUClass* Class = {classObj.SourceName}::StaticClass();\r\n");
-					builder.Append($"\tFNativeFunctionRegistrar::RegisterFunctions(Class, ").AppendArrayView(nativeFunctions, UhtDefineScopeNames.WithEditor, staticsName, "Funcs", 0, ");\r\n");
+					builder.Append($"\tFNativeFunctionRegistrar::RegisterFunctions(Class, ").AppendArrayView(nativeFunctions, staticsName, "Funcs", 0, ");\r\n");
 
 					if (hasVerseCallableFunctions)
 					{
diff --git a/Engine/Source/Programs/Shared/EpicGames.UHT/Exporters/CodeGen/UhtHeaderCodeGeneratorHFile.cs b/Engine/Source/Programs/Shared/EpicGames.UHT/Exporters/CodeGen/UhtHeaderCodeGeneratorHFile.cs
--- a/Engine/Source/Programs/Shared/EpicGames.UHT/Exporters/CodeGen/UhtHeaderCodeGeneratorHFile.cs
+++ b/Engine/Source/Programs/Shared/EpicGames.UHT/Exporters/CodeGen/UhtHeaderCodeGeneratorHFile.cs
@@ -177,7 +177,7 @@ namespace EpicGames.UHT.Exporters.CodeGen
 						if (scriptStruct.ScriptStructFlags.HasAnyFlags(EStructFlags.Native))
 						{
 							using UhtCodeBlockComment blockComment = new(builder, field);
-							using UhtMacroBlockEmitter macroBlockEmitter = new(builder, UhtDefineScopeNames.Standard, scriptStruct.DefineScope);
+							using UhtMacroBlockEmitter macroBlockEmitter = new(builder, scriptStruct.DefineScope);
 							AppendScriptStruct(builder, scriptStruct);
 						}
 					}
@@ -189,7 +189,7 @@ namespace EpicGames.UHT.Exporters.CodeGen
 						if (!classObj.ClassFlags.HasAnyFlags(EClassFlags.Intrinsic))
 						{
 							using UhtCodeBlockComment blockComment = new(builder, field);
-							using UhtMacroBlockEmitter macroBlockEmitter = new(builder, UhtDefineScopeNames.Standard, classObj.DefineScope);
+							using UhtMacroBlockEmitter macroBlockEmitter = new(builder, classObj.DefineScope);
 							AppendClass(builder, classObj);
 						}
 					}
@@ -202,7 +202,7 @@ namespace EpicGames.UHT.Exporters.CodeGen
 						if (!classObj.ClassFlags.HasAnyFlags(EClassFlags.Intrinsic) && classObj.ClassType == UhtClassType.Interface && classObj.IsVerseField)
 						{
 							using UhtCodeBlockComment blockComment = new(builder, field, "Verse Interface Proxy");
-							using UhtMacroBlockEmitter macroBlockEmitter = new(builder, UhtDefineScopeNames.Standard, classObj.DefineScope);
+							using UhtMacroBlockEmitter macroBlockEmitter = new(builder, classObj.DefineScope);
 							AppendNativeInterfaceVerseProxy(builder, classObj);
 						}
 					}
@@ -217,7 +217,7 @@ namespace EpicGames.UHT.Exporters.CodeGen
 					if (field is UhtEnum enumObject)
 					{
 						using UhtCodeBlockComment blockComment = new(builder, field);
-						using UhtMacroBlockEmitter macroBlockEmitter = new(builder, UhtDefineScopeNames.Standard, field.DefineScope);
+						using UhtMacroBlockEmitter macroBlockEmitter = new(builder, field.DefineScope);
 						AppendEnum(builder, enumObject);
 					}
 				}
@@ -1048,7 +1048,7 @@ namespace EpicGames.UHT.Exporters.CodeGen
 			}
 
 			UhtUsedDefineScopes<UhtType> notifyTypes = GetFieldNotifyTypes(classObj);
-			return builder.AppendSingleMacro(notifyTypes, UhtDefineScopeNames.Standard, this, classObj, FieldNotifyMacroSuffix,
+			return builder.AppendSingleMacro(notifyTypes, this, classObj, FieldNotifyMacroSuffix,
 				(builder, filteredTypes) =>
 				{
 					builder.Append("\tUE_FIELD_NOTIFICATION_DECLARE_CLASS_DESCRIPTOR_BEGIN(").Append(Module.Api).Append(") \\\r\n");
@@ -1085,7 +1085,7 @@ namespace EpicGames.UHT.Exporters.CodeGen
 				return builder;
 			}
 
-			return builder.AppendMultiScopedMacros(autoGetterSetterProperties, UhtDefineScopeNames.Standard, this, classObj, AutoGettersSettersMacroSuffix,
+			return builder.AppendMultiScopedMacros(autoGetterSetterProperties, this, classObj, AutoGettersSettersMacroSuffix,
 				(builder, properties) =>
 				{
 					foreach (UhtProperty property in properties)
@@ -1127,7 +1127,7 @@ namespace EpicGames.UHT.Exporters.CodeGen
 
 			AppendSparseStructDeclarations(builder, classObj, sparseScriptStructs);
 
-			builder.AppendMultiScopedMacros(sparseProperties, UhtDefineScopeNames.Standard, this, classObj, SparseDataPropertyAccessorsMacroSuffix,
+			builder.AppendMultiScopedMacros(sparseProperties, this, classObj, SparseDataPropertyAccessorsMacroSuffix,
 				(builder, properties) =>
 				{
 					foreach (UhtProperty property in properties)
@@ -1177,7 +1177,7 @@ namespace EpicGames.UHT.Exporters.CodeGen
 
 		private StringBuilder AppendRpcFunctions(StringBuilder builder, UhtClass classObj, bool usesLegacy, UhtUsedDefineScopes<UhtFunction> functions)
 		{
-			builder.AppendMultiScopedMacros(functions, UhtDefineScopeNames.WithEditor, this, classObj, usesLegacy ? RpcWrappersMacroSuffix : RpcWrappersNoPureDeclsMacroSuffix,
+			builder.AppendMultiScopedMacros(functions, this, classObj, usesLegacy ? RpcWrappersMacroSuffix : RpcWrappersNoPureDeclsMacroSuffix,
 				(builder, functions) =>
 				{
 					if (usesLegacy)
@@ -1199,7 +1199,7 @@ namespace EpicGames.UHT.Exporters.CodeGen
 
 		private StringBuilder AppendVerseNativeCallableFunctions(StringBuilder builder, UhtClass classObj, UhtUsedDefineScopes<UhtFunction> functions)
 		{
-			builder.AppendMultiScopedMacros(functions, UhtDefineScopeNames.WithEditor, this, classObj, VerseNativeCallableWrappersMacroSuffix,
+			builder.AppendMultiScopedMacros(functions, this, classObj, VerseNativeCallableWrappersMacroSuffix,
 				(builder, functions) =>
 				{
 					foreach (UhtFunction function in functions)
@@ -1217,7 +1217,7 @@ namespace EpicGames.UHT.Exporters.CodeGen
 				AppendVerseNativeCallableProxyFunctions(builder, classObj.SuperClass, new(classObj.SuperClass.Functions.Where(x => x.FunctionExportFlags.HasAnyFlags(UhtFunctionExportFlags.VerseNativeCallable))));
 			}
 
-			builder.AppendMultiScoped(functions, UhtDefineScopeNames.WithEditor,
+			builder.AppendMultiScoped(functions, 
 				(builder, functions) =>
 				{
 					foreach (UhtFunction function in functions)
@@ -1357,7 +1357,7 @@ namespace EpicGames.UHT.Exporters.CodeGen
 		{
 			if (!classObj.SerializerArchiveType.HasAnyFlags(type))
 			{
-				builder.AppendScopedMacro(UhtDefineScopeNames.Standard, classObj.SerializerDefineScope, this, classObj, ArchiveSerializerMacroSuffix, false,
+				builder.AppendScopedMacro(classObj.SerializerDefineScope, this, classObj, ArchiveSerializerMacroSuffix, false,
 					builder => builder.Append('\t').Append(declare).Append('(').Append(classObj.SourceName).Append(", ").Append(api[0..^1]).Append(") \\\r\n"));
 			}
 			return builder;
diff --git a/Engine/Source/Programs/Shared/EpicGames.UHT/Exporters/CodeGen/UhtMacroBlockEmitter.cs b/Engine/Source/Programs/Shared/EpicGames.UHT/Exporters/CodeGen/UhtMacroBlockEmitter.cs
--- a/Engine/Source/Programs/Shared/EpicGames.UHT/Exporters/CodeGen/UhtMacroBlockEmitter.cs
+++ b/Engine/Source/Programs/Shared/EpicGames.UHT/Exporters/CodeGen/UhtMacroBlockEmitter.cs
@@ -9,14 +9,12 @@ namespace EpicGames.UHT.Exporters.CodeGen
 	internal struct UhtMacroBlockEmitter : IDisposable
 	{
 		private readonly StringBuilder _builder;
-		private readonly UhtDefineScopeNames _defineScopeNames;
 		private UhtDefineScope _current = UhtDefineScope.None;
 		private bool _open = false;
 
-		public UhtMacroBlockEmitter(StringBuilder builder, UhtDefineScopeNames defineScopeNames, UhtDefineScope initialState)
+		public UhtMacroBlockEmitter(StringBuilder builder, UhtDefineScope initialState)
 		{
 			_builder = builder;
-			_defineScopeNames = defineScopeNames;
 			Set(initialState);
 		}
 
@@ -34,10 +32,10 @@ namespace EpicGames.UHT.Exporters.CodeGen
 			{
 				return;
 			}
-			_builder.AppendEndIfPreprocessor(_current, _defineScopeNames);
+			_builder.AppendEndIfPreprocessor(_current);
 			if (defineScope != UhtDefineScope.None)
 			{
-				_builder.AppendIfPreprocessor(defineScope, _defineScopeNames);
+				_builder.AppendIfPreprocessor(defineScope);
 				_open = true;
 			}
 			else
@@ -66,15 +64,15 @@ namespace EpicGames.UHT.Exporters.CodeGen
 			}
 			if (_current == UhtDefineScope.None)
 			{
-				_builder.AppendIfPreprocessor(defineScope, _defineScopeNames);
+				_builder.AppendIfPreprocessor(defineScope);
 			}
 			else if (defineScope == UhtDefineScope.None)
 			{
-				_builder.AppendElsePreprocessor(_current, _defineScopeNames);
+				_builder.AppendElsePreprocessor(_current);
 			}
 			else
 			{
-				_builder.AppendElseIfPreprocessor(defineScope, _defineScopeNames);
+				_builder.AppendElseIfPreprocessor(defineScope);
 			}
 
 			_current = defineScope;
diff --git a/Engine/Source/Programs/Shared/EpicGames.UHT/Exporters/CodeGen/UhtPackageCodeGenerator.cs b/Engine/Source/Programs/Shared/EpicGames.UHT/Exporters/CodeGen/UhtPackageCodeGenerator.cs
--- a/Engine/Source/Programs/Shared/EpicGames.UHT/Exporters/CodeGen/UhtPackageCodeGenerator.cs
+++ b/Engine/Source/Programs/Shared/EpicGames.UHT/Exporters/CodeGen/UhtPackageCodeGenerator.cs
@@ -170,7 +170,7 @@ namespace EpicGames.UHT.Exporters.CodeGen
 				builder.AppendMetaDataDefinition(type, propertyContext, /*namePrefix*/null, name, /*nameSuffix*/ null, /*metaNameSuffix*/ null, tabs); // Appends an entire array definition
 				if (propertyContext != null && properties != null)
 				{
-					builder.AppendInstances(properties, UhtDefineScopeNames.Standard,
+					builder.AppendInstances(properties, 
 						(builder, property) =>
 						{
 							property.AppendMetaDataDecl(builder, propertyContext, property.EngineName, "", tabs);
diff --git a/Engine/Source/Programs/Shared/EpicGames.UHT/Exporters/CodeGen/UhtPackageCodeGeneratorCppFile.cs b/Engine/Source/Programs/Shared/EpicGames.UHT/Exporters/CodeGen/UhtPackageCodeGeneratorCppFile.cs
--- a/Engine/Source/Programs/Shared/EpicGames.UHT/Exporters/CodeGen/UhtPackageCodeGeneratorCppFile.cs
+++ b/Engine/Source/Programs/Shared/EpicGames.UHT/Exporters/CodeGen/UhtPackageCodeGeneratorCppFile.cs
@@ -163,7 +163,7 @@ namespace EpicGames.UHT.Exporters.CodeGen
 						""");
 					foreach (UhtObject obj in linkableObjectsOrdered)
 					{
-						using UhtMacroBlockEmitter macroBlockEmitter = new(builder, UhtDefineScopeNames.Standard, obj.DefineScope);
+						using UhtMacroBlockEmitter macroBlockEmitter = new(builder, obj.DefineScope);
 						builder.AppendLine($"\t{SingletonRef(this, obj, UhtSingletonType.ConstInit)}, ");
 						if (obj.DefineScope != UhtDefineScope.None)
 						{
@@ -213,7 +213,7 @@ namespace EpicGames.UHT.Exporters.CodeGen
 					string strippedName = PackageInfos[package.PackageTypeIndex].StrippedName;
 					string singletonName = GetSingletonName(package, UhtSingletonType.Registered);
 					string staticsName = GetSingletonName(package, UhtSingletonType.Statics);
-					UhtField[] packageSingletons = [.. singletons.Where(x => x.Package == package)];
+					UhtUsedDefineScopes<UhtField> packageSingletons = new([.. singletons.Where(x => x.Package == package)]);
 					EPackageFlags flags = package.PackageFlags & (EPackageFlags.ClientOptional | EPackageFlags.ServerSideOnly | EPackageFlags.EditorOnly | EPackageFlags.Developer | EPackageFlags.UncookedOnly);
 
 					// Entire loop body is #if/#else/#endif
@@ -245,20 +245,16 @@ namespace EpicGames.UHT.Exporters.CodeGen
 							.Append($"\t\tif (!Z_Registration_Info_UPackage_{strippedName}.OuterSingleton)\r\n")
 							.Append("\t\t{\r\n");
 
-						if (packageSingletons.Length != 0)
-						{
-							builder.Append("\t\tstatic UObject* (*const SingletonFuncArray[])() = {\r\n");
-							foreach (UhtField field in packageSingletons)
-							{
-								builder.Append("\t\t\t(UObject* (*)())").Append(ObjectInfos[field.ObjectTypeIndex].RegisteredSingletonName).Append(",\r\n");
-							}
-							builder.Append("\t\t};\r\n");
-						}
+
+						builder.AppendInstances(packageSingletons, 
+							builder => builder.Append("\t\tstatic UObject* (*const SingletonFuncArray[])() = {\r\n"),
+							(builder, fieldObj) => builder.Append($"\t\t\t(UObject* (*)())").Append(ObjectInfos[fieldObj.ObjectTypeIndex].RegisteredSingletonName).Append(",\r\n"),
+							builder => builder.Append("\t\t};\r\n"));
 
 						builder.Append("\t\tstatic const UECodeGen_Private::FPackageParams PackageParams = {\r\n");
 						builder.Append("\t\t\t").AppendUTF8LiteralString(package.SourceName).Append(",\r\n");
-						builder.Append("\t\t\t").Append(packageSingletons.Length != 0 ? "SingletonFuncArray" : "nullptr").Append(",\r\n");
-						builder.Append("\t\t\t").Append(packageSingletons.Length != 0 ? "UE_ARRAY_COUNT(SingletonFuncArray)" : "0").Append(",\r\n");
+						builder.AppendArrayPtrLine(packageSingletons, null, "SingletonFuncArray", 3, ",\r\n");
+						builder.AppendArrayCountLine(packageSingletons, null, "SingletonFuncArray", 3, ",\r\n");
 						builder.Append("\t\t\tPKG_CompiledIn | ").Append($"0x{(uint)flags:X8}").Append(",\r\n");
 						builder.Append("\t\t\t").Append($"0x{(uint)bodiesHash:X8}").Append(",\r\n");
 						builder.Append("\t\t\t").Append($"0x{(uint)declarationsHash:X8}").Append(",\r\n");
diff --git a/Engine/Source/Programs/Shared/EpicGames.UHT/Exporters/CodeGen/UhtUsedDefineScopes.cs b/Engine/Source/Programs/Shared/EpicGames.UHT/Exporters/CodeGen/UhtUsedDefineScopes.cs
--- a/Engine/Source/Programs/Shared/EpicGames.UHT/Exporters/CodeGen/UhtUsedDefineScopes.cs
+++ b/Engine/Source/Programs/Shared/EpicGames.UHT/Exporters/CodeGen/UhtUsedDefineScopes.cs
@@ -357,13 +357,12 @@ namespace EpicGames.UHT.Exporters.CodeGen
 		/// Append a macro scoped macro
 		/// </summary>
 		/// <param name="builder">Destination builder</param>
-		/// <param name="defineScopeNames">Which set of scope names will be used</param>
 		/// <param name="defineScope">Specified scope</param>
 		/// <param name="appendAction">Action to invoke to append an instance</param>
 		/// <returns>String builder</returns>
-		public static StringBuilder AppendScoped(this StringBuilder builder, UhtDefineScopeNames defineScopeNames, UhtDefineScope defineScope, Action<StringBuilder> appendAction)
+		public static StringBuilder AppendScoped(this StringBuilder builder, UhtDefineScope defineScope, Action<StringBuilder> appendAction)
 		{
-			using UhtMacroBlockEmitter blockEmitter = new(builder, defineScopeNames, defineScope);
+			using UhtMacroBlockEmitter blockEmitter = new(builder, defineScope);
 			appendAction(builder);
 			return builder;
 		}
@@ -373,15 +372,13 @@ namespace EpicGames.UHT.Exporters.CodeGen
 		/// </summary>
 		/// <param name="builder">Destination builder</param>
 		/// <param name="instances">Defined scope instances</param>
-		/// <param name="defineScopeNames">Which set of scope names will be used</param>
 		/// <param name="appendAction">Action to invoke to append an instance</param>
 		/// <returns>String builder</returns>
-		public static StringBuilder AppendMultiScoped<T>(this StringBuilder builder, UhtUsedDefineScopes<T> instances, UhtDefineScopeNames defineScopeNames,
-			Action<StringBuilder, IEnumerable<T>> appendAction) where T : UhtType
+		public static StringBuilder AppendMultiScoped<T>(this StringBuilder builder, UhtUsedDefineScopes<T> instances, Action<StringBuilder, IEnumerable<T>> appendAction) where T : UhtType
 		{
 			foreach (UhtDefineScope defineScope in instances.EnumerateDefinedScopes())
 			{
-				AppendScoped(builder, defineScopeNames, defineScope, builder => appendAction(builder, instances.Instances.Where(x => x.DefineScope == defineScope)));
+				AppendScoped(builder, defineScope, builder => appendAction(builder, instances.Instances.Where(x => x.DefineScope == defineScope)));
 			}
 			return builder;
 		}
@@ -390,7 +387,6 @@ namespace EpicGames.UHT.Exporters.CodeGen
 		/// Append a macro scoped macro
 		/// </summary>
 		/// <param name="builder">Destination builder</param>
-		/// <param name="defineScopeNames">Which set of scope names will be used</param>
 		/// <param name="defineScope">Specified scope</param>
 		/// <param name="generator">Header code generator</param>
 		/// <param name="outerType">Output type owning the instances</param>
@@ -398,10 +394,10 @@ namespace EpicGames.UHT.Exporters.CodeGen
 		/// <param name="includeSuffix">If true, include such things as _EOD onto the macro name</param>
 		/// <param name="appendAction">Action to invoke to append an instance</param>
 		/// <returns>String builder</returns>
-		public static StringBuilder AppendScopedMacro(this StringBuilder builder, UhtDefineScopeNames defineScopeNames, UhtDefineScope defineScope,
+		public static StringBuilder AppendScopedMacro(this StringBuilder builder, UhtDefineScope defineScope,
 			UhtHeaderCodeGenerator generator, UhtType outerType, string macroSuffix, bool includeSuffix, Action<StringBuilder> appendAction)
 		{
-			using (UhtMacroBlockEmitter blockEmitter = new(builder, defineScopeNames, defineScope))
+			using (UhtMacroBlockEmitter blockEmitter = new(builder, defineScope))
 			{
 				using (UhtMacroCreator macro = new(builder, generator, outerType, macroSuffix, defineScope, includeSuffix))
 				{
@@ -421,7 +417,7 @@ namespace EpicGames.UHT.Exporters.CodeGen
 						builder.Length -= 4;
 					}
 
-					builder.AppendElsePreprocessor(defineScope, defineScopeNames);
+					builder.AppendElsePreprocessor(defineScope);
 					using UhtMacroCreator macro = new(builder, generator, outerType, macroSuffix, defineScope, includeSuffix); // Empty macro
 				}
 			}
@@ -438,18 +434,17 @@ namespace EpicGames.UHT.Exporters.CodeGen
 		/// </summary>
 		/// <param name="builder">Destination builder</param>
 		/// <param name="instances">Defined scope instances</param>
-		/// <param name="defineScopeNames">Which set of scope names will be used</param>
 		/// <param name="generator">Header code generator</param>
 		/// <param name="outerType">Output type owning the instances</param>
 		/// <param name="macroSuffix">Macro being created</param>
 		/// <param name="appendAction">Action to invoke to append an instance</param>
 		/// <returns>String builder</returns>
-		public static StringBuilder AppendMultiScopedMacros<T>(this StringBuilder builder, UhtUsedDefineScopes<T> instances, UhtDefineScopeNames defineScopeNames,
+		public static StringBuilder AppendMultiScopedMacros<T>(this StringBuilder builder, UhtUsedDefineScopes<T> instances, 
 			UhtHeaderCodeGenerator generator, UhtType outerType, string macroSuffix, Action<StringBuilder, IEnumerable<T>> appendAction) where T : UhtType
 		{
 			foreach (UhtDefineScope defineScope in instances.EnumerateDefinedScopes())
 			{
-				AppendScopedMacro(builder, defineScopeNames, defineScope, generator, outerType, macroSuffix, true,
+				AppendScopedMacro(builder, defineScope, generator, outerType, macroSuffix, true,
 					builder => appendAction(builder, instances.Instances.Where(x => x.DefineScope == defineScope)));
 			}
 			return builder;
@@ -480,14 +475,13 @@ namespace EpicGames.UHT.Exporters.CodeGen
 		/// </summary>
 		/// <param name="builder">Destination builder</param>
 		/// <param name="instances">Defined scope instances</param>
-		/// <param name="defineScopeNames">Which set of scope names will be used</param>
 		/// <param name="generator">Header code generator</param>
 		/// <param name="outerType">Output type owning the instances</param>
 		/// <param name="macroSuffix">Macro being created</param>
 		/// <param name="appendAction">Action to invoke to append an instance</param>
 		/// <returns>String builder</returns>
-		public static StringBuilder AppendSingleMacro<T>(this StringBuilder builder, UhtUsedDefineScopes<T> instances, UhtDefineScopeNames defineScopeNames,
-			UhtHeaderCodeGenerator generator, UhtType outerType, string macroSuffix, Action<StringBuilder, IEnumerable<T>> appendAction) where T : UhtType
+		public static StringBuilder AppendSingleMacro<T>(this StringBuilder builder, UhtUsedDefineScopes<T> instances, UhtHeaderCodeGenerator generator, 
+			UhtType outerType, string macroSuffix, Action<StringBuilder, IEnumerable<T>> appendAction) where T : UhtType
 		{
 			if (instances.SoleScope == UhtDefineScope.None)
 			{
@@ -501,12 +495,12 @@ namespace EpicGames.UHT.Exporters.CodeGen
 				{
 					if (first)
 					{
-						builder.AppendIfPreprocessor(defineScope, defineScopeNames);
+						builder.AppendIfPreprocessor(defineScope);
 						first = false;
 					}
 					else if (defineScope != UhtDefineScope.None)
 					{
-						builder.AppendElseIfPreprocessor(defineScope, defineScopeNames);
+						builder.AppendElseIfPreprocessor(defineScope);
 					}
 					else
 					{
@@ -536,14 +530,13 @@ namespace EpicGames.UHT.Exporters.CodeGen
 		/// </summary>
 		/// <param name="builder">Destination builder</param>
 		/// <param name="instances">Defined scope instances</param>
-		/// <param name="defineScopeNames">Names to use when outputting the scope</param>
 		/// <param name="appendAction">Action to invoke</param>
 		/// <returns>String builder</returns>
-		public static StringBuilder AppendIfInstances<T>(this StringBuilder builder, UhtUsedDefineScopes<T> instances, UhtDefineScopeNames defineScopeNames, Action<StringBuilder> appendAction) where T : UhtType
+		public static StringBuilder AppendIfInstances<T>(this StringBuilder builder, UhtUsedDefineScopes<T> instances, Action<StringBuilder> appendAction) where T : UhtType
 		{
 			if (!instances.IsEmpty)
 			{
-				using UhtMacroBlockEmitter blockEmitter = new(builder, defineScopeNames, instances.NoneAwareScopes);
+				using UhtMacroBlockEmitter blockEmitter = new(builder, instances.NoneAwareScopes);
 				appendAction(builder);
 			}
 			return builder;
@@ -554,12 +547,11 @@ namespace EpicGames.UHT.Exporters.CodeGen
 		/// </summary>
 		/// <param name="builder">Destination builder</param>
 		/// <param name="instances">Defined scope instances</param>
-		/// <param name="defineScopeNames">Names to use when outputting the scope</param>
 		/// <param name="appendAction">Action to invoke</param>
 		/// <returns>String builder</returns>
-		public static StringBuilder AppendInstances<T>(this StringBuilder builder, UhtUsedDefineScopes<T> instances, UhtDefineScopeNames defineScopeNames, Action<StringBuilder, T> appendAction) where T : UhtType
+		public static StringBuilder AppendInstances<T>(this StringBuilder builder, UhtUsedDefineScopes<T> instances, Action<StringBuilder, T> appendAction) where T : UhtType
 		{
-			return builder.AppendInstances(instances, defineScopeNames, null, appendAction, null);
+			return builder.AppendInstances(instances, null, appendAction, null);
 		}
 
 		/// <summary>
@@ -567,17 +559,16 @@ namespace EpicGames.UHT.Exporters.CodeGen
 		/// </summary>
 		/// <param name="builder">Destination builder</param>
 		/// <param name="instances">Defined scope instances</param>
-		/// <param name="defineScopeNames">Names to use when outputting the scope</param>
 		/// <param name="preambleAction">Action to invoke prior to first instance</param>
 		/// <param name="appendAction">Action to invoke</param>
 		/// <param name="postambleAction">Action to invoke following all instances</param>
 		/// <returns>String builder</returns>
-		public static StringBuilder AppendInstances<T>(this StringBuilder builder, UhtUsedDefineScopes<T> instances, UhtDefineScopeNames defineScopeNames,
+		public static StringBuilder AppendInstances<T>(this StringBuilder builder, UhtUsedDefineScopes<T> instances, 
 			Action<StringBuilder>? preambleAction, Action<StringBuilder, T> appendAction, Action<StringBuilder>? postambleAction) where T : UhtType
 		{
 			if (!instances.IsEmpty)
 			{
-				using UhtMacroBlockEmitter blockEmitter = new(builder, defineScopeNames, instances.SoleScope);
+				using UhtMacroBlockEmitter blockEmitter = new(builder, instances.SoleScope);
 				preambleAction?.Invoke(builder);
 				foreach (T instance in instances.Instances)
 				{
@@ -597,13 +588,12 @@ namespace EpicGames.UHT.Exporters.CodeGen
 		/// </summary>
 		/// <param name="builder">Destination builder</param>
 		/// <param name="instances">Collected instances</param>
-		/// <param name="defineScopeNames">Names to use when outputting the scope</param>
 		/// <param name="staticsName">Name of the statics section</param>
 		/// <param name="arrayName">The name of the arrray</param>
 		/// <param name="tabs">Number of tabs to start the line</param>
 		/// <param name="endl">Text to end the line</param>
 		/// <returns>Destination builder</returns>
-		public static StringBuilder AppendArrayView<T>(this StringBuilder builder, UhtUsedDefineScopes<T> instances, UhtDefineScopeNames defineScopeNames, string? staticsName, string arrayName, int tabs, string endl) where T : UhtType
+		public static StringBuilder AppendArrayView<T>(this StringBuilder builder, UhtUsedDefineScopes<T> instances, string? staticsName, string arrayName, int tabs, string endl) where T : UhtType
 		{
 			UhtDefineScope noneAwareScopes = instances.NoneAwareScopes;
 
@@ -630,46 +620,42 @@ namespace EpicGames.UHT.Exporters.CodeGen
 			// If the only scope we have is WITH_EDITORONLY_DATA, then we can use the existing macros
 			else if (noneAwareScopes == UhtDefineScope.EditorOnlyData)
 			{
-				switch (defineScopeNames)
-				{
-					case UhtDefineScopeNames.Standard:
-						builder
-							.AppendTabs(tabs)
-							.Append("IF_WITH_EDITORONLY_DATA(MakeConstArrayView(")
-							.AppendCppName(staticsName, arrayName)
-							.Append("), {})")
-							.Append(endl);
-						break;
-					case UhtDefineScopeNames.WithEditor:
-						builder
-							.AppendTabs(tabs)
-							.Append("IF_WITH_EDITOR(MakeConstArrayView(")
-							.AppendCppName(staticsName, arrayName)
-							.Append("), {})")
-							.Append(endl);
-						break;
-					default:
-						throw new UhtIceException("Unexpected define scope names value");
-				}
+				builder
+					.AppendTabs(tabs)
+					.Append("IF_WITH_EDITORONLY_DATA(MakeConstArrayView(")
+					.AppendCppName(staticsName, arrayName)
+					.Append("), {})")
+					.Append(endl);
+			}
+
+			// If the only scope we have is WITH_EDITOR, then we can use the existing macros
+			else if (noneAwareScopes == UhtDefineScope.Editor)
+			{
+				builder
+					.AppendTabs(tabs)
+					.Append("IF_WITH_EDITOR(MakeConstArrayView(")
+					.AppendCppName(staticsName, arrayName)
+					.Append("), {})")
+					.Append(endl);
 			}
 
 			// Otherwise we have many different scopes but no none scope, must generate #if block
 			else
 			{
 				builder.Append("\r\n");
-				builder.AppendIfPreprocessor(instances.AllScopes, defineScopeNames);
+				builder.AppendIfPreprocessor(instances.AllScopes);
 				builder
 					.AppendTabs(tabs)
 					.Append("MakeArrayView(")
 					.AppendCppName(staticsName, arrayName)
 					.Append(')')
 					.Append(endl);
-				builder.AppendElsePreprocessor(instances.AllScopes, defineScopeNames);
+				builder.AppendElsePreprocessor(instances.AllScopes);
 				builder
 					.AppendTabs(tabs)
 					.Append("{}")
 					.Append(endl);
-				builder.AppendEndIfPreprocessor(instances.AllScopes, defineScopeNames);
+				builder.AppendEndIfPreprocessor(instances.AllScopes);
 			}
 			return builder;
 		}
@@ -679,13 +665,12 @@ namespace EpicGames.UHT.Exporters.CodeGen
 		/// </summary>
 		/// <param name="builder">Destination builder</param>
 		/// <param name="instances">Collected instances</param>
-		/// <param name="defineScopeNames">Names to use when outputting the scope</param>
 		/// <param name="staticsName">Name of the statics section</param>
 		/// <param name="arrayName">The name of the arrray</param>
 		/// <param name="tabs">Number of tabs to start the line</param>
 		/// <param name="endl">Text to end the line</param>
 		/// <returns>Destination builder</returns>
-		public static StringBuilder AppendArrayPtrAndCountLine<T>(this StringBuilder builder, UhtUsedDefineScopes<T> instances, UhtDefineScopeNames defineScopeNames, string? staticsName, string arrayName, int tabs, string endl) where T : UhtType
+		public static StringBuilder AppendArrayPtrAndCountLine<T>(this StringBuilder builder, UhtUsedDefineScopes<T> instances, string? staticsName, string arrayName, int tabs, string endl) where T : UhtType
 		{
 			UhtDefineScope noneAwareScopes = instances.NoneAwareScopes;
 
@@ -713,37 +698,33 @@ namespace EpicGames.UHT.Exporters.CodeGen
 			// If the only scope we have is WITH_EDITORONLY_DATA, then we can use the existing macros
 			else if (noneAwareScopes == UhtDefineScope.EditorOnlyData)
 			{
-				switch (defineScopeNames)
-				{
-					case UhtDefineScopeNames.Standard:
-						builder
-							.AppendTabs(tabs)
-							.Append("IF_WITH_EDITORONLY_DATA(")
-							.AppendCppName(staticsName, arrayName)
-							.Append(", nullptr), IF_WITH_EDITORONLY_DATA(UE_ARRAY_COUNT(")
-							.AppendCppName(staticsName, arrayName)
-							.Append("), 0)")
-							.Append(endl);
-						break;
-					case UhtDefineScopeNames.WithEditor:
-						builder
-							.AppendTabs(tabs)
-							.Append("IF_WITH_EDITOR(")
-							.AppendCppName(staticsName, arrayName)
-							.Append(", nullptr), IF_WITH_EDITOR(UE_ARRAY_COUNT(")
-							.AppendCppName(staticsName, arrayName)
-							.Append("), 0)")
-							.Append(endl);
-						break;
-					default:
-						throw new UhtIceException("Unexpected define scope names value");
-				}
+				builder
+					.AppendTabs(tabs)
+					.Append("IF_WITH_EDITORONLY_DATA(")
+					.AppendCppName(staticsName, arrayName)
+					.Append(", nullptr), IF_WITH_EDITORONLY_DATA(UE_ARRAY_COUNT(")
+					.AppendCppName(staticsName, arrayName)
+					.Append("), 0)")
+					.Append(endl);
+			}
+
+			// If the only scope we have is WITH_EDITOR, then we can use the existing macros
+			else if (noneAwareScopes == UhtDefineScope.Editor)
+			{
+				builder
+					.AppendTabs(tabs)
+					.Append("IF_WITH_EDITOR(")
+					.AppendCppName(staticsName, arrayName)
+					.Append(", nullptr), IF_WITH_EDITOR(UE_ARRAY_COUNT(")
+					.AppendCppName(staticsName, arrayName)
+					.Append("), 0)")
+					.Append(endl);
 			}
 
 			// Otherwise we have many different scopes but no none scope, must generate #if block
 			else
 			{
-				builder.AppendIfPreprocessor(instances.AllScopes, defineScopeNames);
+				builder.AppendIfPreprocessor(instances.AllScopes);
 				builder
 					.AppendTabs(tabs)
 					.AppendCppName(staticsName, arrayName)
@@ -751,12 +732,12 @@ namespace EpicGames.UHT.Exporters.CodeGen
 					.AppendCppName(staticsName, arrayName)
 					.Append(')')
 					.Append(endl);
-				builder.AppendElsePreprocessor(instances.AllScopes, defineScopeNames);
+				builder.AppendElsePreprocessor(instances.AllScopes);
 				builder
 					.AppendTabs(tabs)
 					.Append("nullptr, 0")
 					.Append(endl);
-				builder.AppendEndIfPreprocessor(instances.AllScopes, defineScopeNames);
+				builder.AppendEndIfPreprocessor(instances.AllScopes);
 			}
 			return builder;
 		}
@@ -766,13 +747,12 @@ namespace EpicGames.UHT.Exporters.CodeGen
 		/// </summary>
 		/// <param name="builder">Destination builder</param>
 		/// <param name="instances">Collected instances</param>
-		/// <param name="defineScopeNames">Names to use when outputting the scope</param>
 		/// <param name="staticsName">Name of the statics section</param>
 		/// <param name="arrayName">The name of the arrray</param>
 		/// <param name="tabs">Number of tabs to start the line</param>
 		/// <param name="endl">Text to end the line</param>
 		/// <returns>Destination builder</returns>
-		public static StringBuilder AppendArrayPtrLine<T>(this StringBuilder builder, UhtUsedDefineScopes<T> instances, UhtDefineScopeNames defineScopeNames, string? staticsName, string arrayName, int tabs, string endl) where T : UhtType
+		public static StringBuilder AppendArrayPtrLine<T>(this StringBuilder builder, UhtUsedDefineScopes<T> instances, string? staticsName, string arrayName, int tabs, string endl) where T : UhtType
 		{
 			UhtDefineScope noneAwareScopes = instances.NoneAwareScopes;
 
@@ -797,43 +777,39 @@ namespace EpicGames.UHT.Exporters.CodeGen
 			// If the only scope we have is WITH_EDITORONLY_DATA, then we can use the existing macros
 			else if (noneAwareScopes == UhtDefineScope.EditorOnlyData)
 			{
-				switch (defineScopeNames)
-				{
-					case UhtDefineScopeNames.Standard:
-						builder
-							.AppendTabs(tabs)
-							.Append("IF_WITH_EDITORONLY_DATA(")
-							.AppendCppName(staticsName, arrayName)
-							.Append(", nullptr)")
-							.Append(endl);
-						break;
-					case UhtDefineScopeNames.WithEditor:
-						builder
-							.AppendTabs(tabs)
-							.Append("IF_WITH_EDITOR(")
-							.AppendCppName(staticsName, arrayName)
-							.Append(", nullptr)")
-							.Append(endl);
-						break;
-					default:
-						throw new UhtIceException("Unexpected define scope names value");
-				}
+				builder
+					.AppendTabs(tabs)
+					.Append("IF_WITH_EDITORONLY_DATA(")
+					.AppendCppName(staticsName, arrayName)
+					.Append(", nullptr)")
+					.Append(endl);
+			}
+
+			// If the only scope we have is WITH_EDITOR, then we can use the existing macros
+			else if (noneAwareScopes == UhtDefineScope.Editor)
+			{
+				builder
+					.AppendTabs(tabs)
+					.Append("IF_WITH_EDITOR(")
+					.AppendCppName(staticsName, arrayName)
+					.Append(", nullptr)")
+					.Append(endl);
 			}
 
 			// Otherwise we have many different scopes but no none scope, must generate #if block
 			else
 			{
-				builder.AppendIfPreprocessor(instances.AllScopes, defineScopeNames);
+				builder.AppendIfPreprocessor(instances.AllScopes);
 				builder
 					.AppendTabs(tabs)
 					.AppendCppName(staticsName, arrayName)
 					.Append(endl);
-				builder.AppendElsePreprocessor(instances.AllScopes, defineScopeNames);
+				builder.AppendElsePreprocessor(instances.AllScopes);
 				builder
 					.AppendTabs(tabs)
 					.Append("nullptr")
 					.Append(endl);
-				builder.AppendEndIfPreprocessor(instances.AllScopes, defineScopeNames);
+				builder.AppendEndIfPreprocessor(instances.AllScopes);
 			}
 			return builder;
 		}
@@ -843,13 +819,12 @@ namespace EpicGames.UHT.Exporters.CodeGen
 		/// </summary>
 		/// <param name="builder">Destination builder</param>
 		/// <param name="instances">Collected instances</param>
-		/// <param name="defineScopeNames">Names to use when outputting the scope</param>
 		/// <param name="staticsName">Name of the statics section</param>
 		/// <param name="arrayName">The name of the arrray</param>
 		/// <param name="tabs">Number of tabs to start the line</param>
 		/// <param name="endl">Text to end the line</param>
 		/// <returns>Destination builder</returns>
-		public static StringBuilder AppendArrayCountLine<T>(this StringBuilder builder, UhtUsedDefineScopes<T> instances, UhtDefineScopeNames defineScopeNames, string? staticsName, string arrayName, int tabs, string endl) where T : UhtType
+		public static StringBuilder AppendArrayCountLine<T>(this StringBuilder builder, UhtUsedDefineScopes<T> instances, string? staticsName, string arrayName, int tabs, string endl) where T : UhtType
 		{
 			UhtDefineScope noneAwareScopes = instances.NoneAwareScopes;
 
@@ -876,45 +851,41 @@ namespace EpicGames.UHT.Exporters.CodeGen
 			// If the only scope we have is WITH_EDITORONLY_DATA, then we can use the existing macros
 			else if (noneAwareScopes == UhtDefineScope.EditorOnlyData)
 			{
-				switch (defineScopeNames)
-				{
-					case UhtDefineScopeNames.Standard:
-						builder
-							.AppendTabs(tabs)
-							.Append("IF_WITH_EDITORONLY_DATA(UE_ARRAY_COUNT(")
-							.AppendCppName(staticsName, arrayName)
-							.Append("), 0)")
-							.Append(endl);
-						break;
-					case UhtDefineScopeNames.WithEditor:
-						builder
-							.AppendTabs(tabs)
-							.Append("IF_WITH_EDITOR(UE_ARRAY_COUNT(")
-							.AppendCppName(staticsName, arrayName)
-							.Append("), 0)")
-							.Append(endl);
-						break;
-					default:
-						throw new UhtIceException("Unexpected define scope names value");
-				}
+				builder
+					.AppendTabs(tabs)
+					.Append("IF_WITH_EDITORONLY_DATA(UE_ARRAY_COUNT(")
+					.AppendCppName(staticsName, arrayName)
+					.Append("), 0)")
+					.Append(endl);
+			}
+
+			// If the only scope we have is WITH_EDITOR, then we can use the existing macros
+			else if (noneAwareScopes == UhtDefineScope.Editor)
+			{
+				builder
+					.AppendTabs(tabs)
+					.Append("IF_WITH_EDITOR(UE_ARRAY_COUNT(")
+					.AppendCppName(staticsName, arrayName)
+					.Append("), 0)")
+					.Append(endl);
 			}
 
 			// Otherwise we have many different scopes but no none scope, must generate #if block
 			else
 			{
-				builder.AppendIfPreprocessor(instances.AllScopes, defineScopeNames);
+				builder.AppendIfPreprocessor(instances.AllScopes);
 				builder
 					.AppendTabs(tabs)
 					.Append("UE_ARRAY_COUNT(")
 					.AppendCppName(staticsName, arrayName)
 					.Append(')')
 					.Append(endl);
-				builder.AppendElsePreprocessor(instances.AllScopes, defineScopeNames);
+				builder.AppendElsePreprocessor(instances.AllScopes);
 				builder
 					.AppendTabs(tabs)
 					.Append('0')
 					.Append(endl);
-				builder.AppendEndIfPreprocessor(instances.AllScopes, defineScopeNames);
+				builder.AppendEndIfPreprocessor(instances.AllScopes);
 			}
 			return builder;
 		}
@@ -924,13 +895,12 @@ namespace EpicGames.UHT.Exporters.CodeGen
 		/// </summary>
 		/// <param name="builder"></param>
 		/// <param name="link"></param>
-		/// <param name="defineScopeNames"></param>
 		/// <param name="context"></param>
 		/// <param name="tabs"></param>
 		/// <param name="prefix"></param>
 		/// <param name="endl"></param>
 		/// <returns></returns>
-		public static StringBuilder AppendScopeLink(this StringBuilder builder, UhtDefineScopeLink<UhtProperty>? link, UhtDefineScopeNames defineScopeNames, IUhtPropertyMemberContext context, int tabs, string? prefix, string? endl)
+		public static StringBuilder AppendScopeLink(this StringBuilder builder, UhtDefineScopeLink<UhtProperty>? link, IUhtPropertyMemberContext context, int tabs, string? prefix, string? endl)
 		{
 			// No link or link contains nothing
 			if (link is null)
@@ -938,7 +908,7 @@ namespace EpicGames.UHT.Exporters.CodeGen
 				return builder.AppendTabs(tabs).Append(prefix).Append("nullptr").Append(endl);
 			}
 
-			using (UhtMacroBlockEmitter blockEmitter = new(builder, defineScopeNames, link.InitialScope))
+			using (UhtMacroBlockEmitter blockEmitter = new(builder, link.InitialScope))
 			{
 				foreach (UhtProperty? next in link.Next)
 				{
@@ -955,13 +925,12 @@ namespace EpicGames.UHT.Exporters.CodeGen
 		/// </summary>
 		/// <param name="builder"></param>
 		/// <param name="link"></param>
-		/// <param name="defineScopeNames"></param>
 		/// <param name="context"></param>
 		/// <param name="tabs"></param>
 		/// <param name="prefix"></param>
 		/// <param name="endl"></param>
 		/// <returns></returns>
-		public static StringBuilder AppendScopeLink(this StringBuilder builder, UhtDefineScopeLink<UhtFunction>? link, UhtDefineScopeNames defineScopeNames, IUhtPropertyMemberContext context, int tabs, string? prefix, string? endl)
+		public static StringBuilder AppendScopeLink(this StringBuilder builder, UhtDefineScopeLink<UhtFunction>? link, IUhtPropertyMemberContext context, int tabs, string? prefix, string? endl)
 		{
 			// No link or link contains nothing
 			if (link is null)
@@ -969,7 +938,7 @@ namespace EpicGames.UHT.Exporters.CodeGen
 				return builder.AppendTabs(tabs).Append(prefix).Append("nullptr").Append(endl);
 			}
 
-			using (UhtMacroBlockEmitter blockEmitter = new(builder, defineScopeNames, link.InitialScope))
+			using (UhtMacroBlockEmitter blockEmitter = new(builder, link.InitialScope))
 			{
 				foreach (UhtFunction? next in link.Next)
 				{
@@ -986,13 +955,12 @@ namespace EpicGames.UHT.Exporters.CodeGen
 		/// </summary>
 		/// <param name="builder">String builder</param>
 		/// <param name="defineScope">Scope</param>
-		/// <param name="defineScopeNames">Which set of scope names will be used</param>
 		/// <returns>String builder</returns>
-		public static StringBuilder AppendIfPreprocessor(this StringBuilder builder, UhtDefineScope defineScope, UhtDefineScopeNames defineScopeNames = UhtDefineScopeNames.Standard)
+		public static StringBuilder AppendIfPreprocessor(this StringBuilder builder, UhtDefineScope defineScope)
 		{
 			if (defineScope != UhtDefineScope.None && defineScope != UhtDefineScope.Invalid)
 			{
-				builder.Append("#if ").AppendScopeExpression(defineScope, defineScopeNames).Append("\r\n");
+				builder.Append("#if ").AppendScopeExpression(defineScope).Append("\r\n");
 			}
 			return builder;
 		}
@@ -1002,13 +970,12 @@ namespace EpicGames.UHT.Exporters.CodeGen
 		/// </summary>
 		/// <param name="builder">String builder</param>
 		/// <param name="defineScope">Scope</param>
-		/// <param name="defineScopeNames">Which set of scope names will be used</param>
 		/// <returns>String builder</returns>
-		public static StringBuilder AppendElsePreprocessor(this StringBuilder builder, UhtDefineScope defineScope, UhtDefineScopeNames defineScopeNames = UhtDefineScopeNames.Standard)
+		public static StringBuilder AppendElsePreprocessor(this StringBuilder builder, UhtDefineScope defineScope)
 		{
 			if (defineScope != UhtDefineScope.None && defineScope != UhtDefineScope.Invalid)
 			{
-				builder.Append("#else // ").AppendScopeExpression(defineScope, defineScopeNames).Append("\r\n");
+				builder.Append("#else // ").AppendScopeExpression(defineScope).Append("\r\n");
 			}
 			return builder;
 		}
@@ -1018,13 +985,12 @@ namespace EpicGames.UHT.Exporters.CodeGen
 		/// </summary>
 		/// <param name="builder">String builder</param>
 		/// <param name="defineScope">Scope</param>
-		/// <param name="defineScopeNames">Which set of scope names will be used</param>
 		/// <returns>String builder</returns>
-		public static StringBuilder AppendElseIfPreprocessor(this StringBuilder builder, UhtDefineScope defineScope, UhtDefineScopeNames defineScopeNames = UhtDefineScopeNames.Standard)
+		public static StringBuilder AppendElseIfPreprocessor(this StringBuilder builder, UhtDefineScope defineScope)
 		{
 			if (defineScope != UhtDefineScope.None && defineScope != UhtDefineScope.Invalid)
 			{
-				builder.Append("#elif ").AppendScopeExpression(defineScope, defineScopeNames).Append("\r\n");
+				builder.Append("#elif ").AppendScopeExpression(defineScope).Append("\r\n");
 			}
 			return builder;
 		}
@@ -1034,13 +1000,12 @@ namespace EpicGames.UHT.Exporters.CodeGen
 		/// </summary>
 		/// <param name="builder">String builder</param>
 		/// <param name="defineScope">Scope</param>
-		/// <param name="defineScopeNames">Which set of scope names will be used</param>
 		/// <returns>String builder</returns>
-		public static StringBuilder AppendEndIfPreprocessor(this StringBuilder builder, UhtDefineScope defineScope, UhtDefineScopeNames defineScopeNames = UhtDefineScopeNames.Standard)
+		public static StringBuilder AppendEndIfPreprocessor(this StringBuilder builder, UhtDefineScope defineScope)
 		{
 			if (defineScope != UhtDefineScope.None && defineScope != UhtDefineScope.Invalid)
 			{
-				builder.Append("#endif // ").AppendScopeExpression(defineScope, defineScopeNames).Append("\r\n");
+				builder.Append("#endif // ").AppendScopeExpression(defineScope).Append("\r\n");
 			}
 			return builder;
 		}
@@ -1050,9 +1015,8 @@ namespace EpicGames.UHT.Exporters.CodeGen
 		/// </summary>
 		/// <param name="builder">String builder</param>
 		/// <param name="defineScope">Scope</param>
-		/// <param name="defineScopeNames">Which set of scope names will be used</param>
 		/// <returns>String builder</returns>
-		public static StringBuilder AppendScopeExpression(this StringBuilder builder, UhtDefineScope defineScope, UhtDefineScopeNames defineScopeNames)
+		public static StringBuilder AppendScopeExpression(this StringBuilder builder, UhtDefineScope defineScope)
 		{
 			if (defineScope != UhtDefineScope.None && defineScope != UhtDefineScope.Invalid)
 			{
@@ -1066,7 +1030,7 @@ namespace EpicGames.UHT.Exporters.CodeGen
 						{
 							builder.Append(" && ");
 						}
-						builder.Append(((UhtDefineScope)(mask & scopes)).GetCompilerDirective(defineScopeNames).GetCompilerDirectiveText());
+						builder.Append(((UhtDefineScope)(mask & scopes)).GetCompilerDirective().GetCompilerDirectiveText());
 						needAnd = true;
 					}
 				}
diff --git a/Engine/Source/Programs/Shared/EpicGames.UHT/Parsers/UhtFunctionParser.cs b/Engine/Source/Programs/Shared/EpicGames.UHT/Parsers/UhtFunctionParser.cs
--- a/Engine/Source/Programs/Shared/EpicGames.UHT/Parsers/UhtFunctionParser.cs
+++ b/Engine/Source/Programs/Shared/EpicGames.UHT/Parsers/UhtFunctionParser.cs
@@ -143,6 +143,14 @@ namespace EpicGames.UHT.Parsers
 					UhtSpecifierContext specifierContext = new(topScope, topScope.TokenReader, function.MetaData);
 					UhtSpecifierParser specifiers = UhtSpecifierParser.GetThreadInstance(specifierContext, ScopeName, parentScope.Session.GetSpecifierTable(UhtTableNames.Function));
 
+					UhtCompilerDirective compilerDirective = topScope.HeaderParser.GetCurrentCompositeCompilerDirective();
+					if ((compilerDirective & UhtCompilerDirective.WithEditor) != 0)
+					{
+						function.FunctionFlags |= EFunctionFlags.EditorOnly;
+						function.DefineScope |= UhtDefineScope.Editor;
+					}
+					function.DefineScope |= compilerDirective.GetDefaultDefineScopes();
+
 					// If this is a UDELEGATE, parse the specifiers
 					StringView delegateMacro = new();
 					if (hasSpecifiers)
@@ -360,7 +368,7 @@ namespace EpicGames.UHT.Parsers
 					if ((compilerDirective & UhtCompilerDirective.WithEditor) != 0)
 					{
 						function.FunctionFlags |= EFunctionFlags.EditorOnly;
-						function.DefineScope |= UhtDefineScope.EditorOnlyData;
+						function.DefineScope |= UhtDefineScope.Editor;
 					}
 					function.DefineScope |= compilerDirective.GetDefaultDefineScopes();
 
diff --git a/Engine/Source/Programs/Shared/EpicGames.UHT/Types/UhtFunction.cs b/Engine/Source/Programs/Shared/EpicGames.UHT/Types/UhtFunction.cs
--- a/Engine/Source/Programs/Shared/EpicGames.UHT/Types/UhtFunction.cs
+++ b/Engine/Source/Programs/Shared/EpicGames.UHT/Types/UhtFunction.cs
@@ -748,7 +748,7 @@ namespace EpicGames.UHT.Types
 
 			ValidateSpecifierFlags();
 
-			if (DefineScope.HasAnyFlags(UhtDefineScope.Tests) && !Outer!.DefineScope.HasAnyFlags(UhtDefineScope.Tests))
+			if (DefineScope.HasAnyFlags(UhtDefineScope.Tests) && Outer is UhtClass && !Outer.DefineScope.HasAnyFlags(UhtDefineScope.Tests))
 			{
 				this.LogError($"Functions are not supported within a {UhtNames.WithTests} scope.");
 			}
diff --git a/Engine/Source/Programs/Shared/EpicGames.UHT/Types/UhtProperty.cs b/Engine/Source/Programs/Shared/EpicGames.UHT/Types/UhtProperty.cs
--- a/Engine/Source/Programs/Shared/EpicGames.UHT/Types/UhtProperty.cs
+++ b/Engine/Source/Programs/Shared/EpicGames.UHT/Types/UhtProperty.cs
@@ -1758,7 +1758,7 @@ namespace EpicGames.UHT.Types
 			{
 				// Note that the caller is responsible for ensuring this is scoped to when this property is defined
 				builder.AppendTabs(tabs).Append("static consteval FProperty* GetNextProperty_").Append(EngineName).Append("() {\r\n");
-				builder.AppendScopeLink(NextProperty, UhtDefineScopeNames.Standard, context, tabs + 1, "return ", ";\r\n");
+				builder.AppendScopeLink(NextProperty, context, tabs + 1, "return ", ";\r\n");
 				builder.AppendTabs(tabs).Append("}\r\n");
 			}
 			return builder;
diff --git a/Engine/Source/Programs/Shared/EpicGames.UHT/Types/UhtType.cs b/Engine/Source/Programs/Shared/EpicGames.UHT/Types/UhtType.cs
--- a/Engine/Source/Programs/Shared/EpicGames.UHT/Types/UhtType.cs
+++ b/Engine/Source/Programs/Shared/EpicGames.UHT/Types/UhtType.cs
@@ -532,23 +532,6 @@ namespace EpicGames.UHT.Types
 		}
 	}
 
-	/// <summary>
-	/// Defines which set of names will be used for each scope
-	/// </summary>
-	public enum UhtDefineScopeNames
-	{
-
-		/// <summary>
-		/// Names match the define
-		/// </summary>
-		Standard,
-
-		/// <summary>
-		/// Use WITH_EDITOR in place of WITH_EDITORONLY_DATA
-		/// </summary>
-		WithEditor,
-	}
-
 	/// <summary>
 	/// #define scope where the type was defined.  This only includes macros such as WITH_EDITOR, WITH_EDITORONLY_DATA, and WITH_VERSE_VM.
 	/// Not all types support the given options.
@@ -588,10 +571,15 @@ namespace EpicGames.UHT.Types
 		/// </summary>
 		VerseBPVM = 1 << 3,
 
+		/// <summary>
+		/// Exists inside of a WITH_EDITOR block
+		/// </summary>
+		Editor = 1 << 4,
+
 		/// <summary>
 		/// Number of unique scope combinations
 		/// </summary>
-		ScopeCount = 1 << 4,
+		ScopeCount = 1 << 5,
 	}
 
 	/// <summary>
@@ -639,14 +627,15 @@ namespace EpicGames.UHT.Types
 		/// Return the compiler directive the matches the given scope
 		/// </summary>
 		/// <param name="defineScope">Scope to return compiler directive</param>
-		/// <param name="defineScopeNames">Used to optionally select WITH_EDITOR instead of WITH_EDITORONLYDATA</param>
 		/// <returns></returns>
-		public static UhtCompilerDirective GetCompilerDirective(this UhtDefineScope defineScope, UhtDefineScopeNames defineScopeNames)
+		public static UhtCompilerDirective GetCompilerDirective(this UhtDefineScope defineScope)
 		{
 			switch (defineScope)
 			{
 				case UhtDefineScope.EditorOnlyData:
-					return defineScopeNames == UhtDefineScopeNames.WithEditor ? UhtCompilerDirective.WithEditor : UhtCompilerDirective.WithEditorOnlyData;
+					return UhtCompilerDirective.WithEditorOnlyData;
+				case UhtDefineScope.Editor:
+					return UhtCompilerDirective.WithEditor;
 				case UhtDefineScope.VerseVM:
 					return UhtCompilerDirective.WithVerseVM;
 				case UhtDefineScope.VerseBPVM:
