Cleaned up the code that generates macro names to make it more clear which execution paths need the macro names with the define scope suffix.
Fix suffix generation to always return a valid suffix instead of just for editor only data.

#rnx
#presubmit 6967c78a520c1ea50ff28a5d

#Synced-CL 49804989

diff --git a/Engine/Source/Programs/Shared/EpicGames.UHT/Exporters/CodeGen/UhtHeaderCodeGenerator.cs b/Engine/Source/Programs/Shared/EpicGames.UHT/Exporters/CodeGen/UhtHeaderCodeGenerator.cs
--- a/Engine/Source/Programs/Shared/EpicGames.UHT/Exporters/CodeGen/UhtHeaderCodeGenerator.cs
+++ b/Engine/Source/Programs/Shared/EpicGames.UHT/Exporters/CodeGen/UhtHeaderCodeGenerator.cs
@@ -663,44 +663,65 @@ namespace EpicGames.UHT.Exporters.CodeGen
 
 	internal static class UhtHaederCodeGeneratorStringBuilderExtensions
 	{
-		public static StringBuilder AppendMacroName(this StringBuilder builder, string fileId, int lineNumber, string macroSuffix, UhtDefineScope defineScope = UhtDefineScope.None, bool includeSuffix = true)
+		public static int GetMacroLineNumber(this UhtType type)
 		{
-			builder.Append(fileId).Append('_').Append(lineNumber).Append('_').Append(macroSuffix);
-			if (includeSuffix)
+			return type switch
 			{
-				// TODO: Both flags should never be set at once, so keeping older codegen here.
-				// Should this list be expanded for the other define scopes??
-				if (defineScope.HasAnyFlags(UhtDefineScope.EditorOnlyData | UhtDefineScope.Editor))
-				{
-					builder.Append("_EOD");
-				}
-			}
+				UhtClass classObj => classObj.GeneratedBodyLineNumber,
+				UhtScriptStruct scriptStruct => scriptStruct.MacroDeclaredLineNumber,
+				UhtFunction function => function.MacroLineNumber,
+				_ => throw new UhtException(type, "Can not use given type to create a macro"),
+			};
+		}
+
+		public static string GetScopeSuffix(this UhtDefineScope definedScope)
+		{
+			return definedScope switch
+			{
+				UhtDefineScope.None => "",
+				UhtDefineScope.EditorOnlyData => "_EOD",
+				UhtDefineScope.Editor => "_EDITOR",
+				UhtDefineScope.Tests => "_TESTS",
+				UhtDefineScope.VerseVM => "_VM",
+				UhtDefineScope.VerseBPVM => "_BPVM",
+				_ => $"_{(int)definedScope}",
+			};
+		}
+
+		public static StringBuilder AppendMacroName(this StringBuilder builder, string fileId, int lineNumber, string macroSuffix)
+		{
+			builder.Append(fileId).Append('_').Append(lineNumber).Append('_').Append(macroSuffix);
 			return builder;
 		}
 
-		public static StringBuilder AppendMacroName(this StringBuilder builder, UhtHeaderCodeGenerator generator, int lineNumber, string macroSuffix, UhtDefineScope defineScope = UhtDefineScope.None, bool includeSuffix = true)
+		public static StringBuilder AppendMacroName(this StringBuilder builder, UhtHeaderCodeGenerator generator, int lineNumber, string macroSuffix)
 		{
-			return builder.AppendMacroName(generator.FileId, lineNumber, macroSuffix, defineScope, includeSuffix);
+			return builder.AppendMacroName(generator.FileId, lineNumber, macroSuffix);
 		}
 
-		public static StringBuilder AppendMacroName(this StringBuilder builder, UhtHeaderCodeGenerator generator, UhtType type, string macroSuffix, UhtDefineScope defineScope = UhtDefineScope.None, bool includeSuffix = true)
+		public static StringBuilder AppendMacroName(this StringBuilder builder, UhtHeaderCodeGenerator generator, UhtType type, string macroSuffix)
 		{
-			if (type is UhtClass classObj)
-			{
-				return builder.AppendMacroName(generator, classObj.GeneratedBodyLineNumber, macroSuffix, defineScope, includeSuffix);
-			}
-			else if (type is UhtScriptStruct scriptStruct)
-			{
-				return builder.AppendMacroName(generator, scriptStruct.MacroDeclaredLineNumber, macroSuffix, defineScope, includeSuffix);
-			}
-			else if (type is UhtFunction function)
-			{
-				return builder.AppendMacroName(generator, function.MacroLineNumber, macroSuffix, defineScope, includeSuffix);
-			}
-			else
+			return builder.AppendMacroName(generator, type.GetMacroLineNumber(), macroSuffix);
+		}
+
+		/// <summary>
+		/// Append the macro name including an optional suffix to uniquely identify a define scope
+		/// </summary>
+		/// <param name="builder"></param>
+		/// <param name="fileId"></param>
+		/// <param name="lineNumber"></param>
+		/// <param name="macroSuffix"></param>
+		/// <param name="defineScope"></param>
+		/// <param name="includeSuffix"></param>
+		/// <returns></returns>
+		public static StringBuilder AppendMacroName(this StringBuilder builder, string fileId, int lineNumber, string macroSuffix, UhtDefineScope defineScope, bool includeSuffix)
+		{
+			builder.AppendMacroName(fileId, lineNumber, macroSuffix);
+			if (includeSuffix && defineScope != UhtDefineScope.None)
 			{
-				throw new UhtException(type, "Can not use given type to create a macro");
+				builder.Append(defineScope.GetScopeSuffix());
 			}
+			return builder;
 		}
 
 		/// <summary>
diff --git a/Engine/Source/Programs/Shared/EpicGames.UHT/Exporters/CodeGen/UhtMacroCreator.cs b/Engine/Source/Programs/Shared/EpicGames.UHT/Exporters/CodeGen/UhtMacroCreator.cs
--- a/Engine/Source/Programs/Shared/EpicGames.UHT/Exporters/CodeGen/UhtMacroCreator.cs
+++ b/Engine/Source/Programs/Shared/EpicGames.UHT/Exporters/CodeGen/UhtMacroCreator.cs
@@ -12,16 +12,23 @@ namespace EpicGames.UHT.Exporters.CodeGen
 		private readonly StringBuilder _builder;
 		private readonly int _startingLength;
 
-		public UhtMacroCreator(StringBuilder builder, UhtHeaderCodeGenerator generator, int lineNumber, string macroSuffix, UhtDefineScope defineScope = UhtDefineScope.None, bool includeSuffix = true)
+		public UhtMacroCreator(StringBuilder builder, UhtHeaderCodeGenerator generator, int lineNumber, string macroSuffix)
 		{
-			builder.Append("#define ").AppendMacroName(generator, lineNumber, macroSuffix, defineScope, includeSuffix).Append(" \\\r\n");
+			builder.Append("#define ").AppendMacroName(generator.FileId, lineNumber, macroSuffix).Append(" \\\r\n");
+			_builder = builder;
+			_startingLength = builder.Length;
+		}
+
+		public UhtMacroCreator(StringBuilder builder, UhtHeaderCodeGenerator generator, UhtType type, string macroSuffix)
+		{
+			builder.Append("#define ").AppendMacroName(generator.FileId, type.GetMacroLineNumber(), macroSuffix).Append(" \\\r\n");
 			_builder = builder;
 			_startingLength = builder.Length;
 		}
 
 		public UhtMacroCreator(StringBuilder builder, UhtHeaderCodeGenerator generator, UhtType type, string macroSuffix, UhtDefineScope defineScope = UhtDefineScope.None, bool includeSuffix = true)
 		{
-			builder.Append("#define ").AppendMacroName(generator, type, macroSuffix, defineScope, includeSuffix).Append(" \\\r\n");
+			builder.Append("#define ").AppendMacroName(generator.FileId, type.GetMacroLineNumber(), macroSuffix, defineScope, includeSuffix).Append(" \\\r\n");
 			_builder = builder;
 			_startingLength = builder.Length;
 		}
diff --git a/Engine/Source/Programs/Shared/EpicGames.UHT/Exporters/CodeGen/UhtUsedDefineScopes.cs b/Engine/Source/Programs/Shared/EpicGames.UHT/Exporters/CodeGen/UhtUsedDefineScopes.cs
--- a/Engine/Source/Programs/Shared/EpicGames.UHT/Exporters/CodeGen/UhtUsedDefineScopes.cs
+++ b/Engine/Source/Programs/Shared/EpicGames.UHT/Exporters/CodeGen/UhtUsedDefineScopes.cs
@@ -464,7 +464,7 @@ namespace EpicGames.UHT.Exporters.CodeGen
 		{
 			foreach (UhtDefineScope defineScope in instances.EnumerateDefinedScopes())
 			{
-				builder.Append('\t').AppendMacroName(generator, outerType, macroSuffix, defineScope).Append(" \\\r\n");
+				builder.Append('\t').AppendMacroName(generator.FileId, outerType.GetMacroLineNumber(), macroSuffix, defineScope, true).Append(" \\\r\n");
 			}
 			return builder;
 		}
@@ -485,7 +485,7 @@ namespace EpicGames.UHT.Exporters.CodeGen
 		{
 			if (instances.SoleScope == UhtDefineScope.None)
 			{
-				using UhtMacroCreator macro = new(builder, generator, outerType, macroSuffix, UhtDefineScope.None);
+				using UhtMacroCreator macro = new(builder, generator, outerType, macroSuffix);
 				appendAction(builder, instances.Instances);
 			}
 			else
@@ -506,7 +506,7 @@ namespace EpicGames.UHT.Exporters.CodeGen
 					{
 						builder.Append("#else\r\n");
 					}
-					using (UhtMacroCreator macro = new(builder, generator, outerType, macroSuffix, UhtDefineScope.None))
+					using (UhtMacroCreator macro = new(builder, generator, outerType, macroSuffix))
 					{
 						appendAction(builder, instances.Instances.Where(x => (x.DefineScope & ~defineScope) == 0));
 					}
