Fix to logging being delayed by 2 logs calling UAT with the Launcher

#Synced-CL 49780983

diff --git a/Engine/Source/Developer/LauncherServices/Private/Launcher/LauncherWorker.cpp b/Engine/Source/Developer/LauncherServices/Private/Launcher/LauncherWorker.cpp
--- a/Engine/Source/Developer/LauncherServices/Private/Launcher/LauncherWorker.cpp
+++ b/Engine/Source/Developer/LauncherServices/Private/Launcher/LauncherWorker.cpp
@@ -55,88 +55,80 @@ bool FLauncherWorker::Init( )
 
 uint32 FLauncherWorker::Run( )
 {
-	FString Line;
-
 	LaunchStartTime = FPlatformTime::Seconds();
 
-	auto MessageReceived = [this](const FString& InMessage)
 	{
-		FStringView MessageView = InMessage;
+		auto MessageReceived = [this](FString InMessage)
 		{
-			FStringView PackageDevicePrefix = TEXTVIEW("Running Package@Device:");
-			if (MessageView.StartsWith(PackageDevicePrefix))
+			InMessage.TrimEndInline();
+			FStringView MessageView = InMessage;
 			{
-				FStringView Value = MessageView.RightChop(PackageDevicePrefix.Len());
-				int32 SplitIndex;
-				if (Value.FindChar('@', SplitIndex))
+				FStringView PackageDevicePrefix = TEXTVIEW("Running Package@Device:");
+				if (MessageView.StartsWith(PackageDevicePrefix))
 				{
-					FString Package(Value.SubStr(0, SplitIndex));
-					FString Device(Value.SubStr(SplitIndex + 1, Value.Len()));
-					AddDevicePackagePair(Device, Package);
+					FStringView Value = MessageView.RightChop(PackageDevicePrefix.Len());
+					int32 SplitIndex;
+					if (Value.FindChar('@', SplitIndex))
+					{
+						FString Package(Value.SubStr(0, SplitIndex));
+						FString Device(Value.SubStr(SplitIndex + 1, Value.Len()));
+						AddDevicePackagePair(Device, Package);
+					}
 				}
 			}
-		}
-		OutputMessageReceived.Broadcast(InMessage);
-	};
-
-	// wait for tasks to be completed
-	while (Status == ELauncherWorkerStatus::Busy)
-	{
-		FPlatformProcess::Sleep(0.01f);
+			OutputMessageReceived.Broadcast(InMessage);
+		};
 
-		FString NewLine = FPlatformProcess::ReadPipe(ReadPipe);
-		if (NewLine.Len() > 0)
+		FString Line;
+		auto ReadLines = [this, &MessageReceived, &Line]
 		{
+			FString NewLine = FPlatformProcess::ReadPipe(ReadPipe);
+			if (NewLine.IsEmpty())
+			{
+				return false;
+			}
+
 			// process the string to break it up in to lines
 			Line += NewLine;
-			TArray<FString> StringArray;
-			int32 count = Line.ParseIntoArray(StringArray, TEXT("\n"), true);
-			if (count > 1)
+	
+			if (int32 NewLineIndex; NewLine.FindChar(TCHAR('\n'), NewLineIndex))
 			{
-				for (int32 Index = 0; Index < count-1; ++Index)
+				FStringView NewLineView = FStringView{NewLine}.RightChop(NewLineIndex + 1);
+				MessageReceived(Line.LeftChop(NewLineView.Len() + 1));
+		
+				while (NewLineView.FindChar(TCHAR('\n'), NewLineIndex))
 				{
-					StringArray[Index].TrimEndInline();
-					MessageReceived(StringArray[Index]);
-				}
-				Line = StringArray[count-1];
-				if (NewLine.EndsWith(TEXT("\n")))
-				{
-					Line += TEXT("\n");
+					FStringView _NewLineView = NewLineView.RightChop(NewLineIndex + 1);
+					MessageReceived(FString{NewLineView.LeftChop(_NewLineView.Len() + 1)});
+					NewLineView = _NewLineView;
 				}
+
+				Line = NewLineView;
 			}
-		}
 
-		if (TaskChain.IsValid() && TaskChain->IsChainFinished())
+			return true;
+		};
+
+		// wait for tasks to be completed
+		while (Status == ELauncherWorkerStatus::Busy)
 		{
-			Status = ELauncherWorkerStatus::Completed;
+			FPlatformProcess::Sleep(0.01f);
+
+			ReadLines();
 
-			NewLine = FPlatformProcess::ReadPipe(ReadPipe);
-			while (NewLine.Len() > 0)
+			if (TaskChain.IsValid() && TaskChain->IsChainFinished())
 			{
-				// process the string to break it up in to lines
-				Line += NewLine;
-				TArray<FString> StringArray;
-				int32 count = Line.ParseIntoArray(StringArray, TEXT("\n"), true);
-				if (count > 1)
+				Status = ELauncherWorkerStatus::Completed;
+
+				while (ReadLines());
+
+				// fire off the last line
+				if (Line.Len() > 0)
 				{
-					for (int32 Index = 0; Index < count-1; ++Index)
-					{
-						StringArray[Index].TrimEndInline();
-						MessageReceived(StringArray[Index]);
-					}
-					Line = StringArray[count-1];
-					if (NewLine.EndsWith(TEXT("\n")))
-					{
-						Line += TEXT("\n");
-					}
+					MessageReceived(MoveTemp(Line));
 				}
-
-				NewLine = FPlatformProcess::ReadPipe(ReadPipe);
+				break;
 			}
-
-			// fire off the last line
-			MessageReceived(Line);
-
 		}
 	}
 
diff --git a/Engine/Source/Programs/Shared/EpicGames.Core.Tests/LogEventParserTests.cs b/Engine/Source/Programs/Shared/EpicGames.Core.Tests/LogEventParserTests.cs
--- a/Engine/Source/Programs/Shared/EpicGames.Core.Tests/LogEventParserTests.cs
+++ b/Engine/Source/Programs/Shared/EpicGames.Core.Tests/LogEventParserTests.cs
@@ -53,8 +53,6 @@ public class LogEventParserTests
 		TestMatcher highest = new(LogEventPriority.Highest, LogLevel.Critical);
 		TestMatcher lowest = new(LogEventPriority.Lowest, LogLevel.Debug);
 
-		const string LineToLog = "hello";
-
 		using LogEventParser parser = new(capture);
 		parser.AddMatchers([highest, lowest]);
 
@@ -273,6 +271,67 @@ public class LogEventParserTests
 		Assert.AreEqual(1, lowest.MatchCalls);
 	}
 
+	[TestMethod]
+	public void TestHangingLines()
+	{
+		LoggerCapture capture = new();
+		TestHangingLinesMatcher highest = new(LogEventPriority.Highest, LogLevel.Critical);
+		TestMatcher lowest = new(LogEventPriority.Lowest, LogLevel.Debug, 3);
+
+		using LogEventParser parser = new(capture);
+		parser.AddMatchers([highest, lowest]);
+
+		parser.WriteLine(LineToLog);
+		Assert.AreEqual(1, highest.MatchCalls);
+		Assert.AreEqual(1, lowest.MatchCalls);
+		// 0 because of lowest 3 line count requirement
+		Assert.AreEqual(0, capture.Events.Count);
+		
+		parser.WriteLine("\t" + LineToLog);
+		Assert.AreEqual(2, highest.MatchCalls);
+		Assert.AreEqual(2, lowest.MatchCalls);
+		// 0 because of lowest 3 line count requirement
+		Assert.AreEqual(0, capture.Events.Count);
+
+		parser.WriteLine("\t" + LineToLog);
+		Assert.AreEqual(3, highest.MatchCalls);
+		Assert.AreEqual(3, lowest.MatchCalls);
+		// 3 because lowest 3 line count requirement is met
+		Assert.AreEqual(3, capture.Events.Count);
+		
+		parser.WriteLine("\t" + LineToLog);
+		Assert.AreEqual(4, highest.MatchCalls);
+		// 3 because highest it incomplete hanging
+		Assert.AreEqual(3, lowest.MatchCalls);
+		// 4 because of highest was incomplete hanging
+		Assert.AreEqual(4, capture.Events.Count);
+		
+		parser.WriteLine(LineToLog);
+		// 6 because highest gets called again once complete
+		Assert.AreEqual(6, highest.MatchCalls);
+		Assert.AreEqual(4, lowest.MatchCalls);
+		// 4 because of lowest 3 line count requirement
+		Assert.AreEqual(4, capture.Events.Count);
+		
+		parser.WriteLine(LineToLog);
+		Assert.AreEqual(7, highest.MatchCalls);
+		Assert.AreEqual(5, lowest.MatchCalls);
+		// 4 because of lowest 3 line count requirement
+		Assert.AreEqual(4, capture.Events.Count);
+		
+		parser.WriteLine(LineToLog);
+		Assert.AreEqual(9, highest.MatchCalls);
+		Assert.AreEqual(7, lowest.MatchCalls);
+		// 5 because lowest 3 line count requirement is met and is not hanging
+		Assert.AreEqual(5, capture.Events.Count);
+
+		parser.Flush();
+		Assert.AreEqual(11, highest.MatchCalls);
+		Assert.AreEqual(9, lowest.MatchCalls);
+		// 5 because flushed
+		Assert.AreEqual(7, capture.Events.Count);
+	}
+
 	private const string LogLineProperty = "LogLine";
 	private const string LineToLog = "hello";
 
@@ -302,7 +361,7 @@ public class LogEventParserTests
 		private int _logLineIndex;
 	}
 
-	private class TestMatcher(LogEventPriority priority, LogLevel level) : ILogEventMatcher
+	private class TestMatcher(LogEventPriority priority, LogLevel level, int requiredLineCount = 1) : ILogEventMatcher
 	{
 		public LogEventMatch? Match(ILogCursor cursor)
 		{
@@ -310,9 +369,37 @@ public class LogEventParserTests
 
 			string? line = cursor.CurrentLine;
 			Assert.IsNotNull(line);
+			for (int index = 1; index < requiredLineCount; ++index)
+			{
+				line = cursor[index];
+			}
+
+			LogEventBuilder builder = new(cursor);
+			return builder.ToMatch(Priority, Level, new EventId(999));
+		}
+
+		public LogEventPriority Priority { get; } = priority;
+		public LogLevel Level { get; } = level;
+		public int MatchCalls { get; private set; } = 0;
+
+		public string GetMatcherName() => Priority.ToString();
+	}
+	
+	private class TestHangingLinesMatcher(LogEventPriority priority, LogLevel level, int requiredLineCount = 1) : ILogEventMatcher
+	{
+		public LogEventMatch? Match(ILogCursor cursor)
+		{
+			++MatchCalls;
+
+			string? line = cursor.CurrentLine;
+			Assert.IsNotNull(line);
+			for (int index = 1; index < requiredLineCount; ++index)
+			{
+				line = cursor[index];
+			}
 
 			LogEventBuilder builder = new(cursor);
-			return builder.ToMatch(LogEventPriority.Highest, Level, new EventId(999));
+			return builder.ToMatchWithHangingLines(Priority, Level, new EventId(999));
 		}
 
 		public LogEventPriority Priority { get; } = priority;
diff --git a/Engine/Source/Programs/Shared/EpicGames.Core/ILogCursor.cs b/Engine/Source/Programs/Shared/EpicGames.Core/ILogCursor.cs
--- a/Engine/Source/Programs/Shared/EpicGames.Core/ILogCursor.cs
+++ b/Engine/Source/Programs/Shared/EpicGames.Core/ILogCursor.cs
@@ -27,6 +27,14 @@ namespace EpicGames.Core
 			get;
 		}
 
+		/// <summary>
+		/// Returns the line count
+		/// </summary>
+		int Length
+		{
+			get;
+		}
+
 		/// <summary>
 		/// Index to find the string at the given offset
 		/// </summary>
@@ -60,6 +68,7 @@ namespace EpicGames.Core
 			public string? this[int index] => _inner[(_baseLineNumber + index) - _inner.CurrentLineNumber];
 			public string? CurrentLine => _inner[_baseLineNumber - _inner.CurrentLineNumber];
 			public int CurrentLineNumber => _baseLineNumber;
+			public int Length => _inner.Length - (_baseLineNumber - _inner.CurrentLineNumber);
 		}
 
 		/// <summary>
diff --git a/Engine/Source/Programs/Shared/EpicGames.Core/LogEventBuilder.cs b/Engine/Source/Programs/Shared/EpicGames.Core/LogEventBuilder.cs
--- a/Engine/Source/Programs/Shared/EpicGames.Core/LogEventBuilder.cs
+++ b/Engine/Source/Programs/Shared/EpicGames.Core/LogEventBuilder.cs
@@ -478,10 +478,30 @@ namespace EpicGames.Core
 		/// <param name="level"></param>
 		/// <param name="eventId">The event id</param>
 		/// <param name="priority"></param>
+		/// <param name="isComplete"></param>
 		/// <returns></returns>
-		public LogEventMatch ToMatch(LogEventPriority priority, LogLevel level, EventId eventId)
+		public LogEventMatch ToMatch(LogEventPriority priority, LogLevel level, EventId eventId, bool isComplete = true)
 		{
-			return new LogEventMatch(priority, ToArray(level, eventId));
+			return new LogEventMatch(priority, ToArray(level, eventId), isComplete);
+		}
+		
+		/// <summary>
+		/// Creates a match object at the given priority with hanging lines
+		/// </summary>
+		/// <param name="level"></param>
+		/// <param name="eventId">The event id</param>
+		/// <param name="priority"></param>
+		/// <returns></returns>
+		public LogEventMatch ToMatchWithHangingLines(LogEventPriority priority, LogLevel level, EventId eventId)
+		{
+			while (Next.Length > 0 && IsNextLineHanging())
+			{
+				MoveNext();
+			}
+
+			// If there are more lines, is must not be hanging, and therefore complete
+			bool isComplete = Next.Length > 0;
+			return new LogEventMatch(priority, ToArray(level, eventId), isComplete);
 		}
 	}
 }
diff --git a/Engine/Source/Programs/Shared/EpicGames.Core/LogEventParser.cs b/Engine/Source/Programs/Shared/EpicGames.Core/LogEventParser.cs
--- a/Engine/Source/Programs/Shared/EpicGames.Core/LogEventParser.cs
+++ b/Engine/Source/Programs/Shared/EpicGames.Core/LogEventParser.cs
@@ -73,22 +73,29 @@ namespace EpicGames.Core
 		/// </summary>
 		public List<LogEvent> Events { get; }
 
+		/// <summary>
+		/// Match is complete
+		/// </summary>
+		public bool IsComplete { get; }
+		
 		/// <summary>
 		/// Constructor
 		/// </summary>
-		public LogEventMatch(LogEventPriority priority, LogEvent logEvent)
+		public LogEventMatch(LogEventPriority priority, LogEvent logEvent, bool isComplete)
 		{
 			Priority = priority;
 			Events = [logEvent];
+			IsComplete = isComplete;
 		}
 
 		/// <summary>
 		/// Constructor
 		/// </summary>
-		public LogEventMatch(LogEventPriority priority, IEnumerable<LogEvent> events)
+		public LogEventMatch(LogEventPriority priority, IEnumerable<LogEvent> events, bool isComplete)
 		{
 			Priority = priority;
 			Events = events.ToList();
+			IsComplete = isComplete;
 		}
 	}
 
@@ -196,6 +203,16 @@ namespace EpicGames.Core
 		/// </summary>
 		int _linesParsed = 0;
 
+		/// <summary>
+		/// Incomplete matcher being processed
+		/// </summary>
+		ILogEventMatcher? _incompleteMatcher;
+		
+		/// <summary>
+		/// Incomplete event count being processed
+		/// </summary>
+		int _incompleteEventCount;
+		
 		/// <summary>
 		/// Public accessor for the logger
 		/// </summary>
@@ -619,62 +636,54 @@ namespace EpicGames.Core
 		{
 			_activeTimer.Start();
 			int startLineCount = _buffer.Length;
-			while (_buffer.Length > 0)
+			if (ProcessIncomplete(bFlush))
 			{
-				// Try to match an event
-				List<LogEvent>? events = null;
-				string line = _buffer[0]!;
-				if (line == null)
-				{
-					line = "";
-				}
-
-				if (HandleLogEventParserCommand(line))
-				{
-					continue;
-				}
-
-				if (Regex.IsMatch(line, "<-- Suspend Log Parsing -->", RegexOptions.IgnoreCase))
-				{
-					_matchingEnabled--;
-				}
-				else if (Regex.IsMatch(line, "<-- Resume Log Parsing -->", RegexOptions.IgnoreCase))
+				while (_buffer.Length > 0)
 				{
-					_matchingEnabled++;
-				}
-				else if (_matchingEnabled >= 0 && disableMatching == false)
-				{
-					events = MatchEvent();
-				}
-
-				// Bail out if we need more data
-				if (_buffer.Length < 1024 && !bFlush && _buffer.NeedMoreData)
-				{
-					break;
-				}
-
-				// If we did match something, check if it's not negated by an ignore pattern. We typically have relatively few errors and many more ignore patterns than matchers, so it's quicker 
-				// to check them in response to an identified error than to treat them as matchers of their own.
-				if (events != null)
-				{
-					foreach (Regex ignorePattern in IgnorePatterns)
+					string? line = _buffer[0];
+					if (line == null)
+					{
+						line = "";
+					}
+					else if (HandleLogEventParserCommand(line))
+					{
+						continue;
+					}
+					else if (Regex.IsMatch(line, "<-- Suspend Log Parsing -->", RegexOptions.IgnoreCase))
+					{
+						_matchingEnabled--;
+					}
+					else if (Regex.IsMatch(line, "<-- Resume Log Parsing -->", RegexOptions.IgnoreCase))
+					{
+						_matchingEnabled++;
+					}
+					else if (_matchingEnabled >= 0 && !disableMatching)
 					{
-						if (ignorePattern.IsMatch(line))
+						(ILogEventMatcher? matcher, LogEventMatch? match) = MatchEvent();
+						// Report the error to the listeners
+						// If we did match something, check if it's not negated by an ignore pattern. We typically have relatively few errors and many more ignore patterns than matchers, so it's quicker 
+						// to check them in response to an identified error than to treat them as matchers of their own.
+						if (match != null && !IgnorePatterns.Any(ignorePattern => ignorePattern.IsMatch(line)))
 						{
-							events = null;
-							break;
+							// Bail out if we need more data, or match is incomplete
+							if (_buffer.Length < 1024 && !bFlush && (_buffer.NeedMoreData || !match.IsComplete))
+							{
+								if (!_buffer.NeedMoreData)
+								{
+									WriteEvents(match.Events);
+									Debug.Assert(matcher != null);
+									_incompleteMatcher = matcher;
+									_incompleteEventCount = match.Events.Count;
+								}
+								break;
+							}
+
+							WriteEvents(match.Events);
+							_buffer.Advance(match.Events.Count);
+							continue;
 						}
 					}
-				}
 
-				// Report the error to the listeners
-				if (events != null)
-				{
-					WriteEvents(events);
-					_buffer.Advance(events.Count);
-				}
-				else
-				{
 					_logger.Log(LogLevel.Information, KnownLogEvents.None, line, null, (state, exception) => state);
 					_buffer.MoveNext();
 				}
@@ -701,40 +710,66 @@ namespace EpicGames.Core
 			}
 		}
 
+		/// <summary>
+		/// Process any incomplete match
+		/// </summary>
+		/// <param name="bFlush">Whether we've reached the end of the stream</param>
+		/// <returns>Whether it finished processing the incomplete match</returns>
+		bool ProcessIncomplete(bool bFlush)
+		{
+			if (_incompleteMatcher != null)
+			{
+				LogEventMatch match = _incompleteMatcher.Match(_buffer)!;
+				Debug.Assert(match.Events.Count >= _incompleteEventCount);
+				WriteEvents(match.Events.Skip(_incompleteEventCount));
+
+				// Bail out if we need more data
+				if (_buffer.Length < 1024 && !bFlush && !match.IsComplete)
+				{
+					_incompleteEventCount = match.Events.Count;
+					return false;
+				}
+
+				_incompleteMatcher = null;
+				_incompleteEventCount = 0;
+				_buffer.Advance(match.Events.Count);
+			}
+
+			return true;
+		}
+		
 		/// <summary>
 		/// Try to match an event from the current buffer
 		/// </summary>
-		/// <returns>The matched event</returns>
-		private List<LogEvent>? MatchEvent()
+		/// <returns>The matcher and the match</returns>
+		private (ILogEventMatcher?, LogEventMatch?) MatchEvent()
 		{
+			ILogEventMatcher? currentMatcher = null;
 			LogEventMatch? currentMatch = null;
 			foreach (ILogEventMatcher matcher in _enabledMatchers)
 			{
-				LogEventMatch? match = null;
 				try
 				{
-					match = matcher.Match(_buffer);
+					LogEventMatch? match = matcher.Match(_buffer);
+					if (match != null && (currentMatch == null || match.Priority > currentMatch.Priority))
+					{
+						currentMatcher = matcher;
+						currentMatch = match;
+					}
 				}
 				catch (Exception ex)
 				{
 					_logger.LogWarning(KnownLogEvents.Systemic_LogEventMatcher, ex, "Exception while parsing log events with {Type}. Buffer size {Length}.", matcher.GetType().Name, _buffer.Length);
 				}
-				if (match != null)
-				{
-					if (currentMatch == null || match.Priority > currentMatch.Priority)
-					{
-						currentMatch = match;
-					}
-				}
 			}
-			return currentMatch?.Events;
+			return (currentMatcher, currentMatch);
 		}
 
 		/// <summary>
 		/// Writes an event to the log
 		/// </summary>
 		/// <param name="logEvents">The event to write</param>
-		protected virtual void WriteEvents(List<LogEvent> logEvents)
+		protected virtual void WriteEvents(IEnumerable<LogEvent> logEvents)
 		{
 			foreach (LogEvent logEvent in logEvents)
 			{
diff --git a/Engine/Source/Programs/Shared/EpicGames.LogMatchers/ContentEventMatcher.cs b/Engine/Source/Programs/Shared/EpicGames.LogMatchers/ContentEventMatcher.cs
--- a/Engine/Source/Programs/Shared/EpicGames.LogMatchers/ContentEventMatcher.cs
+++ b/Engine/Source/Programs/Shared/EpicGames.LogMatchers/ContentEventMatcher.cs
@@ -33,19 +33,13 @@ namespace EpicGames.LogMatchers
 				Match? match = s_pattern.Match(input.CurrentLine!);
 				if (match.Success)
 				{
-					LogEventBuilder builder = new LogEventBuilder(input);
-
+					LogEventBuilder builder = new(input);
 					builder.Annotate(match.Groups["channel"], LogEventMarkup.Channel);
 					builder.Annotate(match.Groups["severity"], LogEventMarkup.Severity);
 					builder.AnnotateAsset(match.Groups["asset"]);
 					builder.Annotate(match.Groups["message"], LogEventMarkup.Message);
 
-					while (builder.IsNextLineHanging())
-					{
-						builder.MoveNext();
-					}
-
-					return builder.ToMatch(LogEventPriority.AboveNormal, GetLogLevelFromSeverity(match), KnownLogEvents.Engine_AssetLog);
+					return builder.ToMatchWithHangingLines(LogEventPriority.AboveNormal, GetLogLevelFromSeverity(match), KnownLogEvents.Engine_AssetLog);
 				}
 			}
 			return null;
diff --git a/Engine/Source/Programs/Shared/EpicGames.LogMatchers/GenericEventMatcher.cs b/Engine/Source/Programs/Shared/EpicGames.LogMatchers/GenericEventMatcher.cs
--- a/Engine/Source/Programs/Shared/EpicGames.LogMatchers/GenericEventMatcher.cs
+++ b/Engine/Source/Programs/Shared/EpicGames.LogMatchers/GenericEventMatcher.cs
@@ -38,12 +38,7 @@ namespace EpicGames.LogMatchers
 					level = LogLevel.Warning;
 				}
 
-				LogEventBuilder builder = new LogEventBuilder(cursor);
-				while (builder.IsNextLineHanging())
-				{
-					builder.MoveNext();
-				}
-				return builder.ToMatch(LogEventPriority.Lowest, level, KnownLogEvents.Generic);
+				return new LogEventBuilder(cursor).ToMatchWithHangingLines(LogEventPriority.Lowest, level, KnownLogEvents.Generic);
 			}
 			if (cursor.IsMatch(s_errorPattern))
 			{
diff --git a/Engine/Source/Programs/Shared/EpicGames.LogMatchers/LogChannelEventMatcher.cs b/Engine/Source/Programs/Shared/EpicGames.LogMatchers/LogChannelEventMatcher.cs
--- a/Engine/Source/Programs/Shared/EpicGames.LogMatchers/LogChannelEventMatcher.cs
+++ b/Engine/Source/Programs/Shared/EpicGames.LogMatchers/LogChannelEventMatcher.cs
@@ -23,21 +23,15 @@ namespace EpicGames.LogMatchers
 		{
 			if (input.TryMatch(s_pattern, out Match? match))
 			{
-
 				if (input.Contains("There is not enough space on the disk."))
 				{
 					return new LogEventBuilder(input).ToMatch(LogEventPriority.Low, LogLevel.Error, KnownLogEvents.Systemic_OutOfDiskSpace);
 				}
 
-				LogEventBuilder builder = new LogEventBuilder(input);
+				LogEventBuilder builder = new(input);
 				builder.Annotate(match.Groups["channel"], LogEventMarkup.Channel);
 				builder.Annotate(match.Groups["severity"], LogEventMarkup.Severity);
 
-				while (builder.IsNextLineHanging())
-				{
-					builder.MoveNext();
-				}
-
 				LogLevel level = match.Groups["severity"].Value switch
 				{
 					"Error" => LogLevel.Error,
@@ -45,7 +39,7 @@ namespace EpicGames.LogMatchers
 					_ => LogLevel.Information,
 				};
 
-				return builder.ToMatch(LogEventPriority.Low, level, KnownLogEvents.Engine_LogChannel);
+				return builder.ToMatchWithHangingLines(LogEventPriority.Low, level, KnownLogEvents.Engine_LogChannel);
 
 			}
 			return null;
