[UHT] Export auto-generated getters/setters when the parent class is Public but unexported/MinimalAPI

#Synced-CL 49757048

diff --git a/Engine/Source/Programs/Shared/EpicGames.UHT/Exporters/CodeGen/UhtHeaderCodeGeneratorHFile.cs b/Engine/Source/Programs/Shared/EpicGames.UHT/Exporters/CodeGen/UhtHeaderCodeGeneratorHFile.cs
--- a/Engine/Source/Programs/Shared/EpicGames.UHT/Exporters/CodeGen/UhtHeaderCodeGeneratorHFile.cs
+++ b/Engine/Source/Programs/Shared/EpicGames.UHT/Exporters/CodeGen/UhtHeaderCodeGeneratorHFile.cs
@@ -955,7 +955,7 @@ namespace EpicGames.UHT.Exporters.CodeGen
 				}
 
 				AppendFieldNotify(builder, classObj);
-				AppendAutoGettersSetters(builder, classObj, autoGetterSetterProperties);
+				AppendAutoGettersSetters(builder, classObj, api, autoGetterSetterProperties);
 				AppendProlog(builder, classObj);
 				AppendGeneratedBodyMacroBlock(builder, classObj, classObj, usesLegacy, rpcFunctions, verseNativeCallableFunctions,
 					autoGetterSetterProperties, sparseScriptStructs.Any(), sparseProperties, getterSetterProperties, hasCallbacks,
@@ -1078,7 +1078,7 @@ namespace EpicGames.UHT.Exporters.CodeGen
 				});
 		}
 
-		private StringBuilder AppendAutoGettersSetters(StringBuilder builder, UhtClass classObj, UhtUsedDefineScopes<UhtProperty> autoGetterSetterProperties)
+		private StringBuilder AppendAutoGettersSetters(StringBuilder builder, UhtClass classObj, string api, UhtUsedDefineScopes<UhtProperty> autoGetterSetterProperties)
 		{
 			if (autoGetterSetterProperties.IsEmpty)
 			{
@@ -1093,12 +1093,12 @@ namespace EpicGames.UHT.Exporters.CodeGen
 						if (property.PropertyExportFlags.HasAnyFlags(UhtPropertyExportFlags.GetterSpecifiedAuto))
 						{
 							string getterCallText = property.Getter ?? "Get" + property.SourceName;
-							builder.Append('\t').AppendPropertyText(property, UhtPropertyTextType.GetterRetVal).Append(getterCallText).Append("() const; \\\r\n");
+							builder.Append('\t').Append(api).AppendPropertyText(property, UhtPropertyTextType.GetterRetVal).Append(getterCallText).Append("() const; \\\r\n");
 						}
 						if (property.PropertyExportFlags.HasAnyFlags(UhtPropertyExportFlags.SetterSpecifiedAuto))
 						{
 							string setterCallText = property.Setter ?? "Set" + property.SourceName;
-							builder.Append("\tvoid ").Append(setterCallText).Append('(').AppendPropertyText(property, UhtPropertyTextType.SetterParameterArgType).Append("InValue").Append("); \\\r\n");
+							builder.Append('\t').Append(api).Append("void ").Append(setterCallText).Append('(').AppendPropertyText(property, UhtPropertyTextType.SetterParameterArgType).Append("InValue").Append("); \\\r\n");
 						}
 					}
 				});
