Horde: Add retry logic to TestAgentMessageLoopTcpAesAsync to handle CI load variability

diff --git a/Engine/Source/Programs/Shared/EpicGames.Horde.Tests/Compute/ComputeSocketTests.cs b/Engine/Source/Programs/Shared/EpicGames.Horde.Tests/Compute/ComputeSocketTests.cs
--- a/Engine/Source/Programs/Shared/EpicGames.Horde.Tests/Compute/ComputeSocketTests.cs
+++ b/Engine/Source/Programs/Shared/EpicGames.Horde.Tests/Compute/ComputeSocketTests.cs
@@ -121,6 +121,7 @@ namespace EpicGames.Horde.Tests.Compute
 		}
 
 		[TestMethod]
+		[Retry(2)] // Retry up to 2 times on failure due to CI load variability
 		public async Task TestAgentMessageLoopTcpAesAsync()
 		{
 			using CancellationTokenSource cts = new(TestTimoutMs);
