[Horde][Analytics][Shared] Expanding capabilities of Analytics.Telemetry to simplify telemetry emissions & reduce duplication keys across code.

#Synced-CL 49874115

diff --git a/Engine/Source/Programs/Shared/EpicGames.Analytics.Telemetry/AbstractTelemetryRecord.cs b/Engine/Source/Programs/Shared/EpicGames.Analytics.Telemetry/AbstractTelemetryRecord.cs
--- a/Engine/Source/Programs/Shared/EpicGames.Analytics.Telemetry/AbstractTelemetryRecord.cs
+++ b/Engine/Source/Programs/Shared/EpicGames.Analytics.Telemetry/AbstractTelemetryRecord.cs
@@ -1,5 +1,6 @@
 // Copyright Epic Games, Inc. All Rights Reserved.
 
+using System.Collections;
 using System.ComponentModel.DataAnnotations.Schema;
 
 namespace EpicGames.Analytics
@@ -9,6 +10,8 @@ namespace EpicGames.Analytics
 	/// </summary>
 	public abstract record AbstractTelemetryRecord
 	{
+		#region -- Public Properties --
+
 		/// <summary>
 		/// The telemetry event name.
 		/// </summary>
@@ -21,6 +24,39 @@ namespace EpicGames.Analytics
 		[Column("schema_version")]
 		public int SchemaVersion { get; init; }
 
+		#endregion -- Public Properties --
+
+		#region -- Public Api --
+
+		/// <summary>
+		/// Produces a hashtable representation of the telemetry record.
+		/// </summary>
+		/// <returns>Hashtable representing the underyling telemetry record.</returns>
+		/// <remarks>
+		///		This is used primarily for interop in <see cref="Hashtable"/> oriented systems.
+		///		For more compositional telemetry structs, utilize <see cref="System.Text.Json.Serialization.JsonConverter"/>.
+		/// </remarks>
+		public Hashtable ToHashtable()
+		{
+			Hashtable result = new Hashtable();
+			WriteProperties(result);
+			return result;
+		}
+
+		#endregion -- Public Api --
+
+		#region -- Protected Api --
+
+		/// <summary>
+		/// No arg constructor for ORM construction.
+		/// </summary>
+		/// <remarks>used for ORM instantiation.</remarks>
+		protected AbstractTelemetryRecord()
+		{
+			EventName = System.String.Empty;
+			SchemaVersion = -1;
+		}
+
 		/// <summary>
 		/// Constructor.
 		/// </summary>
@@ -31,5 +67,18 @@ namespace EpicGames.Analytics
 			EventName = eventName;
 			SchemaVersion = schemaVersion;
 		}
+
+		/// <summary>
+		/// Template pattern method that appends serializable properties to the in hashtable.
+		/// </summary>
+		/// <param name="hashtable">The input hashtable to apply properties to.</param>
+		/// <remarks>When overriding, a call to the base should be made.</remarks>
+		protected virtual void WriteProperties(Hashtable hashtable)
+		{
+			hashtable[nameof(EventName)] = EventName;
+			hashtable[nameof(SchemaVersion)] = SchemaVersion;
+		}
+
+		#endregion -- Protected Api --
 	}
 }
\ No newline at end of file
diff --git a/Engine/Source/Programs/Shared/EpicGames.Analytics.Telemetry/AnalyticsTelemetryHelpers.cs b/Engine/Source/Programs/Shared/EpicGames.Analytics.Telemetry/AnalyticsTelemetryHelpers.cs
new file mode 100644
--- /dev/null
+++ b/Engine/Source/Programs/Shared/EpicGames.Analytics.Telemetry/AnalyticsTelemetryHelpers.cs
@@ -0,0 +1,55 @@
+// Copyright Epic Games, Inc. All Rights Reserved.
+
+using System.Collections;
+
+namespace EpicGames.Analytics
+{
+	/// <summary>
+	/// Common utilities for working with telemetry records.
+	/// </summary>
+	public static class AnalyticsTelemetryHelpers
+	{
+		/// <summary>
+		/// Merges the right hashtable into the left, without mutating the left hashtable.
+		/// </summary>
+		/// <param name="left">The target hashtable to be merged into.</param>
+		/// <param name="right">The hashtable to merge from.</param>
+		/// <param name="overrideTarget">True to override keys in left, false otherwise.</param>
+		/// <returns>The left hashtable with all the right table's elements merged in.</returns>
+		public static Hashtable MergeLeftCopy(Hashtable left, Hashtable right, bool overrideTarget = true)
+		{
+			Hashtable result = new Hashtable(left);
+
+			MergeLeft(result, right, overrideTarget);
+
+			return result;
+		}
+
+		/// <summary>
+		/// Merges the right hashtable into the left.
+		/// </summary>
+		/// <param name="left">The target hashtable to be merged into.</param>
+		/// <param name="right">The hashtable to merge from.</param>
+		/// <param name="overrideTarget">True to override keys in left, false otherwise.</param>
+		/// <returns>The left hashtable with all the right table's elements merged in. Overwrites the values in the left.</returns>
+		public static Hashtable MergeLeft(Hashtable left, Hashtable right, bool overrideTarget = true)
+		{
+			if (left == right)
+			{
+				return left;
+			}
+
+			foreach (DictionaryEntry entry in right)
+			{
+				if (!overrideTarget && left.ContainsKey(entry.Key))
+				{
+					continue;
+				}
+
+				left[entry.Key] = entry.Value;
+			}
+
+			return left;
+		}
+	}
+}
\ No newline at end of file
diff --git a/Engine/Source/Programs/Shared/EpicGames.Analytics.Telemetry/ColumnValidator.cs b/Engine/Source/Programs/Shared/EpicGames.Analytics.Telemetry/ColumnValidator.cs
new file mode 100644
--- /dev/null
+++ b/Engine/Source/Programs/Shared/EpicGames.Analytics.Telemetry/ColumnValidator.cs
@@ -0,0 +1,139 @@
+// Copyright Epic Games, Inc. All Rights Reserved.
+
+using System;
+using System.Collections;
+using System.Collections.Generic;
+using System.ComponentModel.DataAnnotations.Schema;
+using System.Reflection;
+
+namespace EpicGames.Analytics
+{
+	/// <summary>
+	/// Utility class used to validate the payload contents against specific record types, ensuring the presence of <see cref="ColumnAttribute"/> within different payload types.
+	/// </summary>
+	/// <typeparam name="T">The record type to validate against.</typeparam>
+	public static class ColumnValidator<T>
+	{
+		private static readonly object s_initLock = new();
+		private static HashSet<string>? s_columnKeys = null;
+		private static HashSet<string>? s_propertyKeys = null;
+
+		private static void InitializeKeys()
+		{
+			if (s_columnKeys != null && s_propertyKeys != null)
+			{
+				return;
+			}
+
+			lock (s_initLock)
+			{
+				HashSet<string> columnKeys = new(16, StringComparer.OrdinalIgnoreCase);
+				HashSet<string> propertyKeys = new(16, StringComparer.OrdinalIgnoreCase);
+
+				PropertyInfo[] props = typeof(T).GetProperties(
+					BindingFlags.Instance |
+					BindingFlags.Public |
+					BindingFlags.NonPublic |
+					BindingFlags.DeclaredOnly);
+
+				foreach (PropertyInfo prop in props)
+				{
+					ColumnAttribute? columnAttr = prop.GetCustomAttribute<ColumnAttribute>();
+					if (columnAttr != null && columnAttr.Name != null)
+					{
+						columnKeys.Add(columnAttr.Name.ToUpperInvariant());
+						propertyKeys.Add(prop.Name.ToUpperInvariant());
+					}
+				}
+
+				s_columnKeys = columnKeys;
+				s_propertyKeys = propertyKeys;
+			}
+		}
+
+		#region -- Private Static Api --
+
+		private static HashSet<string> ReplicateAsHashSet(Hashtable source, bool fuzzyReplication)
+		{
+			HashSet<string> sourceKeys = new HashSet<string>(source.Count, StringComparer.OrdinalIgnoreCase);
+			foreach (object key in source.Keys)
+			{
+				if (key is string s)
+				{
+					// If fuzzy replication, we trim down the key and remove any unneccesary symbols since we tend to have clean property names.
+					if (fuzzyReplication)
+					{
+						s = s.Trim()
+							.Replace("_", "", StringComparison.Ordinal)
+							.Replace(" ", "", StringComparison.Ordinal);
+					}
+
+					sourceKeys.Add(s);
+				}
+			}
+
+			return sourceKeys;
+		}
+
+		private static bool HasCompleteKeySet(Hashtable source, HashSet<string> comparisonKeys, bool fuzzyMatch, List<string>? missingColumns)
+		{
+			HashSet<string> sourceKeys = ReplicateAsHashSet(source, fuzzyMatch);
+
+			bool allPresent = true;
+			foreach (string requiredKey in comparisonKeys!)
+			{
+				if (!sourceKeys.Contains(requiredKey))
+				{
+					missingColumns?.Add(requiredKey);
+					allPresent = false;
+				}
+			}
+
+			return allPresent;
+		}
+
+		#endregion -- Private Static Api --
+
+		/// <summary>
+		/// Validates whether an source payload of <see cref="Hashtable"/> type contains all the columns.
+		/// </summary>
+		/// <param name="source">The Hashtable to test.</param>
+		/// <param name="missingColumns">List to store missing columns against.</param>
+		/// <returns>True if all <see cref="ColumnAttribute.Name"/> in the type are contained within the source, false otherwise.</returns>
+		/// <exception cref="ReflectionTypeLoadException">Can be thrown due to reflective inspection of <see cref="ColumnValidator{T}"/>.</exception>
+		/// <remarks>This only validates at the current inheritance hierarchy of the <see cref="ColumnValidator{T}"/> in order to support narrow validation.</remarks>
+#pragma warning disable CA1000 // Do not declare static members on generic types
+		public static bool HasCompleteColumnSet(Hashtable source, List<string>? missingColumns)
+#pragma warning restore CA1000 // Do not declare static members on generic types
+		{
+			if (s_columnKeys == null)
+			{
+				InitializeKeys();
+			}
+
+			// Column comparisons should *never* be fuzzy.
+			return HasCompleteKeySet(source, s_columnKeys!, false, missingColumns);
+		}
+
+		/// <summary>
+		/// Validates whether an source payload of <see cref="Hashtable"/> type contains all the columns.
+		/// </summary>
+		/// <param name="source">The Hashtable to test.</param>
+		/// <param name="fuzzyMatch">Whether to use fuzzy matching for proprty comparison, or not.</param>
+		/// <param name="missingColumns">List to store missing columns against.</param>
+		/// <exception cref="ReflectionTypeLoadException">Can be thrown due to reflective inspection of <see cref="ColumnValidator{T}"/>.</exception>
+		/// <returns>True if all properties annotated with <see cref="ColumnAttribute"/> in the type are contained within the source, false otherwise.</returns>
+		/// <remarks>This only validates at the current inheritance hierarchy of the <see cref="ColumnValidator{T}"/> in order to support narrow validation.</remarks>
+#pragma warning disable CA1000 // Do not declare static members on generic types
+		public static bool HasCompletePropertySet(Hashtable source, bool fuzzyMatch, List<string>? missingColumns)
+#pragma warning restore CA1000 // Do not declare static members on generic types
+		{
+			if (s_propertyKeys == null)
+			{
+				InitializeKeys();
+			}
+
+			return HasCompleteKeySet(source, s_propertyKeys!, fuzzyMatch, missingColumns);
+		}
+	}
+}
\ No newline at end of file
diff --git a/Engine/Source/Programs/Shared/EpicGames.Analytics.Telemetry/CorrelatableTelemetryRecord.cs b/Engine/Source/Programs/Shared/EpicGames.Analytics.Telemetry/CorrelatableTelemetryRecord.cs
new file mode 100644
--- /dev/null
+++ b/Engine/Source/Programs/Shared/EpicGames.Analytics.Telemetry/CorrelatableTelemetryRecord.cs
@@ -0,0 +1,73 @@
+// Copyright Epic Games, Inc. All Rights Reserved.
+
+using System;
+using System.Collections;
+using System.ComponentModel.DataAnnotations.Schema;
+
+namespace EpicGames.Analytics
+{
+	/// <summary>
+	/// Record that describes a telemetry that can be correlated inter and intra session.
+	/// </summary>
+	public record CorrelatableTelemetryRecord : AbstractTelemetryRecord
+	{
+		#region -- Public Properties --
+
+		/// <summary>
+		/// A unique id used to correlate telemetry within a session.
+		/// </summary>
+		[Column("session_id")]
+		public string? SessionId { get; init; }
+
+		/// <summary>
+		/// A label used to correlate telemetry across sessions.
+		/// </summary>
+		[Column("session_label")]
+		public string? SessionLabel { get; init; }
+
+		#endregion -- Public Properties --
+
+		#region -- Public Api --
+
+		/// <summary>
+		/// No arg constructor for ORM construction.
+		/// </summary>
+		/// <remarks>used for ORM instantiation.</remarks>
+		protected CorrelatableTelemetryRecord() : base() { }
+
+		/// <summary>
+		/// Constructor for correlatable telemetry.
+		/// </summary>
+		/// <param name="eventName">The event name.</param>
+		/// <param name="schemaVersion">The schema version.</param>
+		/// <param name="sessionId">The id used to correlate telemetry emitted in the same session.</param>
+		/// <param name="sessionLabel">The label used to correlate telemetry emitted across multiple sessions.</param>
+		public CorrelatableTelemetryRecord(string eventName, int schemaVersion, string? sessionId, string? sessionLabel) : base(eventName, schemaVersion)
+		{
+			SessionId = sessionId;
+			SessionLabel = sessionLabel;
+		}
+
+		#endregion -- Public Api --
+
+		#region -- Protected Api --
+
+		/// <inheritdoc/>
+		protected override void WriteProperties(Hashtable hashtable)
+		{
+			base.WriteProperties(hashtable);
+
+			if (!String.IsNullOrEmpty(SessionId))
+			{
+				hashtable[nameof(SessionId)] = SessionId;
+			}
+
+			if (!String.IsNullOrEmpty(SessionLabel))
+			{
+				hashtable[nameof(SessionLabel)] = SessionLabel;
+			}
+		}
+
+		#endregion -- Protected Api --
+	}
+}
\ No newline at end of file
diff --git a/Engine/Source/Programs/Shared/EpicGames.Horde/EpicGames.Horde.csproj b/Engine/Source/Programs/Shared/EpicGames.Horde/EpicGames.Horde.csproj
--- a/Engine/Source/Programs/Shared/EpicGames.Horde/EpicGames.Horde.csproj
+++ b/Engine/Source/Programs/Shared/EpicGames.Horde/EpicGames.Horde.csproj
@@ -32,6 +32,7 @@
   </ItemGroup>
 
   <ItemGroup>
+    <ProjectReference Include="..\EpicGames.Analytics.Telemetry\EpicGames.Analytics.Telemetry.csproj" />
     <ProjectReference Include="..\EpicGames.Core\EpicGames.Core.csproj" />
     <ProjectReference Include="..\EpicGames.IoHash\EpicGames.IoHash.csproj" />
     <ProjectReference Include="..\EpicGames.OIDC\EpicGames.OIDC.csproj" />
diff --git a/Engine/Source/Programs/Shared/EpicGames.Horde/Telemetry/HordeContextTelemetryRecord.cs b/Engine/Source/Programs/Shared/EpicGames.Horde/Telemetry/HordeContextTelemetryRecord.cs
new file mode 100644
--- /dev/null
+++ b/Engine/Source/Programs/Shared/EpicGames.Horde/Telemetry/HordeContextTelemetryRecord.cs
@@ -0,0 +1,224 @@
+// Copyright Epic Games, Inc. All Rights Reserved.
+
+using System;
+using System.Collections;
+using System.ComponentModel.DataAnnotations.Schema;
+using EpicGames.Analytics;
+using EpicGames.Horde.Jobs;
+using EpicGames.Horde.Jobs.Templates;
+using EpicGames.Horde.Streams;
+using EpicGames.Horde.Utilities;
+
+namespace EpicGames.Horde.Telemetry
+{
+	/// <summary>
+	/// Record that describes telemetry that can contain common Horde context (and whether it is within a build context or not).
+	/// </summary>
+	public record HordeContextTelemetryRecord : CorrelatableTelemetryRecord
+	{
+		#region -- Public Properties --
+
+		/// <summary>
+		/// Whether context is build machine, or not.
+		/// </summary>
+		[Column("is_build_machine")]
+		public bool IsBuildMachine { get; init; }
+
+		/// <summary>
+		/// The Horde URL.
+		/// </summary>
+		[Column("horde_url")]
+#pragma warning disable CA1056 // URI-like properties should not be strings
+		public string? HordeUrlStr { get; init; }
+#pragma warning restore CA1056 // URI-like properties should not be strings
+
+		/// <summary>
+		/// The stream id.
+		/// </summary>
+		[Column("stream_id")]
+		public StreamId? StreamId { get; init; }
+
+		/// <summary>
+		/// The template id.
+		/// </summary>
+		[Column("template_id")]
+		public TemplateId? TemplateId { get; init; }
+
+		/// <summary>
+		/// The Horde job id.
+		/// </summary>
+		[Column("job_id")]
+		public JobId? JobId { get; init; }
+
+		/// <summary>
+		/// The Horde step id.
+		/// </summary>
+		[Column("step_id")]
+		public JobStepId? StepId { get; init; }
+
+		/// <summary>
+		/// The change id.
+		/// </summary>
+		[Column("commit_id")]
+		public int? CommitId { get; init; }
+
+		#endregion -- Public Properties --
+
+		#region -- Public Api --
+
+		/// <summary>
+		/// No arg constructor for ORM construction.
+		/// </summary>
+		/// <remarks>used for ORM instantiation.</remarks>
+		protected HordeContextTelemetryRecord() : base() { }
+
+		/// <summary>
+		/// Constructor for standard Horde context.
+		/// </summary>
+		/// <param name="eventName">The event name.</param>
+		/// <param name="schemaVersion">The schema version.</param>
+		/// <param name="sessionId">The id used to correlate telemetry emitted in the same session.</param>
+		/// <param name="sessionLabel">The label used to correlate telemetry emitted across multiple sessions.</param>
+		/// <param name="isBuildMachine">True whether this is a build machine context, false otherwise.</param>
+		/// <param name="hordeUrlStr">The horde url.</param>
+		/// <param name="streamId">The stream id.</param>
+		/// <param name="templateId">The template id.</param>
+		/// <param name="jobId">The Horde job id.</param>
+		/// <param name="stepId">The Horde step id.</param>
+		/// <param name="commitId">The commitId.</param>
+		public HordeContextTelemetryRecord(string eventName, int schemaVersion, string? sessionId, string? sessionLabel, bool isBuildMachine, Uri? hordeUrlStr, StreamId? streamId, TemplateId? templateId, JobId? jobId, JobStepId? stepId, int? commitId) : base(eventName, schemaVersion, sessionId, sessionLabel)
+		{
+			IsBuildMachine = isBuildMachine;
+			HordeUrlStr = hordeUrlStr?.ToString();
+			StreamId = streamId;
+			JobId = jobId;
+			TemplateId = templateId;
+			StepId = stepId;
+			CommitId = commitId;
+		}
+
+		/// <summary>
+		/// Constructor for standard Horde context.
+		/// </summary>
+		/// <param name="eventName">The event name.</param>
+		/// <param name="schemaVersion">The schema version.</param>
+		/// <param name="sessionId">The id used to correlate telemetry emitted in the same session.</param>
+		/// <param name="sessionLabel">The label used to correlate telemetry emitted across multiple sessions.</param>
+		/// <param name="isBuildMachine">True whether this is a build machine context, false otherwise.</param>
+		/// <param name="hordeUrlStr">The horde url.</param>
+		/// <param name="streamId">The stream id.</param>
+		/// <param name="templateId">The template id.</param>
+		/// <param name="jobId">The Horde job id.</param>
+		/// <param name="stepId">The Horde step id.</param>
+		/// <param name="commitId">The commitId.</param>
+#pragma warning disable CA1054 // URI-like parameters should not be strings
+		public HordeContextTelemetryRecord(string eventName, int schemaVersion, string? sessionId, string? sessionLabel, bool isBuildMachine, string? hordeUrlStr, string? streamId, string? templateId, string? jobId, string? stepId, int? commitId) : base(eventName, schemaVersion, sessionId, sessionLabel)
+#pragma warning restore CA1054 // URI-like parameters should not be strings
+		{
+			IsBuildMachine = isBuildMachine;
+			HordeUrlStr = hordeUrlStr;
+
+			StreamId = TryCreate(streamId, s => new StreamId(s));
+			JobId = TryCreate(jobId, j => Jobs.JobId.Parse(j));
+			TemplateId = TryCreate(templateId, t => new TemplateId(t));
+			StepId = TryCreate(stepId, s => JobStepId.Parse(s));
+
+			CommitId = commitId;
+		}
+
+		/// <summary>
+		/// Constructs a standard Horde context from environment variables.
+		/// </summary>
+		/// <param name="eventName">The event name.</param>
+		/// <param name="schemaVersion">The schema version.</param>
+		/// <param name="sessionId">The id used to correlate telemetry emitted in the same session.</param>
+		/// <param name="sessionLabel">The label used to correlate telemetry emitted across multiple sessions.</param>
+		/// <returns>A instance of standard context created from environment variables.</returns>
+		public HordeContextTelemetryRecord(string eventName, int schemaVersion, string? sessionId, string? sessionLabel) : base(eventName, schemaVersion, sessionId, sessionLabel)
+		{
+			IsBuildMachine = String.Equals("1", Environment.GetEnvironmentVariable("IsBuildMachine")?.Trim(), StringComparison.Ordinal);
+			HordeUrlStr = HordeEnvVars.HordeUrl?.ToString();
+
+			string? streamId = HordeEnvVars.StreamId;
+			string? jobId = HordeEnvVars.JobId;
+			string? templateId = HordeEnvVars.TemplateId;
+			string? stepId = HordeEnvVars.StepId;
+
+			StreamId = TryCreate(streamId, s => new StreamId(s));
+			JobId = TryCreate(jobId,  j =>Jobs.JobId.Parse(j));
+			TemplateId =TryCreate(templateId, t => new TemplateId(t));
+			StepId =TryCreate(stepId, s=> JobStepId.Parse(s));
+
+			if(Int32.TryParse(Environment.GetEnvironmentVariable("uebp_CL"), out int commitIdRaw))
+			{
+				CommitId = commitIdRaw;
+			}
+			else
+			{
+				CommitId = 0;
+			}
+		}
+
+		#endregion -- Public Api --
+
+		#region -- Protectied Api --
+
+		/// <inheritdoc/>
+		protected override void WriteProperties(Hashtable hashtable)
+		{
+			base.WriteProperties(hashtable);
+
+			hashtable[nameof(IsBuildMachine)] = IsBuildMachine;
+
+			if (!String.IsNullOrEmpty(HordeUrlStr))
+			{
+				hashtable[nameof(HordeUrlStr)] = HordeUrlStr;
+			}
+
+			if (!String.IsNullOrEmpty(StreamId?.ToString()))
+			{
+				hashtable[nameof(StreamId)] = StreamId;
+			}
+
+			if (!String.IsNullOrEmpty(JobId?.ToString()))
+			{
+				hashtable[nameof(JobId)] = JobId;
+			}
+
+			if (!String.IsNullOrEmpty(TemplateId?.ToString()))
+			{
+				hashtable[nameof(TemplateId)] = TemplateId;
+			}
+
+			if (!String.IsNullOrEmpty(StepId?.ToString()))
+			{
+				hashtable[nameof(StepId)] = StepId;
+			}
+
+			hashtable[nameof(CommitId)] = CommitId;
+		}
+
+		#endregion
+
+		#region -- Private Api --
+
+		private static T? TryCreate<T>(string? value, Func<string, T?> creationDelegate)
+		{
+			if(String.IsNullOrEmpty(value))
+			{
+				return default!;
+			}
+
+			try 
+			{ 
+				return creationDelegate(value);
+			}
+			catch
+			{
+				return default!;
+			}
+		}
+
+		#endregion -- Private Api --
+	}
+}
\ No newline at end of file
