[Horde][Analytics] Minor improvements to nullability check & safe fallback for ColumnAttribute with no name

#Synced-CL 49874115

diff --git a/Engine/Source/Programs/Shared/EpicGames.Analytics.Telemetry/ColumnValidator.cs b/Engine/Source/Programs/Shared/EpicGames.Analytics.Telemetry/ColumnValidator.cs
--- a/Engine/Source/Programs/Shared/EpicGames.Analytics.Telemetry/ColumnValidator.cs
+++ b/Engine/Source/Programs/Shared/EpicGames.Analytics.Telemetry/ColumnValidator.cs
@@ -39,9 +39,18 @@ namespace EpicGames.Analytics
 				foreach (PropertyInfo prop in props)
 				{
 					ColumnAttribute? columnAttr = prop.GetCustomAttribute<ColumnAttribute>();
-					if (columnAttr != null && columnAttr.Name != null)
+
+					if (columnAttr != null)
 					{
-						columnKeys.Add(columnAttr.Name.ToUpperInvariant());
+						// Safely fall back to the column key being the prop name if no name attribute.
+						if (String.IsNullOrEmpty(columnAttr.Name))
+						{
+							columnKeys.Add(prop.Name.ToUpperInvariant());
+						}
+						else
+						{
+							columnKeys.Add(columnAttr.Name.ToUpperInvariant());
+						}
 						propertyKeys.Add(prop.Name.ToUpperInvariant());
 					}
 				}
diff --git a/Engine/Source/Programs/Shared/EpicGames.Horde/Telemetry/HordeContextTelemetryRecord.cs b/Engine/Source/Programs/Shared/EpicGames.Horde/Telemetry/HordeContextTelemetryRecord.cs
--- a/Engine/Source/Programs/Shared/EpicGames.Horde/Telemetry/HordeContextTelemetryRecord.cs
+++ b/Engine/Source/Programs/Shared/EpicGames.Horde/Telemetry/HordeContextTelemetryRecord.cs
@@ -202,11 +202,11 @@ namespace EpicGames.Horde.Telemetry
 
 		#region -- Private Api --
 
-		private static T? TryCreate<T>(string? value, Func<string, T?> creationDelegate)
+		private static T? TryCreate<T>(string? value, Func<string, T> creationDelegate) where T: struct
 		{
 			if(String.IsNullOrEmpty(value))
 			{
-				return default!;
+				return null;
 			}
 
 			try 
@@ -215,7 +215,7 @@ namespace EpicGames.Horde.Telemetry
 			}
 			catch
 			{
-				return default!;
+				return null;
 			}
 		}
 
