[Horde][Analytics][PerformanceTrends] Introducing key stat performance records.

#Synced-CL 49874115

diff --git a/Engine/Source/Programs/Shared/EpicGames.Analytics.PerformanceTrends/EpicGames.Analytics.PerformanceTrends.csproj b/Engine/Source/Programs/Shared/EpicGames.Analytics.PerformanceTrends/EpicGames.Analytics.PerformanceTrends.csproj
--- a/Engine/Source/Programs/Shared/EpicGames.Analytics.PerformanceTrends/EpicGames.Analytics.PerformanceTrends.csproj
+++ b/Engine/Source/Programs/Shared/EpicGames.Analytics.PerformanceTrends/EpicGames.Analytics.PerformanceTrends.csproj
@@ -13,5 +13,6 @@
 
   <ItemGroup>
     <ProjectReference Include="..\EpicGames.Analytics.Telemetry\EpicGames.Analytics.Telemetry.csproj" />
+    <ProjectReference Include="..\EpicGames.Horde\EpicGames.Horde.csproj" />
   </ItemGroup>
 </Project>
diff --git a/Engine/Source/Programs/Shared/EpicGames.Analytics.PerformanceTrends/KeyStatTelemetryRecord.cs b/Engine/Source/Programs/Shared/EpicGames.Analytics.PerformanceTrends/KeyStatTelemetryRecord.cs
new file mode 100644
--- /dev/null
+++ b/Engine/Source/Programs/Shared/EpicGames.Analytics.PerformanceTrends/KeyStatTelemetryRecord.cs
@@ -0,0 +1,100 @@
+// Copyright Epic Games, Inc. All Rights Reserved.
+
+using System.ComponentModel.DataAnnotations.Schema;
+using EpicGames.Analytics.PerformanceTrends;
+using EpicGames.Analytics.Telemetry;
+
+namespace EpicGames.Analytics
+{
+	/// <summary>
+	/// Telemetry record that describes key stats for performance tests.
+	/// </summary>
+	[AnalyticsTableGen]
+	[Table("personal_schema.jgamble.state_performance_key_stat")]
+	public record KeyStatTelemetryRecord : PerformanceTrendTelemetry
+	{
+		/// <summary>
+		///Dynamic resolution max percentage.
+		/// </summary>
+		[Column("dynamic_resolution_percentage_max")]
+		public int DynamicResolutionPercentageMax { get; init; }
+
+		/// <summary>
+		/// Csv id used to create this record.
+		/// </summary>
+		[Column("csvid")]
+		public string? CsvId { get; init; }
+
+		/// <summary>
+		/// Average GPU time.
+		/// </summary>
+		[Column("gpu_time_avg")]
+		public float GpuTimeAvg { get; init; }
+
+		/// <summary>
+		/// Max physical memory used.
+		/// </summary>
+		[Column("physical_used_mb_max")]
+		public float PhysicalUsedMbMax { get; init; }
+
+		/// <summary>
+		/// Dynamic resolution through Nanite.
+		/// </summary>
+		[Column("dynamic_resolution_vsm_nanite")]
+		public int DynamicResolutionVsmNanite { get; init; }
+
+		/// <summary>
+		/// Render thread time on critical path.
+		/// </summary>
+		[Column("render_thread_time_critical_path")]
+		public float RenderThreadTimeCriticalPath { get; init; }
+
+		/// <summary>
+		/// MVP.
+		/// </summary>
+		[Column("mvp")]
+		public float MVP { get; init; }
+
+		/// <summary>
+		/// Hitches per minute.
+		/// </summary>
+		[Column("hitches_min")]
+		public float HitchesMin { get; init; }
+
+		/// <summary>
+		/// Hitch time percent.
+		/// </summary>
+		[Column("hitch_time_percent")]
+		public float HitchTimePercent { get; init; }
+
+		/// <summary>
+		/// Game thread time average.
+		/// </summary>
+		[Column("game_threadtime_avg")]
+		public float GameThreadtimeAvg { get; init; }
+
+		/// <summary>
+		/// Render thread time average.
+		/// </summary>
+		[Column("render_threadtime_avg")]
+		public float RenderThreadtimeAvg { get; init; }
+
+		/// <summary>
+		/// Game thead time on critical path.
+		/// </summary>
+		[Column("game_thread_time_critical_path")]
+		public float GameThreadTimeCriticalPath { get; init; }
+
+		/// <summary>
+		/// Dynamic reoslution percentage average.
+		/// </summary>
+		[Column("dynamic_resolution_percentage_avg")]
+		public float DynamicResolutionPercentageAvg { get; init; }
+
+		/// <summary>
+		/// No arg constructor for ORM construction.
+		/// </summary>
+		/// <remarks>Used for ORM instantiation.</remarks>
+		public KeyStatTelemetryRecord() : base() { }
+	}
+}
\ No newline at end of file
diff --git a/Engine/Source/Programs/Shared/EpicGames.Analytics.PerformanceTrends/PerformanceTrendTelemetry.cs b/Engine/Source/Programs/Shared/EpicGames.Analytics.PerformanceTrends/PerformanceTrendTelemetry.cs
--- a/Engine/Source/Programs/Shared/EpicGames.Analytics.PerformanceTrends/PerformanceTrendTelemetry.cs
+++ b/Engine/Source/Programs/Shared/EpicGames.Analytics.PerformanceTrends/PerformanceTrendTelemetry.cs
@@ -3,6 +3,7 @@
 using System;
 using System.Collections;
 using System.ComponentModel.DataAnnotations.Schema;
+using EpicGames.Horde.Telemetry;
 
 namespace EpicGames.Analytics.PerformanceTrends
 {
@@ -43,10 +44,28 @@ namespace EpicGames.Analytics.PerformanceTrends
 
 			foreach (DictionaryEntry entry in source)
 			{
-				if (entry.Key is string key)
+				if (entry.Key is string sourceKey)
 				{
-					string sanitizedKey = key.Replace("<", "_lt_", StringComparison.Ordinal).Replace(">", "_gt_", StringComparison.Ordinal);
-					sanitized[sanitizedKey] = entry.Value;
+					string sanitizedKey = sourceKey.Trim().Replace("/", "_", StringComparison.Ordinal)
+							.Replace(" ", "_", StringComparison.Ordinal)
+							.Replace("<", "_lt_", StringComparison.Ordinal)
+							.Replace(">", "_gt_", StringComparison.Ordinal);
+
+					// Verify that after sanitization we don't already have this sourceKey, otherwise suffix the sourceKey.
+					string finalKey = sanitizedKey;
+					bool wasSanitized = sourceKey != finalKey;
+					int suffix = 1;
+
+					// If we have actually sanitized, it means we are at risk of flattening an existing sourceKey.
+					// Check to see whether the final sourceKey already exists within our source; if so it means we created a sanitized sourceKey that overwrites something existing (because we know we've sanitized).
+					// Check to see that we also haven't generated a sanitized sourceKey that overlaps with another sanitized sourceKey.
+					while (wasSanitized && (source.ContainsKey(finalKey) || sanitized.ContainsKey(finalKey)))
+					{
+						suffix++;
+						finalKey = $"{sanitizedKey}_{suffix}";
+					}
+
+					sanitized[finalKey] = entry.Value;
 				}
 				else
 				{
@@ -100,26 +119,26 @@ namespace EpicGames.Analytics.PerformanceTrends
 	/// <summary>
 	/// Abstract record that describes the shared metadata for performance trend telemetry.
 	/// </summary>
-	public abstract record PerformanceTrendTelemetry : AbstractTelemetryRecord
+	public abstract record PerformanceTrendTelemetry : HordeContextTelemetryRecord
 	{
 		#region -- Public Members --
 
 		/// <summary>
 		/// The engine version of the telemetry.
 		/// </summary>
-		[Column("engine_version")]
+		[Column("engineversion")]
 		public string? EngineVersion { get; init; }
 
 		/// <summary>
 		/// The engine release version of the telemetry.
 		/// </summary>
-		[Column("engine_release_version")]
+		[Column("enginereleaseversion")]
 		public string? EngineReleaseVersion { get; init; }
 
 		/// <summary>
 		/// The build version of the telemetry.
 		/// </summary>
-		[Column("build_version")]
+		[Column("buildversion")]
 		public string? BuildVersion { get; init; }
 
 		/// <summary>
@@ -131,42 +150,73 @@ namespace EpicGames.Analytics.PerformanceTrends
 		/// <summary>
 		/// The test id for the telemetry, if applicable.
 		/// </summary>
-		[Column("test_id")]
+		[Column("testid")]
 		public string? TestId { get; init; }
 
 		/// <summary>
 		/// The test name for the telemetry, if applicable.
 		/// </summary>
-		[Column("test_name")]
+		[Column("testname")]
 		public string? TestName { get; init; }
 
 		/// <summary>
 		/// The test config name for the telemetry, if applicable.
 		/// </summary>
-		[Column("test_config_name")]
+		[Column("testconfigname")]
 		public string? TestConfigName { get; init; }
 
 		/// <summary>
 		/// The test type, if applicable.
 		/// </summary>
-		[Column("gauntlet_test_type")]
+		[Column("gauntlettesttype")]
 		public string? GauntletTestType { get; init; }
 
 		/// <summary>
 		/// The sub test type, if applicable.
 		/// </summary>
-		[Column("gauntlet_sub_test")]
+		[Column("gauntletsubtest")]
 		public string? GauntletSubTest { get; init; }
 
+		/// <summary>
+		/// Thether this represents collated data, or not.
+		/// </summary>
+		[Column("collated")]
+		public bool Collated { get; init; }
+
+		/// <summary>
+		/// Summary table name.
+		/// </summary>
+		[Column("summary_name")]
+		public string? SummaryName { get; init; }
+
 		#endregion -- Public Members --
 
+		/// <summary>
+		/// No arg constructor for ORM construction.
+		/// </summary>
+		/// <remarks>Used for ORM instantiation.</remarks>
+		protected PerformanceTrendTelemetry() : base() { }
+
 		/// <summary>
 		/// Constructor.
 		/// </summary>
 		/// <param name="source">The source table to construct the telemetry record from.</param>
 		/// <param name="eventName">The event name.</param>
 		/// <param name="schemaVersion">The schema version of the telemetry.</param>
-		protected PerformanceTrendTelemetry(Hashtable source, string eventName, int schemaVersion) : base(eventName, schemaVersion)
+		/// <param name="sessionId">The id used to correlate telemetry emitted in the same session.</param>
+		/// <param name="sessionLabel">The label used to correlate telemetry emitted across multiple sessions.</param>
+		/// <param name="isBuildMachine">True whether this is a build machine context, false otherwise.</param>
+		/// <param name="hordeUrlStr">The horde url.</param>
+		/// <param name="streamId">The stream id.</param>
+		/// <param name="templateId">The template id.</param>
+		/// <param name="jobId">The Horde job id.</param>
+		/// <param name="stepId">The Horde step id.</param>
+		/// <param name="commitId">The commitId.</param>
+		/// <param name="summaryName">The summary name.</param>
+		/// <param name="collated">Whether this record represents pre-aggregation.</param>
+#pragma warning disable CA1054 // URI-like parameters should not be strings
+		protected PerformanceTrendTelemetry(Hashtable source, string eventName, int schemaVersion, string? sessionId, string? sessionLabel, bool isBuildMachine, string? hordeUrlStr, string? streamId, string? templateId, string? jobId, string? stepId, int? commitId, string summaryName, bool collated) : base(eventName, schemaVersion, sessionId, sessionLabel, isBuildMachine, hordeUrlStr, streamId, templateId, jobId, stepId, commitId)
+#pragma warning restore CA1054 // URI-like parameters should not be strings
 		{
 			if (source == null)
 			{
@@ -182,6 +232,8 @@ namespace EpicGames.Analytics.PerformanceTrends
 			TestConfigName = GetString(source, StandardTelemetryMetadata.TestConfigName);
 			GauntletTestType = GetString(source, StandardTelemetryMetadata.GauntletTestType);
 			GauntletSubTest = GetString(source, StandardTelemetryMetadata.GauntletSubTest);
+			SummaryName = summaryName;
+			Collated = collated;
 		}
 
 		#region -- Private API --
