EpicGames.Tracing.Tests: Use 'Assert.ThrowsExactly' instead of 'Assert.ThrowsException' (MSTEST0039)

#rnx

#Synced-CL 49899152

diff --git a/Engine/Source/Programs/Shared/EpicGames.Tracing.Tests/UnrealInsights/TransportPacketTest.cs b/Engine/Source/Programs/Shared/EpicGames.Tracing.Tests/UnrealInsights/TransportPacketTest.cs
--- a/Engine/Source/Programs/Shared/EpicGames.Tracing.Tests/UnrealInsights/TransportPacketTest.cs
+++ b/Engine/Source/Programs/Shared/EpicGames.Tracing.Tests/UnrealInsights/TransportPacketTest.cs
@@ -78,9 +78,9 @@ namespace EpicGames.Tracing.Tests.UnrealInsights
 			AssertWriteUid(18, true);
 			AssertWriteUid(126, true);
 
-			Assert.ThrowsException<NotImplementedException>(() => AssertWriteUid(127, true));
-			Assert.ThrowsException<NotImplementedException>(() => AssertWriteUid(128, true));
-			Assert.ThrowsException<NotImplementedException>(() => AssertWriteUid(4362, true));
+			Assert.ThrowsExactly<NotImplementedException>(() => AssertWriteUid(127, true));
+			Assert.ThrowsExactly<NotImplementedException>(() => AssertWriteUid(128, true));
+			Assert.ThrowsExactly<NotImplementedException>(() => AssertWriteUid(4362, true));
 		}
 
 		[TestMethod]
