EpicGames.MsBuild: Enable nullable

#rnx

#Synced-CL 49895563

diff --git a/Engine/Source/Programs/Shared/EpicGames.MsBuild/CsProjBuilder.cs b/Engine/Source/Programs/Shared/EpicGames.MsBuild/CsProjBuilder.cs
--- a/Engine/Source/Programs/Shared/EpicGames.MsBuild/CsProjBuilder.cs
+++ b/Engine/Source/Programs/Shared/EpicGames.MsBuild/CsProjBuilder.cs
@@ -25,7 +25,7 @@ namespace EpicGames.MsBuild
 	{
 		static FileReference ConstructBuildRecordPath(CsProjBuildHook hook, FileReference projectPath, IEnumerable<DirectoryReference> baseDirectories)
 		{
-			DirectoryReference basePath = null;
+			DirectoryReference? basePath = null;
 
 			foreach (DirectoryReference scriptFolder in baseDirectories)
 			{
@@ -190,9 +190,13 @@ namespace EpicGames.MsBuild
 			stopwatch.Restart();
 			Parallel.ForEach(projects.Values, (project) =>
 			{
+				if (String.IsNullOrEmpty(project.FullPath))
+				{
+					return;
+				}
 				string targetPath = Path.GetRelativePath(project.DirectoryPath, project.ExpandString(project.GetPropertyValue("TargetPath")));
 
-				FileReference projectPath = FileReference.FromString(project.FullPath);
+				FileReference projectPath = new FileReference(project.FullPath);
 				FileReference buildRecordPath = ConstructBuildRecordPath(hook, projectPath, baseDirectories);
 
 				CsProjBuildRecord buildRecord = new CsProjBuildRecord()
@@ -265,8 +269,8 @@ namespace EpicGames.MsBuild
 				Dictionary<DirectoryReference, HashSet<FileReference>> outputPaths = [];
 				foreach (CsProjBuildRecordEntry entry in buildRecords.Values)
 				{
-					DirectoryReference outputDirectory = FileReference.Combine(entry.ProjectFile.Directory, entry.BuildRecord.TargetPath).Directory;
-					FileReference projectPath = FileReference.Combine(entry.ProjectFile.Directory, entry.BuildRecord.ProjectPath);
+					DirectoryReference outputDirectory = FileReference.Combine(entry.ProjectFile.Directory, entry.BuildRecord.TargetPath!).Directory;
+					FileReference projectPath = FileReference.Combine(entry.ProjectFile.Directory, entry.BuildRecord.ProjectPath!);
 					if (!outputPaths.TryAdd(outputDirectory, [projectPath]))
 					{
 						outputPaths[outputDirectory].Add(projectPath);
@@ -318,7 +322,7 @@ namespace EpicGames.MsBuild
 			Parallel.ForEach(buildRecords.Values, entry =>
 			{
 				// Update the target times
-				FileReference fullPath = FileReference.Combine(entry.ProjectFile.Directory, entry.BuildRecord.TargetPath);
+				FileReference fullPath = FileReference.Combine(entry.ProjectFile.Directory, entry.BuildRecord.TargetPath!);
 				entry.BuildRecord.TargetBuildTime = FileReference.GetLastWriteTime(fullPath);
 			});
 
@@ -327,8 +331,8 @@ namespace EpicGames.MsBuild
 			{
 				foreach (CsProjBuildRecordRef referencedProject in entry.BuildRecord.ProjectReferencesAndTimes)
 				{
-					FileReference refProjectPath = new(Path.GetFullPath(referencedProject.ProjectPath, entry.ProjectFile.Directory.FullName));
-					if (buildRecords.TryGetValue(refProjectPath, out CsProjBuildRecordEntry refEntry))
+					FileReference refProjectPath = new(Path.GetFullPath(referencedProject.ProjectPath!, entry.ProjectFile.Directory.FullName));
+					if (buildRecords.TryGetValue(refProjectPath, out CsProjBuildRecordEntry? refEntry))
 					{
 						referencedProject.TargetBuildTime = refEntry.BuildRecord.TargetBuildTime;
 					}
diff --git a/Engine/Source/Programs/Shared/EpicGames.MsBuild/EpicGames.MsBuild.csproj b/Engine/Source/Programs/Shared/EpicGames.MsBuild/EpicGames.MsBuild.csproj
--- a/Engine/Source/Programs/Shared/EpicGames.MsBuild/EpicGames.MsBuild.csproj
+++ b/Engine/Source/Programs/Shared/EpicGames.MsBuild/EpicGames.MsBuild.csproj
@@ -2,7 +2,6 @@
   <PropertyGroup>
     <TargetFramework>net8.0</TargetFramework>
     <ProduceReferenceAssemblyInOutDir>true</ProduceReferenceAssemblyInOutDir>
-    <Nullable>disable</Nullable>
   </PropertyGroup>
 
   <ItemGroup>
diff --git a/Engine/Source/Programs/Shared/EpicGames.MsBuild/MsBuildSupport.cs b/Engine/Source/Programs/Shared/EpicGames.MsBuild/MsBuildSupport.cs
--- a/Engine/Source/Programs/Shared/EpicGames.MsBuild/MsBuildSupport.cs
+++ b/Engine/Source/Programs/Shared/EpicGames.MsBuild/MsBuildSupport.cs
@@ -36,24 +36,24 @@ public static class MsBuildSupport
 
 			// Find our bundled dotnet SDK
 			List<string> listOfSdks = [];
-			ProcessStartInfo startInfo = new ProcessStartInfo
 			{
-				FileName = hook.DotnetPath.FullName,
-				RedirectStandardOutput = true,
-				UseShellExecute = false,
-				ArgumentList = { "--list-sdks" }
-			};
-			startInfo.EnvironmentVariables["DOTNET_MULTILEVEL_LOOKUP"] = "0"; // use only the bundled dotnet installation - ignore any other/system dotnet install
+				ProcessStartInfo startInfo = new ProcessStartInfo
+				{
+					FileName = hook.DotnetPath.FullName,
+					RedirectStandardOutput = true,
+					UseShellExecute = false,
+					ArgumentList = { "--list-sdks" }
+				};
+				startInfo.EnvironmentVariables["DOTNET_MULTILEVEL_LOOKUP"] = "0"; // use only the bundled dotnet installation - ignore any other/system dotnet install
 
-			Process dotnetProcess = Process.Start(startInfo);
-			{
-				string line;
+				using Process? dotnetProcess = Process.Start(startInfo) ?? throw new Exception("Failed to start dotnet process");
+				string? line;
 				while ((line = dotnetProcess.StandardOutput.ReadLine()) != null)
 				{
 					listOfSdks.Add(line);
 				}
+				dotnetProcess.WaitForExit();
 			}
-			dotnetProcess.WaitForExit();
 
 			if (listOfSdks.Count != 1)
 			{
