[UBT]
* Fixed so Action bCanExecuteInUBACrossArchitecture works again (broke when scheduler was moved to uba)

#Synced-CL 49910128

diff --git a/Engine/Source/Programs/Shared/EpicGames.UBA/IScheduler.cs b/Engine/Source/Programs/Shared/EpicGames.UBA/IScheduler.cs
--- a/Engine/Source/Programs/Shared/EpicGames.UBA/IScheduler.cs
+++ b/Engine/Source/Programs/Shared/EpicGames.UBA/IScheduler.cs
@@ -119,7 +119,7 @@ namespace EpicGames.UBA
 		/// <summary>
 		/// Queue process.
 		/// </summary>
-		uint EnqueueProcess(ProcessStartInfo info, double weight, bool canDetour, bool canExecuteRemotely, int[]? dependencies, byte[]? knownInputs, uint knownInputsCount, uint cacheBucket, uint memoryGroup, ulong predictedMemoryUsage, object userData);
+		uint EnqueueProcess(ProcessStartInfo info, double weight, bool canDetour, bool canExecuteRemotely, bool canCrossArchitecture, bool canCrossPlatform, int[]? dependencies, byte[]? knownInputs, uint knownInputsCount, uint cacheBucket, uint memoryGroup, ulong predictedMemoryUsage, object userData);
 
 		/// <summary>
 		/// Cancel all active processes and skip queued ones
diff --git a/Engine/Source/Programs/Shared/EpicGames.UBA/Impl/SchedulerImpl.cs b/Engine/Source/Programs/Shared/EpicGames.UBA/Impl/SchedulerImpl.cs
--- a/Engine/Source/Programs/Shared/EpicGames.UBA/Impl/SchedulerImpl.cs
+++ b/Engine/Source/Programs/Shared/EpicGames.UBA/Impl/SchedulerImpl.cs
@@ -121,7 +121,7 @@ namespace EpicGames.UBA.Impl
 		static extern byte Scheduler_IsEmpty(nint scheduler);
 
 		[DllImport("UbaHost", CharSet = CharSet.Auto)]
-		static extern uint Scheduler_EnqueueProcess3(nint scheduler, nint info, float weight, byte canDetour, byte canExecuteRemotely, int[]? dependencies, int dependencyCount, byte[]? knownInputs, int knownInputsBytes, uint knownInputsCount, uint cacheBucket, uint memoryGroup, ulong predictedMemoryUsage);
+		static extern uint Scheduler_EnqueueProcess4(nint scheduler, nint info, float weight, byte canDetour, byte canExecuteRemotely, byte canCrossArchitecture, byte canCrossPlatform, int[]? dependencies, int dependencyCount, byte[]? knownInputs, int knownInputsBytes, uint knownInputsCount, uint cacheBucket, uint memoryGroup, ulong predictedMemoryUsage);
 
 		[DllImport("UbaHost", CharSet = CharSet.Auto)]
 		static extern float Scheduler_GetProcessWeightThatCanRunRemotelyNow(nint scheduler);
@@ -156,10 +156,13 @@ namespace EpicGames.UBA.Impl
 			Scheduler_Start(_schedulerHandle);
 		}
 
-		public uint EnqueueProcess(ProcessStartInfo info, double weight, bool canDetour, bool canExecuteRemotely, int[]? dependencies, byte[]? knownInputs, uint knownInputsCount, uint cacheBucket, uint memoryGroup, ulong predictedMemoryUsage, object userData)
+		public uint EnqueueProcess(ProcessStartInfo info, double weight, bool canDetour, bool canExecuteRemotely, bool canCrossArchitecture, bool canCrossPlatform, int[]? dependencies, byte[]? knownInputs, uint knownInputsCount, uint cacheBucket, uint memoryGroup, ulong predictedMemoryUsage, object userData)
 		{
+			byte canExecuteRemotely_ = (byte)(canExecuteRemotely ? 1 : 0);
+			byte canCrossArchitecture_ = (byte)(canCrossArchitecture ? 1 : 0);
+			byte canCrossPlatform_ = (byte)(canCrossPlatform ? 1 : 0);
 			SchedulerProcess process = new(this, info, userData);
-			uint index = Scheduler_EnqueueProcess3(_schedulerHandle, process._startInfoHandle, (float)weight, (byte)(canDetour ? 1 : 0), (byte)(canExecuteRemotely ? 1 : 0), dependencies, dependencies != null ? dependencies.Length : 0, knownInputs, knownInputs != null ? knownInputs.Length : 0, knownInputsCount, cacheBucket, memoryGroup, predictedMemoryUsage);
+			uint index = Scheduler_EnqueueProcess4(_schedulerHandle, process._startInfoHandle, (float)weight, (byte)(canDetour ? 1 : 0), canExecuteRemotely_, canCrossArchitecture_, canCrossPlatform_, dependencies, dependencies != null ? dependencies.Length : 0, knownInputs, knownInputs != null ? knownInputs.Length : 0, knownInputsCount, cacheBucket, memoryGroup, predictedMemoryUsage);
 			process._index = index;
 			_processes.Add(process);
 			return index;
diff --git a/Engine/Source/Programs/UnrealBuildTool/Executors/UnrealBuildAccelerator/UBAExecutor.cs b/Engine/Source/Programs/UnrealBuildTool/Executors/UnrealBuildAccelerator/UBAExecutor.cs
--- a/Engine/Source/Programs/UnrealBuildTool/Executors/UnrealBuildAccelerator/UBAExecutor.cs
+++ b/Engine/Source/Programs/UnrealBuildTool/Executors/UnrealBuildAccelerator/UBAExecutor.cs
@@ -1083,14 +1083,15 @@ namespace UnrealBuildTool
 				int[] prereqActionsSortIndex = [.. action.PrerequisiteActions.Select(a => a.SortIndex)];
 				bool enableDetour = !ForceLocalNoDetour(action) && action.bCanExecuteInUBA;
 				bool canRunRemotely = CanRunRemotely(action);
-
+				bool canCrossArchitecture = action.bCanExecuteInUBACrossArchitecture;
+				bool canCrossPlatform = canCrossArchitecture;
 				(byte[]? knownInputs, uint knownInputsCount) = BuildKnownInputs(action);
 
 				uint bucket = GetCacheBucket(logger, action);
 				uint memoryGroup = GetMemoryGroup(action);
 				ulong predictedMemoryUsage = 0; // uba would do a better job scheduling if we knew approx memory usage
 				ProcessUserData userData = new ProcessUserData() { Action = action };
-				scheduler.EnqueueProcess(startInfo, action.Weight, enableDetour, canRunRemotely, prereqActionsSortIndex, knownInputs, knownInputsCount, bucket, memoryGroup, predictedMemoryUsage, userData);
+				scheduler.EnqueueProcess(startInfo, action.Weight, enableDetour, canRunRemotely, canCrossArchitecture, canCrossPlatform, prereqActionsSortIndex, knownInputs, knownInputsCount, bucket, memoryGroup, predictedMemoryUsage, userData);
 
 				++totalActions;
 				cacheTotalActions += bucket != 0 ? 0 : 1;
