EpicGames.OIDC: Add opt-in access token caching

Currently only refresh tokens are cached to disk. This change adds the ability to also cache access tokens, reducing unnecessary calls when tokens are still valid.
Still writes refresh tokens in the old format as "Providers" key, while also writing a new "Providers2" key which includes a cached access tokens.

#Synced-CL 0

diff --git a/Engine/Source/Programs/OidcToken/OidcToken.Tests/FilesystemTokenStoreTests.cs b/Engine/Source/Programs/OidcToken/OidcToken.Tests/FilesystemTokenStoreTests.cs
new file mode 100644
--- /dev/null
+++ b/Engine/Source/Programs/OidcToken/OidcToken.Tests/FilesystemTokenStoreTests.cs
@@ -0,0 +1,124 @@
+// Copyright Epic Games, Inc. All Rights Reserved.
+
+using System.Security.Cryptography;
+using System.Text;
+using System.Text.Json;
+using EpicGames.OIDC;
+
+namespace OidcToken.Tests;
+
+[TestClass]
+public class FilesystemTokenStoreTests
+{
+	private readonly FileInfo _tempStorePath;
+	private const string LegacyProvider = "legacy-test-provider";
+	private const string LegacyRefreshToken = "legacy-refresh-token-foo";
+	private const string NewProvider = "new-test-provider";
+	private const string NewRefreshToken = "new-refresh-token-bar";
+	private const string NewAccessToken = "new-access-token-baz";
+
+	public FilesystemTokenStoreTests()
+	{
+		string tempDir = Path.Combine(Path.GetTempPath(), "OidcTokenTests", Guid.NewGuid().ToString());
+		Directory.CreateDirectory(tempDir);
+		_tempStorePath = new FileInfo(Path.Combine(tempDir, "testTokenStore.dat"));
+	}
+
+	[TestCleanup]
+	public void Cleanup()
+	{
+		if (_tempStorePath.Directory?.Exists == true)
+		{
+			Directory.Delete(_tempStorePath.Directory.FullName, recursive: true);
+		}
+	}
+
+	[TestMethod]
+	public void BothTokensRoundTrip()
+	{
+		DateTimeOffset expiry = DateTimeOffset.UtcNow.AddHours(2);
+
+		using (FilesystemTokenStore store = new(_tempStorePath))
+		{
+			store.AddRefreshToken(NewProvider, NewRefreshToken);
+			store.AddAccessToken(NewProvider, NewAccessToken, expiry);
+			store.Save();
+		}
+
+		using (FilesystemTokenStore store = new(_tempStorePath))
+		{
+			bool foundRefresh = store.TryGetRefreshToken(NewProvider, out string retrievedRefresh);
+			Assert.IsTrue(foundRefresh);
+			Assert.AreEqual(NewRefreshToken, retrievedRefresh);
+
+			bool foundAccess = store.TryGetAccessToken(NewProvider, out string retrievedAccess, out DateTimeOffset retrievedExpiry);
+			Assert.IsTrue(foundAccess);
+			Assert.AreEqual(NewAccessToken, retrievedAccess);
+			Assert.AreEqual(expiry.ToUnixTimeSeconds(), retrievedExpiry.ToUnixTimeSeconds());
+		}
+	}
+
+	[TestMethod]
+	public void LegacyProvidersBackwardsCompatibility()
+	{
+		// Create a legacy V1 format entry
+		WriteLegacyRefreshToken(LegacyProvider, LegacyRefreshToken);
+
+		// Verify we can read the legacy refresh token
+		using (FilesystemTokenStore store = new(_tempStorePath))
+		{
+			bool found = store.TryGetRefreshToken(LegacyProvider, out string retrievedToken);
+			Assert.IsTrue(found);
+			Assert.AreEqual(LegacyRefreshToken, retrievedToken);
+		}
+
+		// Add a new provider and save - this should preserve legacy Providers
+		using (FilesystemTokenStore store = new(_tempStorePath))
+		{
+			store.AddRefreshToken(NewProvider, NewRefreshToken);
+			store.Save();
+		}
+
+		// Verify file structure: legacy in Providers, new in ProvidersV2
+		using (FilesystemTokenStore store = new(_tempStorePath))
+		{
+			TokenStoreState? state = store.ReadRawState();
+			Assert.IsNotNull(state);
+			Assert.IsTrue(state.Providers.ContainsKey(LegacyProvider));
+			Assert.IsTrue(state.ProvidersV2.ContainsKey(NewProvider));
+			Assert.IsFalse(state.Providers.ContainsKey(NewProvider));
+		}
+
+		// Verify both tokens are still readable
+		using (FilesystemTokenStore store = new(_tempStorePath))
+		{
+			bool foundLegacy = store.TryGetRefreshToken(LegacyProvider, out string retrievedLegacy);
+			Assert.IsTrue(foundLegacy);
+			Assert.AreEqual(LegacyRefreshToken, retrievedLegacy);
+
+			bool foundNew = store.TryGetRefreshToken(NewProvider, out string retrievedNew);
+			Assert.IsTrue(foundNew);
+			Assert.AreEqual(NewRefreshToken, retrievedNew);
+		}
+	}
+
+	/// <summary>
+	/// Writes a legacy V1 format token to the store file for testing backwards compatibility.
+	/// </summary>
+	private void WriteLegacyRefreshToken(string providerName, string refreshToken)
+	{
+		using Aes aes = Aes.Create();
+		TokenStoreState state = new(Convert.ToBase64String(aes.IV));
+
+		// Encrypt using legacy format (same as V2 but stored directly in Providers)
+		byte[] bytes = Encoding.Unicode.GetBytes(refreshToken);
+		ICryptoTransform encryptor = aes.CreateEncryptor(FilesystemTokenStore.GetSeed(), aes.IV);
+		byte[] encrypted = encryptor.TransformFinalBlock(bytes, 0, bytes.Length);
+		state.Providers[providerName] = Convert.ToBase64String(encrypted);
+
+		_tempStorePath.Directory?.Create();
+		using FileStream fs = new(_tempStorePath.FullName, FileMode.Create, FileAccess.Write);
+		using Utf8JsonWriter writer = new(fs);
+		JsonSerializer.Serialize(writer, state, TokenStoreStateContext.Default.TokenStoreState);
+	}
+}
diff --git a/Engine/Source/Programs/OidcToken/OidcToken.Tests/OidcToken.Tests.csproj b/Engine/Source/Programs/OidcToken/OidcToken.Tests/OidcToken.Tests.csproj
--- a/Engine/Source/Programs/OidcToken/OidcToken.Tests/OidcToken.Tests.csproj
+++ b/Engine/Source/Programs/OidcToken/OidcToken.Tests/OidcToken.Tests.csproj
@@ -20,6 +20,7 @@
 
   <ItemGroup>
     <ProjectReference Include="..\OidcToken.csproj" />
+    <ProjectReference Include="..\..\Shared\EpicGames.OIDC\EpicGames.OIDC.csproj" />
   </ItemGroup>
 
   <ItemGroup>
diff --git a/Engine/Source/Programs/OidcToken/OidcToken.Tests/WindowsTokenStoreTests.cs b/Engine/Source/Programs/OidcToken/OidcToken.Tests/WindowsTokenStoreTests.cs
new file mode 100644
--- /dev/null
+++ b/Engine/Source/Programs/OidcToken/OidcToken.Tests/WindowsTokenStoreTests.cs
@@ -0,0 +1,133 @@
+// Copyright Epic Games, Inc. All Rights Reserved.
+
+using System.Runtime.InteropServices;
+using System.Text;
+using System.Text.Json;
+using EpicGames.OIDC;
+
+namespace OidcToken.Tests;
+
+[TestClass]
+public class WindowsTokenStoreTests
+{
+	private readonly FileInfo _tempStorePath;
+	private const string LegacyProvider = "legacy-test-provider";
+	private const string LegacyRefreshToken = "legacy-refresh-token-foo";
+	private const string NewProvider = "new-test-provider";
+	private const string NewRefreshToken = "new-refresh-token-bar";
+	private const string NewAccessToken = "new-access-token-baz";
+
+	public WindowsTokenStoreTests()
+	{
+		string tempDir = Path.Combine(Path.GetTempPath(), "OidcTokenTests", Guid.NewGuid().ToString());
+		Directory.CreateDirectory(tempDir);
+		_tempStorePath = new FileInfo(Path.Combine(tempDir, "testTokenStore.dat"));
+	}
+
+	[TestInitialize]
+	public void Init()
+	{
+		if (!RuntimeInformation.IsOSPlatform(OSPlatform.Windows))
+		{
+			Assert.Inconclusive("WindowsTokenStoreTests is only supported on Windows");
+		}
+	}
+
+	[TestCleanup]
+	public void Cleanup()
+	{
+		if (_tempStorePath.Directory?.Exists == true)
+		{
+			Directory.Delete(_tempStorePath.Directory.FullName, recursive: true);
+		}
+	}
+
+	[TestMethod]
+	public void BothTokensRoundTrip()
+	{
+		DateTimeOffset expiry = DateTimeOffset.UtcNow.AddHours(2);
+
+		using (WindowsTokenStore store = new(null, _tempStorePath))
+		{
+			store.AddRefreshToken(NewProvider, NewRefreshToken);
+			store.AddAccessToken(NewProvider, NewAccessToken, expiry);
+			store.Save();
+		}
+
+		using (WindowsTokenStore store = new(null, _tempStorePath))
+		{
+			bool foundRefresh = store.TryGetRefreshToken(NewProvider, out string retrievedRefresh);
+			Assert.IsTrue(foundRefresh);
+			Assert.AreEqual(NewRefreshToken, retrievedRefresh);
+
+			bool foundAccess = store.TryGetAccessToken(NewProvider, out string retrievedAccess, out DateTimeOffset retrievedExpiry);
+			Assert.IsTrue(foundAccess);
+			Assert.AreEqual(NewAccessToken, retrievedAccess);
+			Assert.AreEqual(expiry.ToUnixTimeSeconds(), retrievedExpiry.ToUnixTimeSeconds());
+		}
+	}
+
+	[TestMethod]
+	public void LegacyProvidersBackwardsCompatibility()
+	{
+		// Create a legacy V1 format entry using internal helper
+		using (WindowsTokenStore _ = new(null, _tempStorePath))
+		{
+			WriteLegacyRefreshToken(LegacyProvider, LegacyRefreshToken);
+		}
+
+		// Verify we can read the legacy refresh token
+		using (WindowsTokenStore store = new(null, _tempStorePath))
+		{
+			bool found = store.TryGetRefreshToken(LegacyProvider, out string retrievedToken);
+			Assert.IsTrue(found);
+			Assert.AreEqual(LegacyRefreshToken, retrievedToken);
+		}
+
+		// Add a new provider and save - this should preserve legacy Providers
+		using (WindowsTokenStore store = new(null, _tempStorePath))
+		{
+			store.AddRefreshToken(NewProvider, NewRefreshToken);
+			store.Save();
+		}
+
+		// Verify file structure: legacy in Providers, new in ProvidersV2
+		using (WindowsTokenStore store = new(null, _tempStorePath))
+		{
+			WindowsTokenStoreState? state = store.ReadRawState();
+			Assert.IsNotNull(state);
+			Assert.IsTrue(state.Providers.ContainsKey(LegacyProvider));
+			Assert.IsTrue(state.ProvidersV2.ContainsKey(NewProvider));
+			Assert.IsFalse(state.Providers.ContainsKey(NewProvider));
+		}
+
+		// Verify both tokens are still readable
+		using (WindowsTokenStore store = new(null, _tempStorePath))
+		{
+			bool foundLegacy = store.TryGetRefreshToken(LegacyProvider, out string retrievedLegacy);
+			Assert.IsTrue(foundLegacy);
+			Assert.AreEqual(LegacyRefreshToken, retrievedLegacy);
+
+			bool foundNew = store.TryGetRefreshToken(NewProvider, out string retrievedNew);
+			Assert.IsTrue(foundNew);
+			Assert.AreEqual(NewRefreshToken, retrievedNew);
+		}
+	}
+	
+	/// <summary>
+	/// Writes a legacy V1 format token to the store file for testing backwards compatibility.
+	/// </summary>
+	private void WriteLegacyRefreshToken(string providerName, string refreshToken)
+	{
+		WindowsTokenStoreState state = new ();
+
+		// Encrypt using legacy format
+		byte[] entropy = WindowsTokenStore.GetEntropy(providerName);
+		byte[] encrypted = CryptProtectDataHelper.DoCryptProtectData(Encoding.Unicode.GetBytes(refreshToken), $"OidcToken-{providerName}", entropy);
+		state.Providers[providerName] = Convert.ToBase64String(encrypted);
+		
+		using FileStream fs = new (_tempStorePath.FullName, FileMode.Create, FileAccess.Write);
+		using Utf8JsonWriter writer = new (fs);
+		JsonSerializer.Serialize(writer, state, WindowsTokenStoreStateContext.Default.WindowsTokenStoreState);
+	}
+}
diff --git a/Engine/Source/Programs/Shared/EpicGames.OIDC/AssemblyInfo.cs b/Engine/Source/Programs/Shared/EpicGames.OIDC/AssemblyInfo.cs
new file mode 100644
--- /dev/null
+++ b/Engine/Source/Programs/Shared/EpicGames.OIDC/AssemblyInfo.cs
@@ -0,0 +1,5 @@
+// Copyright Epic Games, Inc. All Rights Reserved.
+
+using System.Runtime.CompilerServices;
+
+[assembly: InternalsVisibleTo("OidcToken.Tests")]
diff --git a/Engine/Source/Programs/Shared/EpicGames.OIDC/FakeOidcTokenManager.cs b/Engine/Source/Programs/Shared/EpicGames.OIDC/FakeOidcTokenManager.cs
--- a/Engine/Source/Programs/Shared/EpicGames.OIDC/FakeOidcTokenManager.cs
+++ b/Engine/Source/Programs/Shared/EpicGames.OIDC/FakeOidcTokenManager.cs
@@ -62,9 +62,23 @@ public class FakeOidcTokenManager : IOidcTokenManager
 		throw new NotImplementedException("Method not in use");
 	}
 
+	/// <inheritdoc/>
+	public Task<OidcTokenInfo> GetAccessToken(string providerIdentifier, GetAccessTokenOptions options, CancellationToken cancellationToken = default)
+	{
+		// Options ignored in fake - just delegate to existing
+		throw new NotImplementedException("Method not in use");
+	}
+
 	/// <inheritdoc/>
 	public Task<OidcTokenInfo?> TryGetAccessToken(string providerIdentifier, CancellationToken cancellationToken = default)
 	{
+		return TryGetAccessToken(providerIdentifier, new GetAccessTokenOptions(), cancellationToken);
+	}
+
+	/// <inheritdoc/>
+	public Task<OidcTokenInfo?> TryGetAccessToken(string providerIdentifier, GetAccessTokenOptions options, CancellationToken cancellationToken = default)
+	{
+		// Options ignored in fake - caching behavior not simulated
 		if (String.IsNullOrEmpty(RefreshToken))
 		{
 			throw new NotLoggedInException();
diff --git a/Engine/Source/Programs/Shared/EpicGames.OIDC/ITokenStore.cs b/Engine/Source/Programs/Shared/EpicGames.OIDC/ITokenStore.cs
--- a/Engine/Source/Programs/Shared/EpicGames.OIDC/ITokenStore.cs
+++ b/Engine/Source/Programs/Shared/EpicGames.OIDC/ITokenStore.cs
@@ -3,17 +3,41 @@
 using System;
 using System.Diagnostics.CodeAnalysis;
 using System.Runtime.InteropServices;
+using System.Text.Json.Serialization;
 using Microsoft.Extensions.DependencyInjection;
 
 #pragma warning disable CS1591 // Missing XML documentation on public types
 
 namespace EpicGames.OIDC
 {
+	/// <summary>
+	/// JSON serialization context for ProviderTokenData
+	/// </summary>
+	[JsonSerializable(typeof(ProviderTokenData))]
+	public partial class ProviderTokenDataContext : JsonSerializerContext
+	{
+	}
+
+	/// <summary>
+	/// Token data stored per provider (shared across all token store implementations)
+	/// </summary>
+	public class ProviderTokenData
+	{
+		public string RefreshToken { get; set; } = string.Empty;
+		public string? AccessToken { get; set; }
+		public long? AccessTokenExpireTime { get; set; }
+	}
+
 	public interface ITokenStore : IDisposable
 	{
 		public bool TryGetRefreshToken(string oidcProvider, [MaybeNullWhen(false)] out string refreshToken);
 
 		public void AddRefreshToken(string providerIdentifier, string refreshToken);
+
+		public bool TryGetAccessToken(string oidcProvider, [MaybeNullWhen(false)] out string accessToken, out DateTimeOffset expiry);
+
+		public void AddAccessToken(string providerIdentifier, string accessToken, DateTimeOffset expiry);
+
 		public void Save();
 	}
 
diff --git a/Engine/Source/Programs/Shared/EpicGames.OIDC/InMemoryTokenStore.cs b/Engine/Source/Programs/Shared/EpicGames.OIDC/InMemoryTokenStore.cs
--- a/Engine/Source/Programs/Shared/EpicGames.OIDC/InMemoryTokenStore.cs
+++ b/Engine/Source/Programs/Shared/EpicGames.OIDC/InMemoryTokenStore.cs
@@ -12,6 +12,7 @@ namespace EpicGames.OIDC;
 public class InMemoryTokenStore : ITokenStore
 {
 	private readonly Dictionary<string, string> _providerToRefreshToken = new();
+	private readonly Dictionary<string, (string Token, DateTimeOffset Expiry)> _providerToAccessToken = new();
 
 	/// <inheritdoc/>
 	public virtual bool TryGetRefreshToken(string oidcProvider, out string refreshToken)
@@ -32,6 +33,27 @@ public class InMemoryTokenStore : ITokenStore
 		_providerToRefreshToken[providerIdentifier] = refreshToken;
 	}
 
+	/// <inheritdoc/>
+	public virtual bool TryGetAccessToken(string oidcProvider, out string accessToken, out DateTimeOffset expiry)
+	{
+		if (!_providerToAccessToken.TryGetValue(oidcProvider, out (string Token, DateTimeOffset Expiry) cached))
+		{
+			accessToken = "";
+			expiry = DateTimeOffset.MinValue;
+			return false;
+		}
+
+		accessToken = cached.Token;
+		expiry = cached.Expiry;
+		return true;
+	}
+
+	/// <inheritdoc/>
+	public virtual void AddAccessToken(string providerIdentifier, string accessToken, DateTimeOffset expiry)
+	{
+		_providerToAccessToken[providerIdentifier] = (accessToken, expiry);
+	}
+
 	/// <inheritdoc/>
 	public virtual void Save()
 	{
diff --git a/Engine/Source/Programs/Shared/EpicGames.OIDC/MacOSTokenStore.cs b/Engine/Source/Programs/Shared/EpicGames.OIDC/MacOSTokenStore.cs
--- a/Engine/Source/Programs/Shared/EpicGames.OIDC/MacOSTokenStore.cs
+++ b/Engine/Source/Programs/Shared/EpicGames.OIDC/MacOSTokenStore.cs
@@ -4,6 +4,7 @@ using System;
 using System.Collections.Generic;
 using System.Diagnostics;
 using System.Text;
+using System.Text.Json;
 
 #pragma warning disable CS1591 // Missing XML documentation on public types
 
@@ -11,31 +12,111 @@ namespace EpicGames.OIDC
 {
 	public sealed class MacOSTokenStore : ITokenStore
 	{
+		// Cache provider data in memory
+		private readonly Dictionary<string, ProviderTokenData> _providerData = new Dictionary<string, ProviderTokenData>();
+
 		public MacOSTokenStore()
 		{
 		}
 
-		public bool TryGetRefreshToken(string oidcProvider, out string refreshToken)
+		private ProviderTokenData? TryGetProviderData(string oidcProvider)
 		{
-			string? token = KeychainHelper.GetKeychainItem(oidcProvider);
+			// Check in-memory cache first
+			if (_providerData.TryGetValue(oidcProvider, out ProviderTokenData? cached))
+			{
+				return cached;
+			}
+
+			// Try to load from Keychain (V2 format first, then legacy)
+			string? json = KeychainHelper.GetKeychainItem($"{oidcProvider}-v2");
+			if (json != null)
+			{
+				try
+				{
+					ProviderTokenData? data = JsonSerializer.Deserialize(json, ProviderTokenDataContext.Default.ProviderTokenData);
+					if (data != null)
+					{
+						_providerData[oidcProvider] = data;
+						return data;
+					}
+				}
+				catch (JsonException)
+				{
+					// Fall through to legacy
+				}
+			}
+
+			// Try legacy format (just refresh token)
+			string? legacyToken = KeychainHelper.GetKeychainItem(oidcProvider);
+			if (legacyToken != null)
+			{
+				ProviderTokenData data = new ProviderTokenData { RefreshToken = legacyToken };
+				_providerData[oidcProvider] = data;
+				return data;
+			}
+
+			return null;
+		}
+
+		private void SetProviderData(string providerIdentifier, ProviderTokenData tokenData)
+		{
+			_providerData[providerIdentifier] = tokenData;
+
+			// Write V2 format only (JSON with all token data)
+			// Note: We don't update the legacy keychain entry - older code will continue using existing entries
+			string json = JsonSerializer.Serialize(tokenData, ProviderTokenDataContext.Default.ProviderTokenData);
+			KeychainHelper.AddKeychainItem(json, $"{providerIdentifier}-v2");
+		}
 
-			if (token == null)
+		public bool TryGetRefreshToken(string oidcProvider, out string refreshToken)
+		{
+			ProviderTokenData? data = TryGetProviderData(oidcProvider);
+			if (data == null || string.IsNullOrEmpty(data.RefreshToken))
 			{
 				refreshToken = "";
 				return false;
 			}
 
-			refreshToken = token;
+			refreshToken = data.RefreshToken;
 			return true;
 		}
 
 		public void AddRefreshToken(string providerIdentifier, string refreshToken)
 		{
-			KeychainHelper.AddKeychainItem(refreshToken, providerIdentifier);
+			ProviderTokenData? existing = TryGetProviderData(providerIdentifier);
+			ProviderTokenData tokenData = existing ?? new ProviderTokenData();
+			tokenData.RefreshToken = refreshToken;
+			SetProviderData(providerIdentifier, tokenData);
+		}
+
+		public bool TryGetAccessToken(string oidcProvider, out string accessToken, out DateTimeOffset expiry)
+		{
+			accessToken = "";
+			expiry = DateTimeOffset.MinValue;
+
+			ProviderTokenData? data = TryGetProviderData(oidcProvider);
+			if (data == null || string.IsNullOrEmpty(data.AccessToken) || data.AccessTokenExpireTime == null)
+			{
+				return false;
+			}
+
+			accessToken = data.AccessToken;
+			expiry = DateTimeOffset.FromUnixTimeSeconds(data.AccessTokenExpireTime.Value);
+			return true;
+		}
+
+		public void AddAccessToken(string providerIdentifier, string accessToken, DateTimeOffset expiry)
+		{
+			ProviderTokenData? existing = TryGetProviderData(providerIdentifier);
+			ProviderTokenData tokenData = existing ?? new ProviderTokenData();
+			tokenData.AccessToken = accessToken;
+			tokenData.AccessTokenExpireTime = expiry.ToUnixTimeSeconds();
+			SetProviderData(providerIdentifier, tokenData);
 		}
 
 		public void Save()
 		{
+			// Keychain writes are immediate, nothing to do here
 		}
 
 		public void Dispose()
diff --git a/Engine/Source/Programs/Shared/EpicGames.OIDC/OidcTokenManager.cs b/Engine/Source/Programs/Shared/EpicGames.OIDC/OidcTokenManager.cs
--- a/Engine/Source/Programs/Shared/EpicGames.OIDC/OidcTokenManager.cs
+++ b/Engine/Source/Programs/Shared/EpicGames.OIDC/OidcTokenManager.cs
@@ -33,6 +33,10 @@ namespace EpicGames.OIDC
 
 		public Task<OidcTokenInfo?> TryGetAccessToken(string providerIdentifier, CancellationToken cancellationToken = default);
 
+		public Task<OidcTokenInfo> GetAccessToken(string providerIdentifier, GetAccessTokenOptions options, CancellationToken cancellationToken = default);
+
+		public Task<OidcTokenInfo?> TryGetAccessToken(string providerIdentifier, GetAccessTokenOptions options, CancellationToken cancellationToken = default);
+
 		public OidcStatus GetStatusForProvider(string providerIdentifier);
 	}
 
@@ -156,6 +160,18 @@ namespace EpicGames.OIDC
 			return _tokenClients[providerIdentifier].TryGetAccessTokenAsync(cancellationToken);
 		}
 
+		/// <inheritdoc/>
+		public Task<OidcTokenInfo> GetAccessToken(string providerIdentifier, GetAccessTokenOptions options, CancellationToken cancellationToken = default)
+		{
+			return _tokenClients[providerIdentifier].GetAccessTokenAsync(options, cancellationToken);
+		}
+
+		/// <inheritdoc/>
+		public Task<OidcTokenInfo?> TryGetAccessToken(string providerIdentifier, GetAccessTokenOptions options, CancellationToken cancellationToken = default)
+		{
+			return _tokenClients[providerIdentifier].TryGetAccessTokenAsync(options, cancellationToken);
+		}
+
 		/// <inheritdoc/>
 		public OidcStatus GetStatusForProvider(string providerIdentifier)
 		{
@@ -221,6 +237,19 @@ namespace EpicGames.OIDC
 		}
 	}
 
+	/// <summary>
+	/// Options for retrieving access tokens
+	/// </summary>
+	public class GetAccessTokenOptions
+	{
+		/// <summary>
+		/// Whether to use disk cache for access tokens.
+		/// When enabled, access tokens are persisted to disk and reused across process restarts if still valid.
+		/// Default is false for backwards compatibility.
+		/// </summary>
+		public bool UseCache { get; set; } = false;
+	}
+
 	/// <summary>
 	/// The data read from a http request into the oidc client, mostly a wrapper around the http.listener context
 	/// </summary>
@@ -409,6 +438,7 @@ namespace EpicGames.OIDC
 		private readonly Uri _authorityUri;
 		private readonly string _clientId;
 		private readonly string _scopes;
+		private readonly TimeSpan _tokenValidityBuffer = TimeSpan.FromMinutes(10);
 
 		private string? _refreshToken;
 		private string? _accessToken;
@@ -836,13 +866,43 @@ namespace EpicGames.OIDC
 			return await TryGetAccessTokenAsync(cancellationToken) ?? throw new NotLoggedInException();
 		}
 
-		public async Task<OidcTokenInfo?> TryGetAccessTokenAsync(CancellationToken cancellationToken)
+		public async Task<OidcTokenInfo> GetAccessTokenAsync(GetAccessTokenOptions options, CancellationToken cancellationToken)
+		{
+			return await TryGetAccessTokenAsync(options, cancellationToken) ?? throw new NotLoggedInException();
+		}
+
+		public Task<OidcTokenInfo?> TryGetAccessTokenAsync(CancellationToken cancellationToken)
+		{
+			return TryGetAccessTokenAsync(new GetAccessTokenOptions(), cancellationToken);
+		}
+
+		public async Task<OidcTokenInfo?> TryGetAccessTokenAsync(GetAccessTokenOptions options, CancellationToken cancellationToken)
 		{
 			if (String.IsNullOrEmpty(_refreshToken))
 			{
 				throw new NotLoggedInException();
 			}
 
+			// If caching is enabled, try to load from disk cache first
+			if (options.UseCache && _tokenStore.TryGetAccessToken(_name, out string? cachedAccessToken, out DateTimeOffset cachedExpiry))
+			{
+				OidcTokenInfo cachedInfo = new()
+				{
+					RefreshToken = _refreshToken,
+					AccessToken = cachedAccessToken,
+					TokenExpiry = cachedExpiry
+				};
+
+				// Ensure token is valid for at least another few minutes
+				if (cachedInfo.IsValid(DateTime.Now - _tokenValidityBuffer))
+				{
+					_accessToken = cachedAccessToken;
+					_tokenExpiry = cachedExpiry;
+					return cachedInfo;
+				}
+			}
+
+			// Check in-memory cache
 			OidcTokenInfo tokenInfo = new()
 			{
 				RefreshToken = _refreshToken,
@@ -850,13 +910,22 @@ namespace EpicGames.OIDC
 				TokenExpiry = _tokenExpiry
 			};
 
-			// Ensure token is valid for at least another two minutes
-			if (tokenInfo.IsValid(DateTime.Now.AddMinutes(-2)))
+			// Ensure token is valid for at least another few minutes
+			if (tokenInfo.IsValid(DateTime.Now - _tokenValidityBuffer))
 			{
 				return tokenInfo;
 			}
 
-			return await TryDoRefreshTokenAsync(_refreshToken, cancellationToken);
+			OidcTokenInfo? refreshedInfo = await TryDoRefreshTokenAsync(_refreshToken, cancellationToken);
+
+			// If caching is enabled and refresh succeeded, save to disk cache
+			if (options.UseCache && refreshedInfo != null && !String.IsNullOrEmpty(refreshedInfo.AccessToken))
+			{
+				_tokenStore.AddAccessToken(_name, refreshedInfo.AccessToken, refreshedInfo.TokenExpiry);
+				_tokenStore.Save();
+			}
+
+			return refreshedInfo;
 		}
 
 		public OidcStatus GetStatus()
diff --git a/Engine/Source/Programs/Shared/EpicGames.OIDC/TokenStore.cs b/Engine/Source/Programs/Shared/EpicGames.OIDC/TokenStore.cs
--- a/Engine/Source/Programs/Shared/EpicGames.OIDC/TokenStore.cs
+++ b/Engine/Source/Programs/Shared/EpicGames.OIDC/TokenStore.cs
@@ -19,24 +19,24 @@ namespace EpicGames.OIDC
 
 	internal class TokenStoreState
 	{
+		/// <summary>
+		/// Legacy format - encrypted refresh tokens only (kept for backwards compatibility when reading)
+		/// </summary>
 		public Dictionary<string, string> Providers { get; set; } = new Dictionary<string, string>();
-		public string Key { get; set; }
 
-		[JsonConstructor]
-		public TokenStoreState(string key, Dictionary<string, string> providers)
-		{
-			Key = key;
-			Providers = providers;
-		}
+		/// <summary>
+		/// New format - encrypted JSON containing both refresh and access tokens
+		/// </summary>
+		public Dictionary<string, string> ProvidersV2 { get; set; } = new Dictionary<string, string>();
 
-		public TokenStoreState(byte[] key, Dictionary<string, byte[]> providers)
-		{
-			Key = Convert.ToBase64String(key);
+		public string Key { get; set; } = string.Empty;
 
-			foreach (KeyValuePair<string, byte[]> pair in providers)
-			{
-				Providers[pair.Key] = Convert.ToBase64String(pair.Value);
-			}
+		[JsonConstructor]
+		public TokenStoreState(string? key = null, Dictionary<string, string>? providers = null, Dictionary<string, string>? providersV2 = null)
+		{
+			Key = key ?? string.Empty;
+			Providers = providers ?? new Dictionary<string, string>();
+			ProvidersV2 = providersV2 ?? new Dictionary<string, string>();
 		}
 	}
 
@@ -45,57 +45,67 @@ namespace EpicGames.OIDC
 	/// </summary>
 	public class FilesystemTokenStore : ITokenStore, IDisposable
 	{
-		private readonly Dictionary<string, byte[]> _providerToRefreshToken;
+		private readonly Dictionary<string, byte[]> _providerData = new Dictionary<string, byte[]>();
 		private readonly List<string> _dirtyProviders = new List<string>();
 		private readonly SymmetricAlgorithm _crypt = Aes.Create();
+		private readonly FileInfo? _storePath;
 
-		public FilesystemTokenStore()
+		public FilesystemTokenStore() : this(null)
 		{
-			_providerToRefreshToken = ReadStoreFromDisk();
 		}
 
-		private static FileInfo GetStorePath()
+		internal FilesystemTokenStore(FileInfo? storePath)
+		{
+			_storePath = storePath;
+
+			// Load V2 provider data from disk
+			TokenStoreState? state = ReadStateFromDisk();
+			if (state != null)
+			{
+				if (!string.IsNullOrEmpty(state.Key))
+				{
+					_crypt.IV = Convert.FromBase64String(state.Key);
+				}
+
+				foreach ((string key, string value) in state.ProvidersV2)
+				{
+					_providerData[key] = Convert.FromBase64String(value);
+				}
+			}
+		}
+
+		private static FileInfo GetDefaultStorePath()
 		{
 			return new FileInfo(Path.Combine(Environment.GetFolderPath(Environment.SpecialFolder.ApplicationData), "Epic", "UnrealEngine", "Common", "OidcToken", "tokenStore.dat"));
 		}
 
-		private Dictionary<string, byte[]> ReadStoreFromDisk()
+		private FileInfo GetStorePath() => _storePath ?? GetDefaultStorePath();
+
+		private TokenStoreState? ReadStateFromDisk()
 		{
 			FileInfo fi = GetStorePath();
+			fi.Refresh(); // Refresh cached metadata in case file was created after FileInfo was instantiated
 			if (!fi.Exists)
 			{
-				// if we have no store on disk then we just initialize it to empty
-				return new Dictionary<string, byte[]>();
+				return null;
 			}
 
 			using FileStream fs = fi.OpenRead();
-			TokenStoreState? state;
 			try
 			{
-				state = JsonSerializer.Deserialize(fs, TokenStoreStateContext.Default.TokenStoreState);
+				return JsonSerializer.Deserialize(fs, TokenStoreStateContext.Default.TokenStoreState);
 			}
 			catch (JsonException)
 			{
-				state = null;
+				return null;
 			}
-
-			if (state == null)
-			{
-				// if we fail to deserialize the state just drop it, will mean users will need to login again
-				return new Dictionary<string, byte[]>();
-			}
-			Dictionary<string, byte[]> providers = new Dictionary<string, byte[]>();
-
-			foreach ((string key, string value) in state.Providers)
-			{
-				providers[key] = Convert.FromBase64String(value);
-			}
-
-			_crypt.IV = Convert.FromBase64String(state.Key);
-
-			return providers;
 		}
 
+		/// <summary>
+		/// Reads the raw state from disk for test verification.
+		/// </summary>
+		internal TokenStoreState? ReadRawState() => ReadStateFromDisk();
+
 		private void SaveStoreToDisk()
 		{
 			FileInfo fi = GetStorePath();
@@ -107,25 +117,38 @@ namespace EpicGames.OIDC
 
 			lock (_dirtyProviders)
 			{
-				// no providers have changed, do not touch the state file
 				if (_dirtyProviders.Count == 0)
 				{
 					return;
 				}
 
-				// read back the state of all providers but only overwrite the state of the ones we have actually got new state for (are dirty)
-				Dictionary<string, byte[]> providers = ReadStoreFromDisk();
+				// Read back the existing state from disk
+				TokenStoreState? existingState = ReadStateFromDisk();
 
+				// Build new state, preserving legacy Providers (we don't modify them)
+				TokenStoreState newState = new TokenStoreState(Convert.ToBase64String(_crypt.IV));
+				if (existingState != null)
+				{
+					// Preserve legacy Providers as-is
+					newState.Providers = existingState.Providers;
+					// Start with existing ProvidersV2
+					foreach ((string key, string value) in existingState.ProvidersV2)
+					{
+						newState.ProvidersV2[key] = value;
+					}
+				}
+
+				// Merge in our dirty changes to ProvidersV2
 				foreach (string providerId in _dirtyProviders)
 				{
-					providers[providerId] = _providerToRefreshToken[providerId];
+					newState.ProvidersV2[providerId] = Convert.ToBase64String(_providerData[providerId]);
 				}
 
 				string tempFile = Path.GetTempFileName();
 				{
 					using FileStream fs = new FileStream(tempFile, FileMode.Create, FileAccess.Write);
 					using Utf8JsonWriter writer = new Utf8JsonWriter(fs);
-					JsonSerializer.Serialize<TokenStoreState>(writer, new TokenStoreState(_crypt.IV, providers), TokenStoreStateContext.Default.TokenStoreState);
+					JsonSerializer.Serialize<TokenStoreState>(writer, newState, TokenStoreStateContext.Default.TokenStoreState);
 				}
 
 				File.Move(tempFile, fi.FullName, true);
@@ -133,34 +156,118 @@ namespace EpicGames.OIDC
 			}
 		}
 
+		private string? TryDecryptData(byte[] encryptedData)
+		{
+			try
+			{
+				ICryptoTransform decryptor = _crypt.CreateDecryptor(GetSeed(), _crypt.IV);
+				byte[] decrypted = decryptor.TransformFinalBlock(encryptedData, 0, encryptedData.Length);
+				return Encoding.Unicode.GetString(decrypted);
+			}
+			catch (CryptographicException)
+			{
+				return null;
+			}
+		}
+
+		private ProviderTokenData? TryGetProviderData(string oidcProvider)
+		{
+			// Try V2 format first (in-memory)
+			if (_providerData.TryGetValue(oidcProvider, out byte[]? encryptedData))
+			{
+				string? json = TryDecryptData(encryptedData);
+				if (json != null)
+				{
+					try
+					{
+						return JsonSerializer.Deserialize(json, ProviderTokenDataContext.Default.ProviderTokenData);
+					}
+					catch (JsonException)
+					{
+						// Fall through to legacy
+					}
+				}
+			}
+
+			// Fall back to legacy format from disk (refresh token only)
+			TokenStoreState? state = ReadStateFromDisk();
+			if (state != null && state.Providers.TryGetValue(oidcProvider, out string? legacyBase64))
+			{
+				byte[] legacyEncrypted = Convert.FromBase64String(legacyBase64);
+				string? refreshToken = TryDecryptData(legacyEncrypted);
+				if (refreshToken != null)
+				{
+					return new ProviderTokenData { RefreshToken = refreshToken };
+				}
+			}
+
+			return null;
+		}
+
+		private void SetProviderData(string providerIdentifier, ProviderTokenData tokenData)
+		{
+			string json = JsonSerializer.Serialize(tokenData, ProviderTokenDataContext.Default.ProviderTokenData);
+			byte[] bytes = Encoding.Unicode.GetBytes(json);
+#pragma warning disable CA5401
+			ICryptoTransform encryptor = _crypt.CreateEncryptor(GetSeed(), _crypt.IV);
+#pragma warning restore CA5401
+			byte[] encrypted = encryptor.TransformFinalBlock(bytes, 0, bytes.Length);
+
+			_providerData[providerIdentifier] = encrypted;
+
+			lock (_dirtyProviders)
+			{
+				if (!_dirtyProviders.Contains(providerIdentifier))
+				{
+					_dirtyProviders.Add(providerIdentifier);
+				}
+			}
+		}
+
 		public bool TryGetRefreshToken(string oidcProvider, out string refreshToken)
 		{
-			if (!_providerToRefreshToken.TryGetValue(oidcProvider, out byte[]? encryptedToken))
+			ProviderTokenData? data = TryGetProviderData(oidcProvider);
+			if (data == null || string.IsNullOrEmpty(data.RefreshToken))
 			{
 				refreshToken = "";
 				return false;
 			}
 
-			ICryptoTransform decryptor = _crypt.CreateDecryptor(GetSeed(), _crypt.IV);
-			byte[] bytes = decryptor.TransformFinalBlock(encryptedToken, 0, encryptedToken.Length);
-			refreshToken = Encoding.Unicode.GetString(bytes);
+			refreshToken = data.RefreshToken;
 			return true;
 		}
 
 		public void AddRefreshToken(string providerIdentifier, string refreshToken)
 		{
-			byte[] bytes = Encoding.Unicode.GetBytes(refreshToken);
-#pragma warning disable CA5401 // Do not use CreateEncryptor with non-default IV
-			ICryptoTransform encryptor = _crypt.CreateEncryptor(GetSeed(), _crypt.IV);
-#pragma warning restore CA5401 // Do not use CreateEncryptor with non-default IV
+			ProviderTokenData? existing = TryGetProviderData(providerIdentifier);
+			ProviderTokenData tokenData = existing ?? new ProviderTokenData();
+			tokenData.RefreshToken = refreshToken;
+			SetProviderData(providerIdentifier, tokenData);
+		}
 
-			byte[] encryptedToken = encryptor.TransformFinalBlock(bytes, 0, bytes.Length);
-			_providerToRefreshToken[providerIdentifier] = encryptedToken;
+		public bool TryGetAccessToken(string oidcProvider, out string accessToken, out DateTimeOffset expiry)
+		{
+			accessToken = "";
+			expiry = DateTimeOffset.MinValue;
 
-			lock (_dirtyProviders)
+			ProviderTokenData? data = TryGetProviderData(oidcProvider);
+			if (data == null || string.IsNullOrEmpty(data.AccessToken) || data.AccessTokenExpireTime == null)
 			{
-				_dirtyProviders.Add(providerIdentifier);
+				return false;
 			}
+
+			accessToken = data.AccessToken;
+			expiry = DateTimeOffset.FromUnixTimeSeconds(data.AccessTokenExpireTime.Value);
+			return true;
+		}
+
+		public void AddAccessToken(string providerIdentifier, string accessToken, DateTimeOffset expiry)
+		{
+			ProviderTokenData? existing = TryGetProviderData(providerIdentifier);
+			ProviderTokenData tokenData = existing ?? new ProviderTokenData();
+			tokenData.AccessToken = accessToken;
+			tokenData.AccessTokenExpireTime = expiry.ToUnixTimeSeconds();
+			SetProviderData(providerIdentifier, tokenData);
 		}
 
 		public void Save()
@@ -168,7 +275,7 @@ namespace EpicGames.OIDC
 			SaveStoreToDisk();
 		}
 
-		static byte[] GetSeed()
+		internal static byte[] GetSeed()
 		{
 			return new byte[]
 			{
@@ -192,4 +299,4 @@ namespace EpicGames.OIDC
 			GC.SuppressFinalize(this);
 		}
 	}
-}
\ No newline at end of file
+}
diff --git a/Engine/Source/Programs/Shared/EpicGames.OIDC/WindowsTokenStore.cs b/Engine/Source/Programs/Shared/EpicGames.OIDC/WindowsTokenStore.cs
--- a/Engine/Source/Programs/Shared/EpicGames.OIDC/WindowsTokenStore.cs
+++ b/Engine/Source/Programs/Shared/EpicGames.OIDC/WindowsTokenStore.cs
@@ -23,86 +23,86 @@ namespace EpicGames.OIDC
 
 	internal class WindowsTokenStoreState
 	{
+		/// <summary>
+		/// Legacy format - encrypted refresh tokens only (kept for backwards compatibility when reading)
+		/// </summary>
 		public Dictionary<string, string> Providers { get; set; } = new Dictionary<string, string>();
 
+		/// <summary>
+		/// New format - encrypted JSON containing both refresh and access tokens
+		/// </summary>
+		public Dictionary<string, string> ProvidersV2 { get; set; } = new Dictionary<string, string>();
+
 		[JsonConstructor]
-		public WindowsTokenStoreState(Dictionary<string, string> providers)
+		public WindowsTokenStoreState(Dictionary<string, string>? providers = null, Dictionary<string, string>? providersV2 = null)
 		{
-			Providers = providers;
-		}
-
-		public WindowsTokenStoreState(Dictionary<string, byte[]> providers)
-		{
-			foreach (KeyValuePair<string, byte[]> pair in providers)
-			{
-				Providers[pair.Key] = Convert.ToBase64String(pair.Value);
-			}
+			Providers = providers ?? new Dictionary<string, string>();
+			ProvidersV2 = providersV2 ?? new Dictionary<string, string>();
 		}
 	}
 
 	public class WindowsTokenStore : ITokenStore, IDisposable
 	{
 		private readonly ILogger<WindowsTokenStore>? _logger = null;
+		private readonly FileInfo? _storePath = null;
 
-		private readonly Dictionary<string, byte[]> _providerToRefreshToken = new Dictionary<string, byte[]>();
+		private readonly Dictionary<string, byte[]> _providerData = new Dictionary<string, byte[]>();
 		private readonly List<string> _dirtyProviders = new List<string>();
 
-		public WindowsTokenStore()
+		public WindowsTokenStore() : this(null, null)
 		{
-			_providerToRefreshToken = ReadStoreFromDisk();
 		}
 
 		[ActivatorUtilitiesConstructor]
-		public WindowsTokenStore(ILogger<WindowsTokenStore> logger)
+		public WindowsTokenStore(ILogger<WindowsTokenStore>? logger) : this(logger, null)
+		{
+		}
+
+		internal WindowsTokenStore(ILogger<WindowsTokenStore>? logger, FileInfo? storePath)
 		{
 			_logger = logger;
+			_storePath = storePath;
 
-			_providerToRefreshToken = ReadStoreFromDisk();
+			// Load V2 provider data from disk
+			WindowsTokenStoreState? state = ReadStateFromDisk();
+			if (state != null)
+			{
+				foreach ((string key, string value) in state.ProvidersV2)
+				{
+					_providerData[key] = Convert.FromBase64String(value);
+				}
+			}
 		}
 
-		private static FileInfo GetStorePath()
+		private static FileInfo GetDefaultStorePath()
 		{
 			return new FileInfo(Path.Combine(Environment.GetFolderPath(Environment.SpecialFolder.LocalApplicationData), "UnrealEngine", "Common", "OidcToken", "oidcTokenStore.dat"));
 		}
 
-		private Dictionary<string, byte[]> ReadStoreFromDisk()
+		private FileInfo GetStorePath() => _storePath ?? GetDefaultStorePath();
+
+		private WindowsTokenStoreState? ReadStateFromDisk()
 		{
 			FileInfo fi = GetStorePath();
+			fi.Refresh(); // Refresh cached metadata in case file was created after FileInfo was instantiated
 			if (!fi.Exists)
 			{
 				_logger?.LogDebug("No existing token store found at {Path}. Assuming empty store.", fi.FullName);
-				// if we have no store on disk then we just initialize it to empty
-				return new Dictionary<string, byte[]>();
+				return null;
 			}
 
 			using FileStream fs = fi.Open(FileMode.Open, FileAccess.Read, FileShare.Read);
 			using TextReader tr = new StreamReader(fs);
 
-			WindowsTokenStoreState? state;
 			try
 			{
-				state = JsonSerializer.Deserialize(tr.ReadToEnd(), WindowsTokenStoreStateContext.Default.WindowsTokenStoreState);
+				return JsonSerializer.Deserialize(tr.ReadToEnd(), WindowsTokenStoreStateContext.Default.WindowsTokenStoreState);
 			}
-			catch (JsonException)
+			catch (JsonException je)
 			{
-				state = null;
+				_logger?.LogDebug("Failed to deserialize state. Reason: {Message}", je.Message);
+				return null;
 			}
-
-			if (state == null)
-			{
-				_logger?.LogDebug("Failed to deserialize state. Dropping the existing state.");
-				// if we fail to deserialize the state just drop it, will mean users will need to login again
-				return new Dictionary<string, byte[]>();
-			}
-
-			Dictionary<string, byte[]> providers = new Dictionary<string, byte[]>();
-
-			foreach ((string key, string value) in state.Providers)
-			{
-				providers[key] = Convert.FromBase64String(value);
-			}
-
-			return providers;
 		}
 
 		private void SaveStoreToDisk()
@@ -116,7 +116,7 @@ namespace EpicGames.OIDC
 
 			lock (_dirtyProviders)
 			{
-				// no providers have changed, do not touch the state file
+				// No providers have changed, do not touch the state file
 				if (_dirtyProviders.Count == 0)
 				{
 					return;
@@ -130,21 +130,35 @@ namespace EpicGames.OIDC
 				}
 				catch (AbandonedMutexException)
 				{
-
 				}
 
-				// read back the state of all providers but only overwrite the state of the ones we have actually got new state for (are dirty)
-				Dictionary<string, byte[]> providers = ReadStoreFromDisk();
+				// Read back the existing state from disk
+				WindowsTokenStoreState? existingState = ReadStateFromDisk();
 
+				// Build new state, preserving legacy Providers (we don't modify them)
+				WindowsTokenStoreState newState = new ();
+				if (existingState != null)
+				{
+					// Preserve legacy Providers as-is
+					newState.Providers = existingState.Providers;
+					// Start with existing ProvidersV2
+					foreach ((string key, string value) in existingState.ProvidersV2)
+					{
+						newState.ProvidersV2[key] = value;
+					}
+				}
+
+				// Merge in our dirty changes to ProvidersV2
 				foreach (string providerId in _dirtyProviders)
 				{
-					providers[providerId] = _providerToRefreshToken[providerId];
+					newState.ProvidersV2[providerId] = Convert.ToBase64String(_providerData[providerId]);
 				}
+
 				string tempFile = Path.GetTempFileName();
 				{
 					using FileStream fs = new FileStream(tempFile, FileMode.Create, FileAccess.Write);
 					using Utf8JsonWriter writer = new Utf8JsonWriter(fs);
-					JsonSerializer.Serialize(writer, new WindowsTokenStoreState(providers), WindowsTokenStoreStateContext.Default.WindowsTokenStoreState);
+					JsonSerializer.Serialize(writer, newState, WindowsTokenStoreStateContext.Default.WindowsTokenStoreState);
 				}
 
 				File.Move(tempFile, fi.FullName, true);
@@ -155,59 +169,123 @@ namespace EpicGames.OIDC
 			}
 		}
 
-		public bool TryGetRefreshToken(string oidcProvider, out string refreshToken)
+		private string? TryDecryptData(byte[] encryptedData, string description, string oidcProvider)
 		{
-			if (!_providerToRefreshToken.TryGetValue(oidcProvider, out byte[]? encryptedToken))
+			try
 			{
-				refreshToken = "";
-				return false;
+				byte[] decrypted = CryptProtectDataHelper.DoCryptUnprotectData(encryptedData, description, GetEntropy(oidcProvider));
+				return Encoding.Unicode.GetString(decrypted);
 			}
-
-			try
+			catch (Win32Exception e)
 			{
-				byte[] bytes = CryptProtectDataHelper.DoCryptUnprotectData(encryptedToken, $"OidcToken-{oidcProvider}", GetEntropy(oidcProvider));
-				refreshToken = Encoding.Unicode.GetString(bytes);
+				if (e.NativeErrorCode != 13 && e.NativeErrorCode != unchecked((int)0x8009000B))
+				{
+					throw;
+				}
+				_logger?.LogDebug("Unable to decrypt token data for provider {Provider}. Ignoring.", oidcProvider);
+				return null;
+			}
+		}
 
-				return true;
+		private ProviderTokenData? TryGetProviderData(string oidcProvider)
+		{
+			// Try V2 format first (in-memory)
+			if (_providerData.TryGetValue(oidcProvider, out byte[]? encryptedData))
+			{
+				string? json = TryDecryptData(encryptedData, $"OidcTokenV2-{oidcProvider}", oidcProvider);
+				if (json != null)
+				{
+					try
+					{
+						return JsonSerializer.Deserialize(json, ProviderTokenDataContext.Default.ProviderTokenData);
+					}
+					catch (JsonException)
+					{
+						_logger?.LogDebug("Unable to deserialize V2 token data for provider {Provider}. Ignoring.", oidcProvider);
+					}
+				}
 			}
-			catch (Win32Exception e)
+
+			// Fall back to legacy format from disk (refresh token only)
+			WindowsTokenStoreState? state = ReadStateFromDisk();
+			if (state != null && state.Providers.TryGetValue(oidcProvider, out string? legacyBase64))
 			{
-				if (e.NativeErrorCode == 13) // data is invalid
+				byte[] legacyEncrypted = Convert.FromBase64String(legacyBase64);
+				string? refreshToken = TryDecryptData(legacyEncrypted, $"OidcToken-{oidcProvider}", oidcProvider);
+				if (refreshToken != null)
 				{
-					// unable to decrypt the data, ignore it
-					refreshToken = "";
-					_logger?.LogDebug("Unable to decrypt refresh token. Ignoring.");
-					return false;
+					return new ProviderTokenData { RefreshToken = refreshToken };
 				}
-				if (e.NativeErrorCode == unchecked((int)0x8009000B)) // key not valid for use in specified state
+			}
+
+			return null;
+		}
+
+		private void SetProviderData(string providerIdentifier, ProviderTokenData tokenData)
+		{
+			string json = JsonSerializer.Serialize(tokenData, ProviderTokenDataContext.Default.ProviderTokenData);
+			byte[] encrypted = CryptProtectDataHelper.DoCryptProtectData(Encoding.Unicode.GetBytes(json), $"OidcTokenV2-{providerIdentifier}", GetEntropy(providerIdentifier));
+
+			_providerData[providerIdentifier] = encrypted;
+
+			lock (_dirtyProviders)
+			{
+				if (!_dirtyProviders.Contains(providerIdentifier))
 				{
-					// unable to decrypt the data, ignore it
-					refreshToken = "";
-					_logger?.LogDebug("Unable to decrypt refresh token, key not valid for use in specified state. Ignoring.");
-					return false;
+					_dirtyProviders.Add(providerIdentifier);
 				}
-				throw;
 			}
 		}
 
-		private static byte[] GetEntropy(string oidcProvider)
+		public bool TryGetRefreshToken(string oidcProvider, out string refreshToken)
 		{
-			byte[] providerBytes = Encoding.UTF8.GetBytes(oidcProvider);
+			ProviderTokenData? data = TryGetProviderData(oidcProvider);
+			if (data == null || string.IsNullOrEmpty(data.RefreshToken))
+			{
+				refreshToken = "";
+				return false;
+			}
 
-			return providerBytes;
+			refreshToken = data.RefreshToken;
+			return true;
+		}
+
+		internal static byte[] GetEntropy(string oidcProvider)
+		{
+			return Encoding.UTF8.GetBytes(oidcProvider);
 		}
 
 		public void AddRefreshToken(string providerIdentifier, string refreshToken)
 		{
-			byte[] bytes = Encoding.Unicode.GetBytes(refreshToken);
-			byte[] encryptedToken = CryptProtectDataHelper.DoCryptProtectData(bytes, $"OidcToken-{providerIdentifier}", GetEntropy(providerIdentifier));
+			ProviderTokenData? existing = TryGetProviderData(providerIdentifier);
+			ProviderTokenData tokenData = existing ?? new ProviderTokenData();
+			tokenData.RefreshToken = refreshToken;
+			SetProviderData(providerIdentifier, tokenData);
+		}
 
-			_providerToRefreshToken[providerIdentifier] = encryptedToken;
+		public bool TryGetAccessToken(string oidcProvider, out string accessToken, out DateTimeOffset expiry)
+		{
+			accessToken = "";
+			expiry = DateTimeOffset.MinValue;
 
-			lock (_dirtyProviders)
+			ProviderTokenData? data = TryGetProviderData(oidcProvider);
+			if (data == null || string.IsNullOrEmpty(data.AccessToken) || data.AccessTokenExpireTime == null)
 			{
-				_dirtyProviders.Add(providerIdentifier);
+				return false;
 			}
+
+			accessToken = data.AccessToken;
+			expiry = DateTimeOffset.FromUnixTimeSeconds(data.AccessTokenExpireTime.Value);
+			return true;
+		}
+
+		public void AddAccessToken(string providerIdentifier, string accessToken, DateTimeOffset expiry)
+		{
+			ProviderTokenData? existing = TryGetProviderData(providerIdentifier);
+			ProviderTokenData tokenData = existing ?? new ProviderTokenData();
+			tokenData.AccessToken = accessToken;
+			tokenData.AccessTokenExpireTime = expiry.ToUnixTimeSeconds();
+			SetProviderData(providerIdentifier, tokenData);
 		}
 
 		public void Save()
@@ -215,6 +293,11 @@ namespace EpicGames.OIDC
 			SaveStoreToDisk();
 		}
 
+		/// <summary>
+		/// Reads the raw state from disk for test verification.
+		/// </summary>
+		internal WindowsTokenStoreState? ReadRawState() => ReadStateFromDisk();
+
 		protected virtual void Dispose(bool disposing)
 		{
 			if (disposing)
@@ -401,4 +484,4 @@ namespace EpicGames.OIDC
 		}
 	}
 #pragma warning restore IDE1006 // Naming Styles
-}
\ No newline at end of file
+}
