Horde: Enable cache when requesting access token via EpicGames.OIDC

Uses new caching functionality from 49949594

diff --git a/Engine/Source/Programs/Shared/EpicGames.Horde/HordeHttpAuthHandler.cs b/Engine/Source/Programs/Shared/EpicGames.Horde/HordeHttpAuthHandler.cs
--- a/Engine/Source/Programs/Shared/EpicGames.Horde/HordeHttpAuthHandler.cs
+++ b/Engine/Source/Programs/Shared/EpicGames.Horde/HordeHttpAuthHandler.cs
@@ -414,7 +414,7 @@ namespace EpicGames.Horde
 			{
 				try
 				{
-					result = await oidcTokenManager.TryGetAccessToken(oidcProvider, cancellationToken);
+					result = await oidcTokenManager.TryGetAccessToken(oidcProvider, new GetAccessTokenOptions() { UseCache = true }, cancellationToken);
 				}
 				catch (Exception ex)
 				{
