Shared: NU1510 no longer needs to be suppressed
System.Drawing.Common is no longer needed as no other Shared modules or UBT require it
Remove ProduceReferenceAssemblyInOutDir as it is unneeded

#rnx

#Synced-CL 49956463

diff --git a/Engine/Source/Programs/Shared/EpicGames.Analytics.Telemetry/EpicGames.Analytics.Telemetry.csproj b/Engine/Source/Programs/Shared/EpicGames.Analytics.Telemetry/EpicGames.Analytics.Telemetry.csproj
--- a/Engine/Source/Programs/Shared/EpicGames.Analytics.Telemetry/EpicGames.Analytics.Telemetry.csproj
+++ b/Engine/Source/Programs/Shared/EpicGames.Analytics.Telemetry/EpicGames.Analytics.Telemetry.csproj
@@ -2,7 +2,6 @@
   <PropertyGroup>
     <TargetFramework>net8.0</TargetFramework>
     <RootNamespace>EpicGames.Analytics</RootNamespace>
-    <ProduceReferenceAssemblyInOutDir>true</ProduceReferenceAssemblyInOutDir>
     <AllowUnsafeBlocks>true</AllowUnsafeBlocks>
   </PropertyGroup>
 
diff --git a/Engine/Source/Programs/Shared/EpicGames.Build/EpicGames.Build.csproj b/Engine/Source/Programs/Shared/EpicGames.Build/EpicGames.Build.csproj
--- a/Engine/Source/Programs/Shared/EpicGames.Build/EpicGames.Build.csproj
+++ b/Engine/Source/Programs/Shared/EpicGames.Build/EpicGames.Build.csproj
@@ -2,7 +2,6 @@
   <PropertyGroup>
     <TargetFramework>net8.0</TargetFramework>
     <RootNamespace>UnrealBuildBase</RootNamespace>
-    <ProduceReferenceAssemblyInOutDir>true</ProduceReferenceAssemblyInOutDir>
     <AllowUnsafeBlocks>true</AllowUnsafeBlocks>
   </PropertyGroup>
 
diff --git a/Engine/Source/Programs/Shared/EpicGames.MsBuild/EpicGames.MsBuild.csproj b/Engine/Source/Programs/Shared/EpicGames.MsBuild/EpicGames.MsBuild.csproj
--- a/Engine/Source/Programs/Shared/EpicGames.MsBuild/EpicGames.MsBuild.csproj
+++ b/Engine/Source/Programs/Shared/EpicGames.MsBuild/EpicGames.MsBuild.csproj
@@ -1,13 +1,11 @@
 ﻿<Project Sdk="Microsoft.NET.Sdk">
   <PropertyGroup>
     <TargetFramework>net8.0</TargetFramework>
-    <ProduceReferenceAssemblyInOutDir>true</ProduceReferenceAssemblyInOutDir>
   </PropertyGroup>
 
   <ItemGroup>
     <PackageReference Include="Microsoft.Build" Version="17.11.48" ExcludeAssets="runtime" />
     <PackageReference Include="Microsoft.Build.Locator" Version="1.9.1" />
-    <PackageReference Include="System.Drawing.Common" Version="8.0.20" />
   </ItemGroup>
 
   <ItemGroup>
diff --git a/Engine/Source/Programs/Shared/EpicGames.UBA/EpicGames.UBA.csproj b/Engine/Source/Programs/Shared/EpicGames.UBA/EpicGames.UBA.csproj
--- a/Engine/Source/Programs/Shared/EpicGames.UBA/EpicGames.UBA.csproj
+++ b/Engine/Source/Programs/Shared/EpicGames.UBA/EpicGames.UBA.csproj
@@ -3,7 +3,6 @@
     <TargetFramework>net8.0</TargetFramework>
     <AllowUnsafeBlocks>true</AllowUnsafeBlocks>
     <GenerateDocumentationFile>True</GenerateDocumentationFile>
-    <ProduceReferenceAssemblyInOutDir>true</ProduceReferenceAssemblyInOutDir>
     <PackageReadmeFile>README.md</PackageReadmeFile>
   </PropertyGroup>
 
diff --git a/Engine/Source/Programs/Shared/EpicGames.UHT/EpicGames.UHT.csproj b/Engine/Source/Programs/Shared/EpicGames.UHT/EpicGames.UHT.csproj
--- a/Engine/Source/Programs/Shared/EpicGames.UHT/EpicGames.UHT.csproj
+++ b/Engine/Source/Programs/Shared/EpicGames.UHT/EpicGames.UHT.csproj
@@ -3,7 +3,6 @@
     <TargetFramework>net8.0</TargetFramework>
     <AllowUnsafeBlocks>true</AllowUnsafeBlocks>
     <GenerateDocumentationFile>True</GenerateDocumentationFile>
-    <ProduceReferenceAssemblyInOutDir>true</ProduceReferenceAssemblyInOutDir>
     <EnforceCodeStyleInBuild>False</EnforceCodeStyleInBuild>
   </PropertyGroup>
 
