UHT: Match against whole path segments when finding the shortest include path.

#Synced-CL 49988106

diff --git a/Engine/Source/Programs/Shared/EpicGames.UHT/Utils/UhtSession.cs b/Engine/Source/Programs/Shared/EpicGames.UHT/Utils/UhtSession.cs
--- a/Engine/Source/Programs/Shared/EpicGames.UHT/Utils/UhtSession.cs
+++ b/Engine/Source/Programs/Shared/EpicGames.UHT/Utils/UhtSession.cs
@@ -330,13 +330,15 @@ namespace EpicGames.UHT.Utils
 		public string GetShortestIncludePath(UHTManifest.Module moduleObj, string filePath)
 		{
 			string? relativePath = null;
+			FileReference fileRef = new FileReference(filePath);
 			foreach (string includePath in moduleObj.IncludePaths)
 			{
-				if (filePath.StartsWith(includePath, StringComparison.Ordinal))
+				DirectoryReference includePathRef = new DirectoryReference(includePath);
+				if (fileRef.IsUnderDirectory(new(includePath)))
 				{
 					if (relativePath == null || (filePath.Length - includePath.Length - 1) < relativePath.Length)
 					{
-						relativePath = filePath.Substring(includePath.Length + 1);
+						relativePath = fileRef.MakeRelativeTo(includePathRef);
 					}
 				}
 			}
