[Backout] - CL49949594
#rnx

Original CL Desc
-----------------------------------------------------------------
EpicGames.OIDC: Add opt-in access token caching

Currently only refresh tokens are cached to disk. This change adds the ability to also cache access tokens, reducing unnecessary calls when tokens are still valid.
Still writes refresh tokens in the old format as "Providers" key, while also writing a new "Providers2" key which includes a cached access tokens.

#Synced-CL 0

diff --git a/Engine/Source/Programs/OidcToken/OidcToken.Tests/FilesystemTokenStoreTests.cs b/Engine/Source/Programs/OidcToken/OidcToken.Tests/FilesystemTokenStoreTests.cs
deleted file mode 100644
diff --git a/Engine/Source/Programs/OidcToken/OidcToken.Tests/OidcToken.Tests.csproj b/Engine/Source/Programs/OidcToken/OidcToken.Tests/OidcToken.Tests.csproj
--- a/Engine/Source/Programs/OidcToken/OidcToken.Tests/OidcToken.Tests.csproj
+++ b/Engine/Source/Programs/OidcToken/OidcToken.Tests/OidcToken.Tests.csproj
@@ -20,7 +20,6 @@
 
   <ItemGroup>
     <ProjectReference Include="..\OidcToken.csproj" />
-    <ProjectReference Include="..\..\Shared\EpicGames.OIDC\EpicGames.OIDC.csproj" />
   </ItemGroup>
 
   <ItemGroup>
diff --git a/Engine/Source/Programs/OidcToken/OidcToken.Tests/WindowsTokenStoreTests.cs b/Engine/Source/Programs/OidcToken/OidcToken.Tests/WindowsTokenStoreTests.cs
deleted file mode 100644
diff --git a/Engine/Source/Programs/Shared/EpicGames.OIDC/AssemblyInfo.cs b/Engine/Source/Programs/Shared/EpicGames.OIDC/AssemblyInfo.cs
deleted file mode 100644
diff --git a/Engine/Source/Programs/Shared/EpicGames.OIDC/FakeOidcTokenManager.cs b/Engine/Source/Programs/Shared/EpicGames.OIDC/FakeOidcTokenManager.cs
--- a/Engine/Source/Programs/Shared/EpicGames.OIDC/FakeOidcTokenManager.cs
+++ b/Engine/Source/Programs/Shared/EpicGames.OIDC/FakeOidcTokenManager.cs
@@ -62,23 +62,9 @@ public class FakeOidcTokenManager : IOidcTokenManager
 		throw new NotImplementedException("Method not in use");
 	}
 
-	/// <inheritdoc/>
-	public Task<OidcTokenInfo> GetAccessToken(string providerIdentifier, GetAccessTokenOptions options, CancellationToken cancellationToken = default)
-	{
-		// Options ignored in fake - just delegate to existing
-		throw new NotImplementedException("Method not in use");
-	}
-
 	/// <inheritdoc/>
 	public Task<OidcTokenInfo?> TryGetAccessToken(string providerIdentifier, CancellationToken cancellationToken = default)
 	{
-		return TryGetAccessToken(providerIdentifier, new GetAccessTokenOptions(), cancellationToken);
-	}
-
-	/// <inheritdoc/>
-	public Task<OidcTokenInfo?> TryGetAccessToken(string providerIdentifier, GetAccessTokenOptions options, CancellationToken cancellationToken = default)
-	{
-		// Options ignored in fake - caching behavior not simulated
 		if (String.IsNullOrEmpty(RefreshToken))
 		{
 			throw new NotLoggedInException();
diff --git a/Engine/Source/Programs/Shared/EpicGames.OIDC/ITokenStore.cs b/Engine/Source/Programs/Shared/EpicGames.OIDC/ITokenStore.cs
--- a/Engine/Source/Programs/Shared/EpicGames.OIDC/ITokenStore.cs
+++ b/Engine/Source/Programs/Shared/EpicGames.OIDC/ITokenStore.cs
@@ -3,41 +3,17 @@
 using System;
 using System.Diagnostics.CodeAnalysis;
 using System.Runtime.InteropServices;
-using System.Text.Json.Serialization;
 using Microsoft.Extensions.DependencyInjection;
 
 #pragma warning disable CS1591 // Missing XML documentation on public types
 
 namespace EpicGames.OIDC
 {
-	/// <summary>
-	/// JSON serialization context for ProviderTokenData
-	/// </summary>
-	[JsonSerializable(typeof(ProviderTokenData))]
-	public partial class ProviderTokenDataContext : JsonSerializerContext
-	{
-	}
-
-	/// <summary>
-	/// Token data stored per provider (shared across all token store implementations)
-	/// </summary>
-	public class ProviderTokenData
-	{
-		public string RefreshToken { get; set; } = string.Empty;
-		public string? AccessToken { get; set; }
-		public long? AccessTokenExpireTime { get; set; }
-	}
-
 	public interface ITokenStore : IDisposable
 	{
 		public bool TryGetRefreshToken(string oidcProvider, [MaybeNullWhen(false)] out string refreshToken);
 
 		public void AddRefreshToken(string providerIdentifier, string refreshToken);
-
-		public bool TryGetAccessToken(string oidcProvider, [MaybeNullWhen(false)] out string accessToken, out DateTimeOffset expiry);
-
-		public void AddAccessToken(string providerIdentifier, string accessToken, DateTimeOffset expiry);
-
 		public void Save();
 	}
 
diff --git a/Engine/Source/Programs/Shared/EpicGames.OIDC/InMemoryTokenStore.cs b/Engine/Source/Programs/Shared/EpicGames.OIDC/InMemoryTokenStore.cs
--- a/Engine/Source/Programs/Shared/EpicGames.OIDC/InMemoryTokenStore.cs
+++ b/Engine/Source/Programs/Shared/EpicGames.OIDC/InMemoryTokenStore.cs
@@ -12,7 +12,6 @@ namespace EpicGames.OIDC;
 public class InMemoryTokenStore : ITokenStore
 {
 	private readonly Dictionary<string, string> _providerToRefreshToken = new();
-	private readonly Dictionary<string, (string Token, DateTimeOffset Expiry)> _providerToAccessToken = new();
 
 	/// <inheritdoc/>
 	public virtual bool TryGetRefreshToken(string oidcProvider, out string refreshToken)
@@ -33,27 +32,6 @@ public class InMemoryTokenStore : ITokenStore
 		_providerToRefreshToken[providerIdentifier] = refreshToken;
 	}
 
-	/// <inheritdoc/>
-	public virtual bool TryGetAccessToken(string oidcProvider, out string accessToken, out DateTimeOffset expiry)
-	{
-		if (!_providerToAccessToken.TryGetValue(oidcProvider, out (string Token, DateTimeOffset Expiry) cached))
-		{
-			accessToken = "";
-			expiry = DateTimeOffset.MinValue;
-			return false;
-		}
-
-		accessToken = cached.Token;
-		expiry = cached.Expiry;
-		return true;
-	}
-
-	/// <inheritdoc/>
-	public virtual void AddAccessToken(string providerIdentifier, string accessToken, DateTimeOffset expiry)
-	{
-		_providerToAccessToken[providerIdentifier] = (accessToken, expiry);
-	}
-
 	/// <inheritdoc/>
 	public virtual void Save()
 	{
diff --git a/Engine/Source/Programs/Shared/EpicGames.OIDC/MacOSTokenStore.cs b/Engine/Source/Programs/Shared/EpicGames.OIDC/MacOSTokenStore.cs
--- a/Engine/Source/Programs/Shared/EpicGames.OIDC/MacOSTokenStore.cs
+++ b/Engine/Source/Programs/Shared/EpicGames.OIDC/MacOSTokenStore.cs
@@ -4,7 +4,6 @@ using System;
 using System.Collections.Generic;
 using System.Diagnostics;
 using System.Text;
-using System.Text.Json;
 
 #pragma warning disable CS1591 // Missing XML documentation on public types
 
@@ -12,111 +11,31 @@ namespace EpicGames.OIDC
 {
 	public sealed class MacOSTokenStore : ITokenStore
 	{
-		// Cache provider data in memory
-		private readonly Dictionary<string, ProviderTokenData> _providerData = new Dictionary<string, ProviderTokenData>();
-
 		public MacOSTokenStore()
 		{
 		}
 
-		private ProviderTokenData? TryGetProviderData(string oidcProvider)
-		{
-			// Check in-memory cache first
-			if (_providerData.TryGetValue(oidcProvider, out ProviderTokenData? cached))
-			{
-				return cached;
-			}
-
-			// Try to load from Keychain (V2 format first, then legacy)
-			string? json = KeychainHelper.GetKeychainItem($"{oidcProvider}-v2");
-			if (json != null)
-			{
-				try
-				{
-					ProviderTokenData? data = JsonSerializer.Deserialize(json, ProviderTokenDataContext.Default.ProviderTokenData);
-					if (data != null)
-					{
-						_providerData[oidcProvider] = data;
-						return data;
-					}
-				}
-				catch (JsonException)
-				{
-					// Fall through to legacy
-				}
-			}
-
-			// Try legacy format (just refresh token)
-			string? legacyToken = KeychainHelper.GetKeychainItem(oidcProvider);
-			if (legacyToken != null)
-			{
-				ProviderTokenData data = new ProviderTokenData { RefreshToken = legacyToken };
-				_providerData[oidcProvider] = data;
-				return data;
-			}
-
-			return null;
-		}
-
-		private void SetProviderData(string providerIdentifier, ProviderTokenData tokenData)
-		{
-			_providerData[providerIdentifier] = tokenData;
-
-			// Write V2 format only (JSON with all token data)
-			// Note: We don't update the legacy keychain entry - older code will continue using existing entries
-			string json = JsonSerializer.Serialize(tokenData, ProviderTokenDataContext.Default.ProviderTokenData);
-			KeychainHelper.AddKeychainItem(json, $"{providerIdentifier}-v2");
-		}
-
 		public bool TryGetRefreshToken(string oidcProvider, out string refreshToken)
 		{
-			ProviderTokenData? data = TryGetProviderData(oidcProvider);
-			if (data == null || string.IsNullOrEmpty(data.RefreshToken))
+			string? token = KeychainHelper.GetKeychainItem(oidcProvider);
+
+			if (token == null)
 			{
 				refreshToken = "";
 				return false;
 			}
 
-			refreshToken = data.RefreshToken;
+			refreshToken = token;
 			return true;
 		}
 
 		public void AddRefreshToken(string providerIdentifier, string refreshToken)
 		{
-			ProviderTokenData? existing = TryGetProviderData(providerIdentifier);
-			ProviderTokenData tokenData = existing ?? new ProviderTokenData();
-			tokenData.RefreshToken = refreshToken;
-			SetProviderData(providerIdentifier, tokenData);
-		}
-
-		public bool TryGetAccessToken(string oidcProvider, out string accessToken, out DateTimeOffset expiry)
-		{
-			accessToken = "";
-			expiry = DateTimeOffset.MinValue;
-
-			ProviderTokenData? data = TryGetProviderData(oidcProvider);
-			if (data == null || string.IsNullOrEmpty(data.AccessToken) || data.AccessTokenExpireTime == null)
-			{
-				return false;
-			}
-
-			accessToken = data.AccessToken;
-			expiry = DateTimeOffset.FromUnixTimeSeconds(data.AccessTokenExpireTime.Value);
-			return true;
-		}
-
-		public void AddAccessToken(string providerIdentifier, string accessToken, DateTimeOffset expiry)
-		{
-			ProviderTokenData? existing = TryGetProviderData(providerIdentifier);
-			ProviderTokenData tokenData = existing ?? new ProviderTokenData();
-			tokenData.AccessToken = accessToken;
-			tokenData.AccessTokenExpireTime = expiry.ToUnixTimeSeconds();
-			SetProviderData(providerIdentifier, tokenData);
+			KeychainHelper.AddKeychainItem(refreshToken, providerIdentifier);
 		}
 
 		public void Save()
 		{
-			// Keychain writes are immediate, nothing to do here
 		}
 
 		public void Dispose()
diff --git a/Engine/Source/Programs/Shared/EpicGames.OIDC/OidcTokenManager.cs b/Engine/Source/Programs/Shared/EpicGames.OIDC/OidcTokenManager.cs
--- a/Engine/Source/Programs/Shared/EpicGames.OIDC/OidcTokenManager.cs
+++ b/Engine/Source/Programs/Shared/EpicGames.OIDC/OidcTokenManager.cs
@@ -33,10 +33,6 @@ namespace EpicGames.OIDC
 
 		public Task<OidcTokenInfo?> TryGetAccessToken(string providerIdentifier, CancellationToken cancellationToken = default);
 
-		public Task<OidcTokenInfo> GetAccessToken(string providerIdentifier, GetAccessTokenOptions options, CancellationToken cancellationToken = default);
-
-		public Task<OidcTokenInfo?> TryGetAccessToken(string providerIdentifier, GetAccessTokenOptions options, CancellationToken cancellationToken = default);
-
 		public OidcStatus GetStatusForProvider(string providerIdentifier);
 	}
 
@@ -160,18 +156,6 @@ namespace EpicGames.OIDC
 			return _tokenClients[providerIdentifier].TryGetAccessTokenAsync(cancellationToken);
 		}
 
-		/// <inheritdoc/>
-		public Task<OidcTokenInfo> GetAccessToken(string providerIdentifier, GetAccessTokenOptions options, CancellationToken cancellationToken = default)
-		{
-			return _tokenClients[providerIdentifier].GetAccessTokenAsync(options, cancellationToken);
-		}
-
-		/// <inheritdoc/>
-		public Task<OidcTokenInfo?> TryGetAccessToken(string providerIdentifier, GetAccessTokenOptions options, CancellationToken cancellationToken = default)
-		{
-			return _tokenClients[providerIdentifier].TryGetAccessTokenAsync(options, cancellationToken);
-		}
-
 		/// <inheritdoc/>
 		public OidcStatus GetStatusForProvider(string providerIdentifier)
 		{
@@ -237,19 +221,6 @@ namespace EpicGames.OIDC
 		}
 	}
 
-	/// <summary>
-	/// Options for retrieving access tokens
-	/// </summary>
-	public class GetAccessTokenOptions
-	{
-		/// <summary>
-		/// Whether to use disk cache for access tokens.
-		/// When enabled, access tokens are persisted to disk and reused across process restarts if still valid.
-		/// Default is false for backwards compatibility.
-		/// </summary>
-		public bool UseCache { get; set; } = false;
-	}
-
 	/// <summary>
 	/// The data read from a http request into the oidc client, mostly a wrapper around the http.listener context
 	/// </summary>
@@ -438,7 +409,6 @@ namespace EpicGames.OIDC
 		private readonly Uri _authorityUri;
 		private readonly string _clientId;
 		private readonly string _scopes;
-		private readonly TimeSpan _tokenValidityBuffer = TimeSpan.FromMinutes(10);
 
 		private string? _refreshToken;
 		private string? _accessToken;
@@ -866,43 +836,13 @@ namespace EpicGames.OIDC
 			return await TryGetAccessTokenAsync(cancellationToken) ?? throw new NotLoggedInException();
 		}
 
-		public async Task<OidcTokenInfo> GetAccessTokenAsync(GetAccessTokenOptions options, CancellationToken cancellationToken)
-		{
-			return await TryGetAccessTokenAsync(options, cancellationToken) ?? throw new NotLoggedInException();
-		}
-
-		public Task<OidcTokenInfo?> TryGetAccessTokenAsync(CancellationToken cancellationToken)
-		{
-			return TryGetAccessTokenAsync(new GetAccessTokenOptions(), cancellationToken);
-		}
-
-		public async Task<OidcTokenInfo?> TryGetAccessTokenAsync(GetAccessTokenOptions options, CancellationToken cancellationToken)
+		public async Task<OidcTokenInfo?> TryGetAccessTokenAsync(CancellationToken cancellationToken)
 		{
 			if (String.IsNullOrEmpty(_refreshToken))
 			{
 				throw new NotLoggedInException();
 			}
 
-			// If caching is enabled, try to load from disk cache first
-			if (options.UseCache && _tokenStore.TryGetAccessToken(_name, out string? cachedAccessToken, out DateTimeOffset cachedExpiry))
-			{
-				OidcTokenInfo cachedInfo = new()
-				{
-					RefreshToken = _refreshToken,
-					AccessToken = cachedAccessToken,
-					TokenExpiry = cachedExpiry
-				};
-
-				// Ensure token is valid for at least another few minutes
-				if (cachedInfo.IsValid(DateTime.Now - _tokenValidityBuffer))
-				{
-					_accessToken = cachedAccessToken;
-					_tokenExpiry = cachedExpiry;
-					return cachedInfo;
-				}
-			}
-
-			// Check in-memory cache
 			OidcTokenInfo tokenInfo = new()
 			{
 				RefreshToken = _refreshToken,
@@ -910,22 +850,13 @@ namespace EpicGames.OIDC
 				TokenExpiry = _tokenExpiry
 			};
 
-			// Ensure token is valid for at least another few minutes
-			if (tokenInfo.IsValid(DateTime.Now - _tokenValidityBuffer))
+			// Ensure token is valid for at least another two minutes
+			if (tokenInfo.IsValid(DateTime.Now.AddMinutes(-2)))
 			{
 				return tokenInfo;
 			}
 
-			OidcTokenInfo? refreshedInfo = await TryDoRefreshTokenAsync(_refreshToken, cancellationToken);
-
-			// If caching is enabled and refresh succeeded, save to disk cache
-			if (options.UseCache && refreshedInfo != null && !String.IsNullOrEmpty(refreshedInfo.AccessToken))
-			{
-				_tokenStore.AddAccessToken(_name, refreshedInfo.AccessToken, refreshedInfo.TokenExpiry);
-				_tokenStore.Save();
-			}
-
-			return refreshedInfo;
+			return await TryDoRefreshTokenAsync(_refreshToken, cancellationToken);
 		}
 
 		public OidcStatus GetStatus()
diff --git a/Engine/Source/Programs/Shared/EpicGames.OIDC/TokenStore.cs b/Engine/Source/Programs/Shared/EpicGames.OIDC/TokenStore.cs
--- a/Engine/Source/Programs/Shared/EpicGames.OIDC/TokenStore.cs
+++ b/Engine/Source/Programs/Shared/EpicGames.OIDC/TokenStore.cs
@@ -19,24 +19,24 @@ namespace EpicGames.OIDC
 
 	internal class TokenStoreState
 	{
-		/// <summary>
-		/// Legacy format - encrypted refresh tokens only (kept for backwards compatibility when reading)
-		/// </summary>
 		public Dictionary<string, string> Providers { get; set; } = new Dictionary<string, string>();
-
-		/// <summary>
-		/// New format - encrypted JSON containing both refresh and access tokens
-		/// </summary>
-		public Dictionary<string, string> ProvidersV2 { get; set; } = new Dictionary<string, string>();
-
-		public string Key { get; set; } = string.Empty;
+		public string Key { get; set; }
 
 		[JsonConstructor]
-		public TokenStoreState(string? key = null, Dictionary<string, string>? providers = null, Dictionary<string, string>? providersV2 = null)
+		public TokenStoreState(string key, Dictionary<string, string> providers)
+		{
+			Key = key;
+			Providers = providers;
+		}
+
+		public TokenStoreState(byte[] key, Dictionary<string, byte[]> providers)
 		{
-			Key = key ?? string.Empty;
-			Providers = providers ?? new Dictionary<string, string>();
-			ProvidersV2 = providersV2 ?? new Dictionary<string, string>();
+			Key = Convert.ToBase64String(key);
+
+			foreach (KeyValuePair<string, byte[]> pair in providers)
+			{
+				Providers[pair.Key] = Convert.ToBase64String(pair.Value);
+			}
 		}
 	}
 
@@ -45,66 +45,56 @@ namespace EpicGames.OIDC
 	/// </summary>
 	public class FilesystemTokenStore : ITokenStore, IDisposable
 	{
-		private readonly Dictionary<string, byte[]> _providerData = new Dictionary<string, byte[]>();
+		private readonly Dictionary<string, byte[]> _providerToRefreshToken;
 		private readonly List<string> _dirtyProviders = new List<string>();
 		private readonly SymmetricAlgorithm _crypt = Aes.Create();
-		private readonly FileInfo? _storePath;
 
-		public FilesystemTokenStore() : this(null)
+		public FilesystemTokenStore()
 		{
+			_providerToRefreshToken = ReadStoreFromDisk();
 		}
 
-		internal FilesystemTokenStore(FileInfo? storePath)
-		{
-			_storePath = storePath;
-
-			// Load V2 provider data from disk
-			TokenStoreState? state = ReadStateFromDisk();
-			if (state != null)
-			{
-				if (!string.IsNullOrEmpty(state.Key))
-				{
-					_crypt.IV = Convert.FromBase64String(state.Key);
-				}
-
-				foreach ((string key, string value) in state.ProvidersV2)
-				{
-					_providerData[key] = Convert.FromBase64String(value);
-				}
-			}
-		}
-
-		private static FileInfo GetDefaultStorePath()
+		private static FileInfo GetStorePath()
 		{
 			return new FileInfo(Path.Combine(Environment.GetFolderPath(Environment.SpecialFolder.ApplicationData), "Epic", "UnrealEngine", "Common", "OidcToken", "tokenStore.dat"));
 		}
 
-		private FileInfo GetStorePath() => _storePath ?? GetDefaultStorePath();
-
-		private TokenStoreState? ReadStateFromDisk()
+		private Dictionary<string, byte[]> ReadStoreFromDisk()
 		{
 			FileInfo fi = GetStorePath();
-			fi.Refresh(); // Refresh cached metadata in case file was created after FileInfo was instantiated
 			if (!fi.Exists)
 			{
-				return null;
+				// if we have no store on disk then we just initialize it to empty
+				return new Dictionary<string, byte[]>();
 			}
 
 			using FileStream fs = fi.OpenRead();
+			TokenStoreState? state;
 			try
 			{
-				return JsonSerializer.Deserialize(fs, TokenStoreStateContext.Default.TokenStoreState);
+				state = JsonSerializer.Deserialize(fs, TokenStoreStateContext.Default.TokenStoreState);
 			}
 			catch (JsonException)
 			{
-				return null;
+				state = null;
 			}
-		}
 
-		/// <summary>
-		/// Reads the raw state from disk for test verification.
-		/// </summary>
-		internal TokenStoreState? ReadRawState() => ReadStateFromDisk();
+			if (state == null)
+			{
+				// if we fail to deserialize the state just drop it, will mean users will need to login again
+				return new Dictionary<string, byte[]>();
+			}
+			Dictionary<string, byte[]> providers = new Dictionary<string, byte[]>();
+
+			foreach ((string key, string value) in state.Providers)
+			{
+				providers[key] = Convert.FromBase64String(value);
+			}
+
+			_crypt.IV = Convert.FromBase64String(state.Key);
+
+			return providers;
+		}
 
 		private void SaveStoreToDisk()
 		{
@@ -117,38 +107,25 @@ namespace EpicGames.OIDC
 
 			lock (_dirtyProviders)
 			{
+				// no providers have changed, do not touch the state file
 				if (_dirtyProviders.Count == 0)
 				{
 					return;
 				}
 
-				// Read back the existing state from disk
-				TokenStoreState? existingState = ReadStateFromDisk();
+				// read back the state of all providers but only overwrite the state of the ones we have actually got new state for (are dirty)
+				Dictionary<string, byte[]> providers = ReadStoreFromDisk();
 
-				// Build new state, preserving legacy Providers (we don't modify them)
-				TokenStoreState newState = new TokenStoreState(Convert.ToBase64String(_crypt.IV));
-				if (existingState != null)
-				{
-					// Preserve legacy Providers as-is
-					newState.Providers = existingState.Providers;
-					// Start with existing ProvidersV2
-					foreach ((string key, string value) in existingState.ProvidersV2)
-					{
-						newState.ProvidersV2[key] = value;
-					}
-				}
-
-				// Merge in our dirty changes to ProvidersV2
 				foreach (string providerId in _dirtyProviders)
 				{
-					newState.ProvidersV2[providerId] = Convert.ToBase64String(_providerData[providerId]);
+					providers[providerId] = _providerToRefreshToken[providerId];
 				}
 
 				string tempFile = Path.GetTempFileName();
 				{
 					using FileStream fs = new FileStream(tempFile, FileMode.Create, FileAccess.Write);
 					using Utf8JsonWriter writer = new Utf8JsonWriter(fs);
-					JsonSerializer.Serialize<TokenStoreState>(writer, newState, TokenStoreStateContext.Default.TokenStoreState);
+					JsonSerializer.Serialize<TokenStoreState>(writer, new TokenStoreState(_crypt.IV, providers), TokenStoreStateContext.Default.TokenStoreState);
 				}
 
 				File.Move(tempFile, fi.FullName, true);
@@ -156,118 +133,34 @@ namespace EpicGames.OIDC
 			}
 		}
 
-		private string? TryDecryptData(byte[] encryptedData)
-		{
-			try
-			{
-				ICryptoTransform decryptor = _crypt.CreateDecryptor(GetSeed(), _crypt.IV);
-				byte[] decrypted = decryptor.TransformFinalBlock(encryptedData, 0, encryptedData.Length);
-				return Encoding.Unicode.GetString(decrypted);
-			}
-			catch (CryptographicException)
-			{
-				return null;
-			}
-		}
-
-		private ProviderTokenData? TryGetProviderData(string oidcProvider)
-		{
-			// Try V2 format first (in-memory)
-			if (_providerData.TryGetValue(oidcProvider, out byte[]? encryptedData))
-			{
-				string? json = TryDecryptData(encryptedData);
-				if (json != null)
-				{
-					try
-					{
-						return JsonSerializer.Deserialize(json, ProviderTokenDataContext.Default.ProviderTokenData);
-					}
-					catch (JsonException)
-					{
-						// Fall through to legacy
-					}
-				}
-			}
-
-			// Fall back to legacy format from disk (refresh token only)
-			TokenStoreState? state = ReadStateFromDisk();
-			if (state != null && state.Providers.TryGetValue(oidcProvider, out string? legacyBase64))
-			{
-				byte[] legacyEncrypted = Convert.FromBase64String(legacyBase64);
-				string? refreshToken = TryDecryptData(legacyEncrypted);
-				if (refreshToken != null)
-				{
-					return new ProviderTokenData { RefreshToken = refreshToken };
-				}
-			}
-
-			return null;
-		}
-
-		private void SetProviderData(string providerIdentifier, ProviderTokenData tokenData)
-		{
-			string json = JsonSerializer.Serialize(tokenData, ProviderTokenDataContext.Default.ProviderTokenData);
-			byte[] bytes = Encoding.Unicode.GetBytes(json);
-#pragma warning disable CA5401
-			ICryptoTransform encryptor = _crypt.CreateEncryptor(GetSeed(), _crypt.IV);
-#pragma warning restore CA5401
-			byte[] encrypted = encryptor.TransformFinalBlock(bytes, 0, bytes.Length);
-
-			_providerData[providerIdentifier] = encrypted;
-
-			lock (_dirtyProviders)
-			{
-				if (!_dirtyProviders.Contains(providerIdentifier))
-				{
-					_dirtyProviders.Add(providerIdentifier);
-				}
-			}
-		}
-
 		public bool TryGetRefreshToken(string oidcProvider, out string refreshToken)
 		{
-			ProviderTokenData? data = TryGetProviderData(oidcProvider);
-			if (data == null || string.IsNullOrEmpty(data.RefreshToken))
+			if (!_providerToRefreshToken.TryGetValue(oidcProvider, out byte[]? encryptedToken))
 			{
 				refreshToken = "";
 				return false;
 			}
 
-			refreshToken = data.RefreshToken;
+			ICryptoTransform decryptor = _crypt.CreateDecryptor(GetSeed(), _crypt.IV);
+			byte[] bytes = decryptor.TransformFinalBlock(encryptedToken, 0, encryptedToken.Length);
+			refreshToken = Encoding.Unicode.GetString(bytes);
 			return true;
 		}
 
 		public void AddRefreshToken(string providerIdentifier, string refreshToken)
 		{
-			ProviderTokenData? existing = TryGetProviderData(providerIdentifier);
-			ProviderTokenData tokenData = existing ?? new ProviderTokenData();
-			tokenData.RefreshToken = refreshToken;
-			SetProviderData(providerIdentifier, tokenData);
-		}
+			byte[] bytes = Encoding.Unicode.GetBytes(refreshToken);
+#pragma warning disable CA5401 // Do not use CreateEncryptor with non-default IV
+			ICryptoTransform encryptor = _crypt.CreateEncryptor(GetSeed(), _crypt.IV);
+#pragma warning restore CA5401 // Do not use CreateEncryptor with non-default IV
 
-		public bool TryGetAccessToken(string oidcProvider, out string accessToken, out DateTimeOffset expiry)
-		{
-			accessToken = "";
-			expiry = DateTimeOffset.MinValue;
+			byte[] encryptedToken = encryptor.TransformFinalBlock(bytes, 0, bytes.Length);
+			_providerToRefreshToken[providerIdentifier] = encryptedToken;
 
-			ProviderTokenData? data = TryGetProviderData(oidcProvider);
-			if (data == null || string.IsNullOrEmpty(data.AccessToken) || data.AccessTokenExpireTime == null)
+			lock (_dirtyProviders)
 			{
-				return false;
+				_dirtyProviders.Add(providerIdentifier);
 			}
-
-			accessToken = data.AccessToken;
-			expiry = DateTimeOffset.FromUnixTimeSeconds(data.AccessTokenExpireTime.Value);
-			return true;
-		}
-
-		public void AddAccessToken(string providerIdentifier, string accessToken, DateTimeOffset expiry)
-		{
-			ProviderTokenData? existing = TryGetProviderData(providerIdentifier);
-			ProviderTokenData tokenData = existing ?? new ProviderTokenData();
-			tokenData.AccessToken = accessToken;
-			tokenData.AccessTokenExpireTime = expiry.ToUnixTimeSeconds();
-			SetProviderData(providerIdentifier, tokenData);
 		}
 
 		public void Save()
@@ -275,7 +168,7 @@ namespace EpicGames.OIDC
 			SaveStoreToDisk();
 		}
 
-		internal static byte[] GetSeed()
+		static byte[] GetSeed()
 		{
 			return new byte[]
 			{
@@ -299,4 +192,4 @@ namespace EpicGames.OIDC
 			GC.SuppressFinalize(this);
 		}
 	}
-}
+}
\ No newline at end of file
diff --git a/Engine/Source/Programs/Shared/EpicGames.OIDC/WindowsTokenStore.cs b/Engine/Source/Programs/Shared/EpicGames.OIDC/WindowsTokenStore.cs
--- a/Engine/Source/Programs/Shared/EpicGames.OIDC/WindowsTokenStore.cs
+++ b/Engine/Source/Programs/Shared/EpicGames.OIDC/WindowsTokenStore.cs
@@ -23,86 +23,86 @@ namespace EpicGames.OIDC
 
 	internal class WindowsTokenStoreState
 	{
-		/// <summary>
-		/// Legacy format - encrypted refresh tokens only (kept for backwards compatibility when reading)
-		/// </summary>
 		public Dictionary<string, string> Providers { get; set; } = new Dictionary<string, string>();
 
-		/// <summary>
-		/// New format - encrypted JSON containing both refresh and access tokens
-		/// </summary>
-		public Dictionary<string, string> ProvidersV2 { get; set; } = new Dictionary<string, string>();
-
 		[JsonConstructor]
-		public WindowsTokenStoreState(Dictionary<string, string>? providers = null, Dictionary<string, string>? providersV2 = null)
+		public WindowsTokenStoreState(Dictionary<string, string> providers)
 		{
-			Providers = providers ?? new Dictionary<string, string>();
-			ProvidersV2 = providersV2 ?? new Dictionary<string, string>();
+			Providers = providers;
+		}
+
+		public WindowsTokenStoreState(Dictionary<string, byte[]> providers)
+		{
+			foreach (KeyValuePair<string, byte[]> pair in providers)
+			{
+				Providers[pair.Key] = Convert.ToBase64String(pair.Value);
+			}
 		}
 	}
 
 	public class WindowsTokenStore : ITokenStore, IDisposable
 	{
 		private readonly ILogger<WindowsTokenStore>? _logger = null;
-		private readonly FileInfo? _storePath = null;
 
-		private readonly Dictionary<string, byte[]> _providerData = new Dictionary<string, byte[]>();
+		private readonly Dictionary<string, byte[]> _providerToRefreshToken = new Dictionary<string, byte[]>();
 		private readonly List<string> _dirtyProviders = new List<string>();
 
-		public WindowsTokenStore() : this(null, null)
+		public WindowsTokenStore()
 		{
+			_providerToRefreshToken = ReadStoreFromDisk();
 		}
 
 		[ActivatorUtilitiesConstructor]
-		public WindowsTokenStore(ILogger<WindowsTokenStore>? logger) : this(logger, null)
-		{
-		}
-
-		internal WindowsTokenStore(ILogger<WindowsTokenStore>? logger, FileInfo? storePath)
+		public WindowsTokenStore(ILogger<WindowsTokenStore> logger)
 		{
 			_logger = logger;
-			_storePath = storePath;
 
-			// Load V2 provider data from disk
-			WindowsTokenStoreState? state = ReadStateFromDisk();
-			if (state != null)
-			{
-				foreach ((string key, string value) in state.ProvidersV2)
-				{
-					_providerData[key] = Convert.FromBase64String(value);
-				}
-			}
+			_providerToRefreshToken = ReadStoreFromDisk();
 		}
 
-		private static FileInfo GetDefaultStorePath()
+		private static FileInfo GetStorePath()
 		{
 			return new FileInfo(Path.Combine(Environment.GetFolderPath(Environment.SpecialFolder.LocalApplicationData), "UnrealEngine", "Common", "OidcToken", "oidcTokenStore.dat"));
 		}
 
-		private FileInfo GetStorePath() => _storePath ?? GetDefaultStorePath();
-
-		private WindowsTokenStoreState? ReadStateFromDisk()
+		private Dictionary<string, byte[]> ReadStoreFromDisk()
 		{
 			FileInfo fi = GetStorePath();
-			fi.Refresh(); // Refresh cached metadata in case file was created after FileInfo was instantiated
 			if (!fi.Exists)
 			{
 				_logger?.LogDebug("No existing token store found at {Path}. Assuming empty store.", fi.FullName);
-				return null;
+				// if we have no store on disk then we just initialize it to empty
+				return new Dictionary<string, byte[]>();
 			}
 
 			using FileStream fs = fi.Open(FileMode.Open, FileAccess.Read, FileShare.Read);
 			using TextReader tr = new StreamReader(fs);
 
+			WindowsTokenStoreState? state;
 			try
 			{
-				return JsonSerializer.Deserialize(tr.ReadToEnd(), WindowsTokenStoreStateContext.Default.WindowsTokenStoreState);
+				state = JsonSerializer.Deserialize(tr.ReadToEnd(), WindowsTokenStoreStateContext.Default.WindowsTokenStoreState);
 			}
-			catch (JsonException je)
+			catch (JsonException)
 			{
-				_logger?.LogDebug("Failed to deserialize state. Reason: {Message}", je.Message);
-				return null;
+				state = null;
 			}
+
+			if (state == null)
+			{
+				_logger?.LogDebug("Failed to deserialize state. Dropping the existing state.");
+				// if we fail to deserialize the state just drop it, will mean users will need to login again
+				return new Dictionary<string, byte[]>();
+			}
+
+			Dictionary<string, byte[]> providers = new Dictionary<string, byte[]>();
+
+			foreach ((string key, string value) in state.Providers)
+			{
+				providers[key] = Convert.FromBase64String(value);
+			}
+
+			return providers;
 		}
 
 		private void SaveStoreToDisk()
@@ -116,7 +116,7 @@ namespace EpicGames.OIDC
 
 			lock (_dirtyProviders)
 			{
-				// No providers have changed, do not touch the state file
+				// no providers have changed, do not touch the state file
 				if (_dirtyProviders.Count == 0)
 				{
 					return;
@@ -130,35 +130,21 @@ namespace EpicGames.OIDC
 				}
 				catch (AbandonedMutexException)
 				{
-				}
 
-				// Read back the existing state from disk
-				WindowsTokenStoreState? existingState = ReadStateFromDisk();
-
-				// Build new state, preserving legacy Providers (we don't modify them)
-				WindowsTokenStoreState newState = new ();
-				if (existingState != null)
-				{
-					// Preserve legacy Providers as-is
-					newState.Providers = existingState.Providers;
-					// Start with existing ProvidersV2
-					foreach ((string key, string value) in existingState.ProvidersV2)
-					{
-						newState.ProvidersV2[key] = value;
-					}
 				}
 
-				// Merge in our dirty changes to ProvidersV2
+				// read back the state of all providers but only overwrite the state of the ones we have actually got new state for (are dirty)
+				Dictionary<string, byte[]> providers = ReadStoreFromDisk();
+
 				foreach (string providerId in _dirtyProviders)
 				{
-					newState.ProvidersV2[providerId] = Convert.ToBase64String(_providerData[providerId]);
+					providers[providerId] = _providerToRefreshToken[providerId];
 				}
-
 				string tempFile = Path.GetTempFileName();
 				{
 					using FileStream fs = new FileStream(tempFile, FileMode.Create, FileAccess.Write);
 					using Utf8JsonWriter writer = new Utf8JsonWriter(fs);
-					JsonSerializer.Serialize(writer, newState, WindowsTokenStoreStateContext.Default.WindowsTokenStoreState);
+					JsonSerializer.Serialize(writer, new WindowsTokenStoreState(providers), WindowsTokenStoreStateContext.Default.WindowsTokenStoreState);
 				}
 
 				File.Move(tempFile, fi.FullName, true);
@@ -169,123 +155,59 @@ namespace EpicGames.OIDC
 			}
 		}
 
-		private string? TryDecryptData(byte[] encryptedData, string description, string oidcProvider)
+		public bool TryGetRefreshToken(string oidcProvider, out string refreshToken)
 		{
-			try
-			{
-				byte[] decrypted = CryptProtectDataHelper.DoCryptUnprotectData(encryptedData, description, GetEntropy(oidcProvider));
-				return Encoding.Unicode.GetString(decrypted);
-			}
-			catch (Win32Exception e)
+			if (!_providerToRefreshToken.TryGetValue(oidcProvider, out byte[]? encryptedToken))
 			{
-				if (e.NativeErrorCode != 13 && e.NativeErrorCode != unchecked((int)0x8009000B))
-				{
-					throw;
-				}
-				_logger?.LogDebug("Unable to decrypt token data for provider {Provider}. Ignoring.", oidcProvider);
-				return null;
+				refreshToken = "";
+				return false;
 			}
-		}
 
-		private ProviderTokenData? TryGetProviderData(string oidcProvider)
-		{
-			// Try V2 format first (in-memory)
-			if (_providerData.TryGetValue(oidcProvider, out byte[]? encryptedData))
+			try
 			{
-				string? json = TryDecryptData(encryptedData, $"OidcTokenV2-{oidcProvider}", oidcProvider);
-				if (json != null)
-				{
-					try
-					{
-						return JsonSerializer.Deserialize(json, ProviderTokenDataContext.Default.ProviderTokenData);
-					}
-					catch (JsonException)
-					{
-						_logger?.LogDebug("Unable to deserialize V2 token data for provider {Provider}. Ignoring.", oidcProvider);
-					}
-				}
-			}
+				byte[] bytes = CryptProtectDataHelper.DoCryptUnprotectData(encryptedToken, $"OidcToken-{oidcProvider}", GetEntropy(oidcProvider));
+				refreshToken = Encoding.Unicode.GetString(bytes);
 
-			// Fall back to legacy format from disk (refresh token only)
-			WindowsTokenStoreState? state = ReadStateFromDisk();
-			if (state != null && state.Providers.TryGetValue(oidcProvider, out string? legacyBase64))
+				return true;
+			}
+			catch (Win32Exception e)
 			{
-				byte[] legacyEncrypted = Convert.FromBase64String(legacyBase64);
-				string? refreshToken = TryDecryptData(legacyEncrypted, $"OidcToken-{oidcProvider}", oidcProvider);
-				if (refreshToken != null)
+				if (e.NativeErrorCode == 13) // data is invalid
 				{
-					return new ProviderTokenData { RefreshToken = refreshToken };
+					// unable to decrypt the data, ignore it
+					refreshToken = "";
+					_logger?.LogDebug("Unable to decrypt refresh token. Ignoring.");
+					return false;
 				}
-			}
-
-			return null;
-		}
-
-		private void SetProviderData(string providerIdentifier, ProviderTokenData tokenData)
-		{
-			string json = JsonSerializer.Serialize(tokenData, ProviderTokenDataContext.Default.ProviderTokenData);
-			byte[] encrypted = CryptProtectDataHelper.DoCryptProtectData(Encoding.Unicode.GetBytes(json), $"OidcTokenV2-{providerIdentifier}", GetEntropy(providerIdentifier));
-
-			_providerData[providerIdentifier] = encrypted;
-
-			lock (_dirtyProviders)
-			{
-				if (!_dirtyProviders.Contains(providerIdentifier))
+				if (e.NativeErrorCode == unchecked((int)0x8009000B)) // key not valid for use in specified state
 				{
-					_dirtyProviders.Add(providerIdentifier);
+					// unable to decrypt the data, ignore it
+					refreshToken = "";
+					_logger?.LogDebug("Unable to decrypt refresh token, key not valid for use in specified state. Ignoring.");
+					return false;
 				}
+				throw;
 			}
 		}
 
-		public bool TryGetRefreshToken(string oidcProvider, out string refreshToken)
+		private static byte[] GetEntropy(string oidcProvider)
 		{
-			ProviderTokenData? data = TryGetProviderData(oidcProvider);
-			if (data == null || string.IsNullOrEmpty(data.RefreshToken))
-			{
-				refreshToken = "";
-				return false;
-			}
+			byte[] providerBytes = Encoding.UTF8.GetBytes(oidcProvider);
 
-			refreshToken = data.RefreshToken;
-			return true;
-		}
-
-		internal static byte[] GetEntropy(string oidcProvider)
-		{
-			return Encoding.UTF8.GetBytes(oidcProvider);
+			return providerBytes;
 		}
 
 		public void AddRefreshToken(string providerIdentifier, string refreshToken)
 		{
-			ProviderTokenData? existing = TryGetProviderData(providerIdentifier);
-			ProviderTokenData tokenData = existing ?? new ProviderTokenData();
-			tokenData.RefreshToken = refreshToken;
-			SetProviderData(providerIdentifier, tokenData);
-		}
+			byte[] bytes = Encoding.Unicode.GetBytes(refreshToken);
+			byte[] encryptedToken = CryptProtectDataHelper.DoCryptProtectData(bytes, $"OidcToken-{providerIdentifier}", GetEntropy(providerIdentifier));
 
-		public bool TryGetAccessToken(string oidcProvider, out string accessToken, out DateTimeOffset expiry)
-		{
-			accessToken = "";
-			expiry = DateTimeOffset.MinValue;
+			_providerToRefreshToken[providerIdentifier] = encryptedToken;
 
-			ProviderTokenData? data = TryGetProviderData(oidcProvider);
-			if (data == null || string.IsNullOrEmpty(data.AccessToken) || data.AccessTokenExpireTime == null)
+			lock (_dirtyProviders)
 			{
-				return false;
+				_dirtyProviders.Add(providerIdentifier);
 			}
-
-			accessToken = data.AccessToken;
-			expiry = DateTimeOffset.FromUnixTimeSeconds(data.AccessTokenExpireTime.Value);
-			return true;
-		}
-
-		public void AddAccessToken(string providerIdentifier, string accessToken, DateTimeOffset expiry)
-		{
-			ProviderTokenData? existing = TryGetProviderData(providerIdentifier);
-			ProviderTokenData tokenData = existing ?? new ProviderTokenData();
-			tokenData.AccessToken = accessToken;
-			tokenData.AccessTokenExpireTime = expiry.ToUnixTimeSeconds();
-			SetProviderData(providerIdentifier, tokenData);
 		}
 
 		public void Save()
@@ -293,11 +215,6 @@ namespace EpicGames.OIDC
 			SaveStoreToDisk();
 		}
 
-		/// <summary>
-		/// Reads the raw state from disk for test verification.
-		/// </summary>
-		internal WindowsTokenStoreState? ReadRawState() => ReadStateFromDisk();
-
 		protected virtual void Dispose(bool disposing)
 		{
 			if (disposing)
@@ -484,4 +401,4 @@ namespace EpicGames.OIDC
 		}
 	}
 #pragma warning restore IDE1006 // Naming Styles
-}
+}
\ No newline at end of file
