[Horde][Analytics] Support both ordered and unordered commit id.

#Synced-CL 49930492

diff --git a/Engine/Source/Programs/Shared/EpicGames.Analytics.PerformanceTrends/PerformanceTrendTelemetry.cs b/Engine/Source/Programs/Shared/EpicGames.Analytics.PerformanceTrends/PerformanceTrendTelemetry.cs
--- a/Engine/Source/Programs/Shared/EpicGames.Analytics.PerformanceTrends/PerformanceTrendTelemetry.cs
+++ b/Engine/Source/Programs/Shared/EpicGames.Analytics.PerformanceTrends/PerformanceTrendTelemetry.cs
@@ -211,11 +211,12 @@ namespace EpicGames.Analytics.PerformanceTrends
 		/// <param name="templateId">The template id.</param>
 		/// <param name="jobId">The Horde job id.</param>
 		/// <param name="stepId">The Horde step id.</param>
+		/// <param name="commitIdOrdered">The ordered commitId.</param>
 		/// <param name="commitId">The commitId.</param>
 		/// <param name="summaryName">The summary name.</param>
 		/// <param name="collated">Whether this record represents pre-aggregation.</param>
 #pragma warning disable CA1054 // URI-like parameters should not be strings
-		protected PerformanceTrendTelemetry(Hashtable source, string eventName, int schemaVersion, string? sessionId, string? sessionLabel, bool isBuildMachine, string? hordeUrlStr, string? streamId, string? templateId, string? jobId, string? stepId, int? commitId, string summaryName, bool collated) : base(eventName, schemaVersion, sessionId, sessionLabel, isBuildMachine, hordeUrlStr, streamId, templateId, jobId, stepId, commitId)
+		protected PerformanceTrendTelemetry(Hashtable source, string eventName, int schemaVersion, string? sessionId, string? sessionLabel, bool isBuildMachine, string? hordeUrlStr, string? streamId, string? templateId, string? jobId, string? stepId, int? commitIdOrdered, string? commitId, string summaryName, bool collated) : base(eventName, schemaVersion, sessionId, sessionLabel, isBuildMachine, hordeUrlStr, streamId, templateId, jobId, stepId, commitIdOrdered, commitId)
 #pragma warning restore CA1054 // URI-like parameters should not be strings
 		{
 			if (source == null)
diff --git a/Engine/Source/Programs/Shared/EpicGames.Horde/Telemetry/HordeContextTelemetryRecord.cs b/Engine/Source/Programs/Shared/EpicGames.Horde/Telemetry/HordeContextTelemetryRecord.cs
--- a/Engine/Source/Programs/Shared/EpicGames.Horde/Telemetry/HordeContextTelemetryRecord.cs
+++ b/Engine/Source/Programs/Shared/EpicGames.Horde/Telemetry/HordeContextTelemetryRecord.cs
@@ -57,10 +57,16 @@ namespace EpicGames.Horde.Telemetry
 		public JobStepId? StepId { get; init; }
 
 		/// <summary>
-		/// The change id.
+		/// The ordered commit id.
+		/// </summary>
+		[Column("commit_id_ordered")]
+		public int? CommitIdOrdered { get; init; }
+
+		/// <summary>
+		/// The unordered commit id.
 		/// </summary>
 		[Column("commit_id")]
-		public int? CommitId { get; init; }
+		public string? CommitId { get; init; }
 
 		#endregion -- Public Properties --
 
@@ -85,8 +91,9 @@ namespace EpicGames.Horde.Telemetry
 		/// <param name="templateId">The template id.</param>
 		/// <param name="jobId">The Horde job id.</param>
 		/// <param name="stepId">The Horde step id.</param>
+		/// <param name="commitIdOrdered">The ordered commitId.</param>
 		/// <param name="commitId">The commitId.</param>
-		public HordeContextTelemetryRecord(string eventName, int schemaVersion, string? sessionId, string? sessionLabel, bool isBuildMachine, Uri? hordeUrlStr, StreamId? streamId, TemplateId? templateId, JobId? jobId, JobStepId? stepId, int? commitId) : base(eventName, schemaVersion, sessionId, sessionLabel)
+		public HordeContextTelemetryRecord(string eventName, int schemaVersion, string? sessionId, string? sessionLabel, bool isBuildMachine, Uri? hordeUrlStr, StreamId? streamId, TemplateId? templateId, JobId? jobId, JobStepId? stepId, int? commitIdOrdered, string? commitId) : base(eventName, schemaVersion, sessionId, sessionLabel)
 		{
 			IsBuildMachine = isBuildMachine;
 			HordeUrlStr = hordeUrlStr?.ToString();
@@ -94,6 +101,7 @@ namespace EpicGames.Horde.Telemetry
 			JobId = jobId;
 			TemplateId = templateId;
 			StepId = stepId;
+			CommitIdOrdered = commitIdOrdered;
 			CommitId = commitId;
 		}
 
@@ -110,9 +118,10 @@ namespace EpicGames.Horde.Telemetry
 		/// <param name="templateId">The template id.</param>
 		/// <param name="jobId">The Horde job id.</param>
 		/// <param name="stepId">The Horde step id.</param>
+		/// <param name="commitIdOrdered">The ordered commitId.</param>
 		/// <param name="commitId">The commitId.</param>
 #pragma warning disable CA1054 // URI-like parameters should not be strings
-		public HordeContextTelemetryRecord(string eventName, int schemaVersion, string? sessionId, string? sessionLabel, bool isBuildMachine, string? hordeUrlStr, string? streamId, string? templateId, string? jobId, string? stepId, int? commitId) : base(eventName, schemaVersion, sessionId, sessionLabel)
+		public HordeContextTelemetryRecord(string eventName, int schemaVersion, string? sessionId, string? sessionLabel, bool isBuildMachine, string? hordeUrlStr, string? streamId, string? templateId, string? jobId, string? stepId, int? commitIdOrdered, string? commitId) : base(eventName, schemaVersion, sessionId, sessionLabel)
 #pragma warning restore CA1054 // URI-like parameters should not be strings
 		{
 			IsBuildMachine = isBuildMachine;
@@ -123,6 +132,7 @@ namespace EpicGames.Horde.Telemetry
 			TemplateId = TryCreate(templateId, t => new TemplateId(t));
 			StepId = TryCreate(stepId, s => JobStepId.Parse(s));
 
+			CommitIdOrdered = commitIdOrdered;
 			CommitId = commitId;
 		}
 
@@ -143,19 +153,21 @@ namespace EpicGames.Horde.Telemetry
 			string? jobId = HordeEnvVars.JobId;
 			string? templateId = HordeEnvVars.TemplateId;
 			string? stepId = HordeEnvVars.StepId;
+			string? commitId = Environment.GetEnvironmentVariable("uebp_CL");
 
 			StreamId = TryCreate(streamId, s => new StreamId(s));
-			JobId = TryCreate(jobId,  j =>Jobs.JobId.Parse(j));
-			TemplateId =TryCreate(templateId, t => new TemplateId(t));
-			StepId =TryCreate(stepId, s=> JobStepId.Parse(s));
+			JobId = TryCreate(jobId, j => Jobs.JobId.Parse(j));
+			TemplateId = TryCreate(templateId, t => new TemplateId(t));
+			StepId = TryCreate(stepId, s => JobStepId.Parse(s));
+			CommitId = commitId;
 
-			if(Int32.TryParse(Environment.GetEnvironmentVariable("uebp_CL"), out int commitIdRaw))
+			if (!String.IsNullOrEmpty(commitId) && Int32.TryParse(commitId, out int commitIdRaw))
 			{
-				CommitId = commitIdRaw;
+				CommitIdOrdered = commitIdRaw;
 			}
 			else
 			{
-				CommitId = 0;
+				CommitIdOrdered = -1;
 			}
 		}
 
@@ -196,21 +208,22 @@ namespace EpicGames.Horde.Telemetry
 			}
 
 			hashtable[nameof(CommitId)] = CommitId;
+			hashtable[nameof(CommitIdOrdered)] = CommitIdOrdered;
 		}
 
 		#endregion
 
 		#region -- Private Api --
 
-		private static T? TryCreate<T>(string? value, Func<string, T> creationDelegate) where T: struct
+		private static T? TryCreate<T>(string? value, Func<string, T> creationDelegate) where T : struct
 		{
-			if(String.IsNullOrEmpty(value))
+			if (String.IsNullOrEmpty(value))
 			{
 				return null;
 			}
 
-			try 
-			{ 
+			try
+			{
 				return creationDelegate(value);
 			}
 			catch
