[UBT]
* Fixed so c# code is using new GetProcessWeightThatCanRunRemotelyNow api from UBA

#Synced-CL 50116777

diff --git a/Engine/Source/Programs/Shared/EpicGames.UBA/IScheduler.cs b/Engine/Source/Programs/Shared/EpicGames.UBA/IScheduler.cs
--- a/Engine/Source/Programs/Shared/EpicGames.UBA/IScheduler.cs
+++ b/Engine/Source/Programs/Shared/EpicGames.UBA/IScheduler.cs
@@ -134,7 +134,7 @@ namespace EpicGames.UBA
 		/// <summary>
 		/// Accumulated weight of all processes that are Queued and can run right now (dependencies are done)
 		/// </summary>
-		double GetProcessWeightThatCanRunRemotelyNow();
+		void GetProcessWeightThatCanRunRemotelyNow(out double totalWeight, out double crossArchitecture, out double crossPlatform);
 
 		/// <summary>
 		/// Set callback for when process has finished
diff --git a/Engine/Source/Programs/Shared/EpicGames.UBA/Impl/SchedulerImpl.cs b/Engine/Source/Programs/Shared/EpicGames.UBA/Impl/SchedulerImpl.cs
--- a/Engine/Source/Programs/Shared/EpicGames.UBA/Impl/SchedulerImpl.cs
+++ b/Engine/Source/Programs/Shared/EpicGames.UBA/Impl/SchedulerImpl.cs
@@ -124,7 +124,7 @@ namespace EpicGames.UBA.Impl
 		static extern uint Scheduler_EnqueueProcess4(nint scheduler, nint info, float weight, byte canDetour, byte canExecuteRemotely, byte canCrossArchitecture, byte canCrossPlatform, int[]? dependencies, int dependencyCount, byte[]? knownInputs, int knownInputsBytes, uint knownInputsCount, uint cacheBucket, uint memoryGroup, ulong predictedMemoryUsage);
 
 		[DllImport("UbaHost", CharSet = CharSet.Auto)]
-		static extern float Scheduler_GetProcessWeightThatCanRunRemotelyNow(nint scheduler);
+		static extern float Scheduler_GetProcessWeightThatCanRunRemotelyNow2(nint scheduler, out float totalWeight, out float crossArchitecture, out float crossPlatform);
 
 		[DllImport("UbaHost", CharSet = CharSet.Auto)]
 		static extern float Scheduler_SetAllowDisableRemoteExecution(nint scheduler, byte allow);
@@ -185,7 +185,16 @@ namespace EpicGames.UBA.Impl
 
 		public bool IsEmpty => Scheduler_IsEmpty(_schedulerHandle) != 0;
 
-		public double GetProcessWeightThatCanRunRemotelyNow() => (double)Scheduler_GetProcessWeightThatCanRunRemotelyNow(_schedulerHandle);
+		public void GetProcessWeightThatCanRunRemotelyNow(out double totalWeight, out double crossArchitecture, out double crossPlatform)
+		{
+			float totalWeight_;
+			float crossArchitecture_;
+			float crossPlatform_;
+			Scheduler_GetProcessWeightThatCanRunRemotelyNow2(_schedulerHandle, out totalWeight_, out crossArchitecture_, out crossPlatform_);
+			totalWeight = totalWeight_;
+			crossArchitecture = crossArchitecture_;
+			crossPlatform = crossPlatform_;
+		}
 
 		public nint GetHandle() => _schedulerHandle;
 
diff --git a/Engine/Source/Programs/UnrealBuildTool/Executors/UnrealBuildAccelerator/UBAAgentCoordinator.cs b/Engine/Source/Programs/UnrealBuildTool/Executors/UnrealBuildAccelerator/UBAAgentCoordinator.cs
--- a/Engine/Source/Programs/UnrealBuildTool/Executors/UnrealBuildAccelerator/UBAAgentCoordinator.cs
+++ b/Engine/Source/Programs/UnrealBuildTool/Executors/UnrealBuildAccelerator/UBAAgentCoordinator.cs
@@ -14,7 +14,7 @@ namespace UnrealBuildTool
 	interface IUbaAgentCoordinatorScheduler
 	{
 		bool IsEmpty { get; }
-		double GetProcessWeightThatCanRunRemotelyNow();
+		void GetProcessWeightThatCanRunRemotelyNow(out double totalWeight, out double crossArchitecture, out double crossPlatform);
 		bool AddClient(string ip, int port, string crypto = "");
 	}
 
diff --git a/Engine/Source/Programs/UnrealBuildTool/Executors/UnrealBuildAccelerator/UBAAgentCoordinatorHorde.cs b/Engine/Source/Programs/UnrealBuildTool/Executors/UnrealBuildAccelerator/UBAAgentCoordinatorHorde.cs
--- a/Engine/Source/Programs/UnrealBuildTool/Executors/UnrealBuildAccelerator/UBAAgentCoordinatorHorde.cs
+++ b/Engine/Source/Programs/UnrealBuildTool/Executors/UnrealBuildAccelerator/UBAAgentCoordinatorHorde.cs
@@ -978,7 +978,11 @@ namespace UnrealBuildTool
 
 						hordeSession.RemoveCompleteWorkersAsync();
 
-						double queueWeight = scheduler.GetProcessWeightThatCanRunRemotelyNow();
+						double totalWeight;
+						double crossArchitecture;
+						double crossPlatform;
+						scheduler.GetProcessWeightThatCanRunRemotelyNow(out totalWeight, out crossArchitecture, out crossPlatform);
+						double queueWeight = totalWeight;
 
 						int activeCores;
 						int queuedCores;
diff --git a/Engine/Source/Programs/UnrealBuildTool/Executors/UnrealBuildAccelerator/UBAExecutor.cs b/Engine/Source/Programs/UnrealBuildTool/Executors/UnrealBuildAccelerator/UBAExecutor.cs
--- a/Engine/Source/Programs/UnrealBuildTool/Executors/UnrealBuildAccelerator/UBAExecutor.cs
+++ b/Engine/Source/Programs/UnrealBuildTool/Executors/UnrealBuildAccelerator/UBAExecutor.cs
@@ -798,11 +798,17 @@ namespace UnrealBuildTool
 				}
 			}
 
-			public double GetProcessWeightThatCanRunRemotelyNow()
+			public void GetProcessWeightThatCanRunRemotelyNow(out double totalWeight, out double crossArchitecture, out double crossPlatform)
 			{
+				totalWeight = 0;
+				crossArchitecture = 0;
+				crossPlatform = 0;
 				lock (_lockObject)
 				{
-					return _scheduler != null ? _scheduler.GetProcessWeightThatCanRunRemotelyNow() : 0;
+					if (_scheduler != null)
+					{
+						_scheduler.GetProcessWeightThatCanRunRemotelyNow(out totalWeight, out crossArchitecture, out crossPlatform);
+					}
 				}
 			}
 
