[Horde][PerformanceTrends][Server] Updated perf record for use in new API endpoints with more context.

diff --git a/Engine/Source/Programs/Shared/EpicGames.Analytics.PerformanceTrends/PerformanceTrendTelemetry.cs b/Engine/Source/Programs/Shared/EpicGames.Analytics.PerformanceTrends/PerformanceTrendTelemetry.cs
--- a/Engine/Source/Programs/Shared/EpicGames.Analytics.PerformanceTrends/PerformanceTrendTelemetry.cs
+++ b/Engine/Source/Programs/Shared/EpicGames.Analytics.PerformanceTrends/PerformanceTrendTelemetry.cs
@@ -3,6 +3,7 @@
 using System;
 using System.Collections;
 using System.ComponentModel.DataAnnotations.Schema;
+using EpicGames.Analytics.Telemetry;
 using EpicGames.Horde.Telemetry;
 
 namespace EpicGames.Analytics.PerformanceTrends
@@ -23,6 +24,13 @@ namespace EpicGames.Analytics.PerformanceTrends
 		internal const string TestConfigName = "testconfigname";
 		internal const string GauntletTestType = "gauntlettesttype";
 		internal const string GauntletSubTest = "gauntletsubtest";
+		internal const string Branch = "branch";
+		internal const string Changelist = "changelist";
+		internal const string StartTimestamp = "starttimestamp";
+		internal const string EndTimestamp = "endtimestamp";
+		internal const string BuildSuffix = "buildsuffix";
+		internal const string TestBuildIsPreflight = "testbuildispreflight";
+
 
 		#endregion -- Internal Constants --
 
@@ -89,7 +97,13 @@ namespace EpicGames.Analytics.PerformanceTrends
 				TestName,
 				TestConfigName,
 				GauntletTestType,
-				GauntletSubTest
+				GauntletSubTest,
+				Branch,
+				Changelist,
+				StartTimestamp,
+				EndTimestamp,
+				BuildSuffix,
+				TestBuildIsPreflight
 		};
 
 		/// <summary>
@@ -119,6 +133,9 @@ namespace EpicGames.Analytics.PerformanceTrends
 	/// <summary>
 	/// Abstract record that describes the shared metadata for performance trend telemetry.
 	/// </summary>
+	/// <remarks>Whilst this table specifies <see cref="AnalyticsTableGenAttribute"/>, it does so for a narrow use in shared query context, and not against any specific <see cref="PerformanceTrendTelemetry"/> table.</remarks>
+	[AnalyticsTableGen]
+	[Table("default")]
 	public abstract record PerformanceTrendTelemetry : HordeContextTelemetryRecord
 	{
 		#region -- Public Members --
@@ -183,11 +200,57 @@ namespace EpicGames.Analytics.PerformanceTrends
 		[Column("collated")]
 		public bool Collated { get; init; }
 
+		/// <summary>
+		/// The branch the underyling test build was produced from.
+		/// </summary>
+		[Column("branch")]
+		public string? Branch { get; init; }
+
+		/// <summary>
+		/// The changelist the underyling test build was produced from.
+		/// </summary>
+		[Column("changelist")]
+		public int? Changelist { get; init; }
+
+		/// <summary>
+		/// The start timestamp of the performance trend session.
+		/// </summary>
+		/// <remarks>Time since epoch.</remarks>
+		[Column("starttimestamp")]
+		public long? StartTimestamp { get; init; }
+
+		/// <summary>
+		/// The end timestamp of the performance trend session.
+		/// </summary>
+		/// <remarks>Time since epoch.</remarks>
+		[Column("endtimestamp")]
+		public long? EndTimestamp { get; init; }
+
+		/// <summary>
+		/// The build suffix for the build string.
+		/// </summary>
+		/// <remarks>This typically refers to the test product name.</remarks>
+		[Column("buildsuffix")]
+		public string? BuildSuffix { get; init; }
+
+
+		/// <summary>
+		/// Int boolean representation of whether the test build was generated by preflight or not.
+		/// </summary>
+		[Column("testbuildispreflight")]
+		public int TestBuildIsPreflight { get; init; }
+
+		/// <summary>
+		/// The data source of the performance trend telemetry.
+		/// </summary>
+		[Column("data_source")]
+		public string? DataSource { get; init; }
+
 		/// <summary>
 		/// Summary table name.
 		/// </summary>
 		[Column("summary_name")]
-		public string? SummaryName { get; init; }
+		public string SummaryName { get; init; }
 
 		#endregion -- Public Members --
 
@@ -195,7 +258,10 @@ namespace EpicGames.Analytics.PerformanceTrends
 		/// No arg constructor for ORM construction.
 		/// </summary>
 		/// <remarks>Used for ORM instantiation.</remarks>
-		protected PerformanceTrendTelemetry() : base() { }
+		protected PerformanceTrendTelemetry() : base()
+		{
+			SummaryName = String.Empty;
+		}
 
 		/// <summary>
 		/// Constructor.
@@ -215,14 +281,15 @@ namespace EpicGames.Analytics.PerformanceTrends
 		/// <param name="commitId">The commitId.</param>
 		/// <param name="summaryName">The summary name.</param>
 		/// <param name="collated">Whether this record represents pre-aggregation.</param>
+		/// <param name="dataSource">The context in which the performance trend was generated.</param>
 #pragma warning disable CA1054 // URI-like parameters should not be strings
-		protected PerformanceTrendTelemetry(Hashtable source, string eventName, int schemaVersion, string? sessionId, string? sessionLabel, bool isBuildMachine, string? hordeUrlStr, string? streamId, string? templateId, string? jobId, string? stepId, int? commitIdOrdered, string? commitId, string summaryName, bool collated) : base(eventName, schemaVersion, sessionId, sessionLabel, isBuildMachine, hordeUrlStr, streamId, templateId, jobId, stepId, commitIdOrdered, commitId)
+		protected PerformanceTrendTelemetry(Hashtable source, string eventName, int schemaVersion, string? sessionId, string? sessionLabel, bool isBuildMachine, string? hordeUrlStr, string? streamId, string? templateId, string? jobId, string? stepId, int? commitIdOrdered, string? commitId, string summaryName, bool collated, string? dataSource)
+			: base(eventName, schemaVersion, sessionId, sessionLabel, isBuildMachine, hordeUrlStr, streamId, templateId, jobId, stepId, commitIdOrdered, commitId)
 #pragma warning restore CA1054 // URI-like parameters should not be strings
 		{
-			if (source == null)
-			{
-				return;
-			}
+			SummaryName = summaryName;
+			Collated = collated;
+			DataSource = dataSource;
 
 			EngineVersion = GetString(source, StandardTelemetryMetadata.EngineVersion);
 			EngineReleaseVersion = GetString(source, StandardTelemetryMetadata.EngineReleaseVersion);
@@ -233,8 +300,12 @@ namespace EpicGames.Analytics.PerformanceTrends
 			TestConfigName = GetString(source, StandardTelemetryMetadata.TestConfigName);
 			GauntletTestType = GetString(source, StandardTelemetryMetadata.GauntletTestType);
 			GauntletSubTest = GetString(source, StandardTelemetryMetadata.GauntletSubTest);
-			SummaryName = summaryName;
-			Collated = collated;
+			Branch = GetString(source, StandardTelemetryMetadata.Branch);
+			Changelist = GetInt(source, StandardTelemetryMetadata.Changelist);
+			StartTimestamp = GetLong(source, StandardTelemetryMetadata.StartTimestamp);
+			EndTimestamp = GetLong(source, StandardTelemetryMetadata.EndTimestamp);
+			BuildSuffix = GetString(source, StandardTelemetryMetadata.BuildSuffix);
+			TestBuildIsPreflight = GetInt(source, StandardTelemetryMetadata.TestBuildIsPreflight) ?? 0;
 		}
 
 		#region -- Private API --
@@ -249,6 +320,54 @@ namespace EpicGames.Analytics.PerformanceTrends
 			return table[key] as string;
 		}
 
+		private static int? GetInt(Hashtable table, string key)
+		{
+			if (!table.Contains(key))
+			{
+				return null;
+			}
+
+			try
+			{
+				object? value = table[key];
+
+				if (value == null)
+				{
+					return null;
+				}
+
+				return Convert.ToInt32(value);
+			}
+			catch
+			{
+				return null;
+			}
+		}
+
+		private static long? GetLong(Hashtable table, string key)
+		{
+			if (!table.Contains(key))
+			{
+				return null;
+			}
+
+			try
+			{
+				object? value = table[key];
+
+				if (value == null)
+				{
+					return null;
+				}
+
+				return Convert.ToInt64(value);
+			}
+			catch
+			{
+				return null;
+			}
+		}
+
 		#endregion -- Private API --
 	}
 }
\ No newline at end of file
