Fix issue where UHT wasn't generating the verse aliases correctly.  This required transform, rotation, and vector3 to be marked as no alias since they already have aliases for the old types.

#rnx

#Synced-CL 50159186

diff --git a/Engine/Source/Programs/Shared/EpicGames.UHT/Exporters/CodeGen/UhtHeaderCodeGeneratorHFile.cs b/Engine/Source/Programs/Shared/EpicGames.UHT/Exporters/CodeGen/UhtHeaderCodeGeneratorHFile.cs
--- a/Engine/Source/Programs/Shared/EpicGames.UHT/Exporters/CodeGen/UhtHeaderCodeGeneratorHFile.cs
+++ b/Engine/Source/Programs/Shared/EpicGames.UHT/Exporters/CodeGen/UhtHeaderCodeGeneratorHFile.cs
@@ -710,7 +710,7 @@ namespace EpicGames.UHT.Exporters.CodeGen
 				{
 					if (!scriptStruct.FieldExportFlags.HasAnyFlags(UhtFieldExportFlags.NoVerseAlias))
 					{
-						builder.Append("namespace verse { using ").Append(scriptStruct.EngineName).Append(" = ").Append(scriptStruct.SourceName).Append("; }\r\n");
+						builder.Append("namespace verse { using ").Append(scriptStruct.VerseName).Append(" = ").Append(scriptStruct.SourceName).Append("; }\r\n");
 					}
 					builder
 						.Append("V_DECLARE_STRUCT_RUNTIME_TRAITS(, ")
@@ -972,7 +972,7 @@ namespace EpicGames.UHT.Exporters.CodeGen
 				{
 					if (classObj.ClassType == UhtClassType.Class)
 					{
-						builder.Append("namespace verse { using ").Append(classObj.EngineName).Append(" = ").AppendClassSourceNameOrInterfaceProxyName(classObj).Append("; }\r\n");
+						builder.Append("namespace verse { using ").Append(classObj.VerseName).Append(" = ").AppendClassSourceNameOrInterfaceProxyName(classObj).Append("; }\r\n");
 					}
 				}
 			}
@@ -1011,7 +1011,7 @@ namespace EpicGames.UHT.Exporters.CodeGen
 
 			if (!classObj.FieldExportFlags.HasAnyFlags(UhtFieldExportFlags.NoVerseAlias))
 			{
-				builder.Append("namespace verse { using ").Append(classObj.EngineName).Append(" = ").AppendClassSourceNameOrInterfaceProxyName(classObj).Append("; }\r\n");
+				builder.Append("namespace verse { using ").Append(classObj.VerseName).Append(" = ").AppendClassSourceNameOrInterfaceProxyName(classObj).Append("; }\r\n");
 			}
 			return builder;
 		}
