Horde: Add more detailed process information to file lock diagnostics

Enhance FileLockInfoWin32 to capture and display command-line arguments and parent process information when diagnosing file lock issues. This provides better context when troubleshooting which process is holding a file lock, including how the process was started and what spawned it.

diff --git a/Engine/Source/Programs/Shared/EpicGames.Core.Tests/FileUtilsTests.cs b/Engine/Source/Programs/Shared/EpicGames.Core.Tests/FileUtilsTests.cs
--- a/Engine/Source/Programs/Shared/EpicGames.Core.Tests/FileUtilsTests.cs
+++ b/Engine/Source/Programs/Shared/EpicGames.Core.Tests/FileUtilsTests.cs
@@ -1,6 +1,7 @@
 // Copyright Epic Games, Inc. All Rights Reserved.
 
 using System;
+using System.Collections.Generic;
 using System.IO;
 using Microsoft.VisualStudio.TestTools.UnitTesting;
 
@@ -83,5 +84,45 @@ namespace EpicGames.Core.Tests
 			Directory.CreateDirectory(tempDir);
 			return new DirectoryReference(tempDir);
 		}
+
+		[TestMethod]
+		public void GetFileLockInfoWin32_ReturnsCommandLineAndParentInfo()
+		{
+			if (!OperatingSystem.IsWindows())
+			{
+				Assert.Inconclusive("This test only runs on Windows");
+				return;
+			}
+
+			string tempFile = Path.Combine(_tempDir.FullName, "locked_file.txt");
+			File.WriteAllText(tempFile, "test content");
+
+			// Create a file and lock it from this process and verify we can detect the lock
+			using FileStream lockedStream = new (tempFile, FileMode.Open, FileAccess.Read, FileShare.None);
+			List<FileLockInfoWin32> lockInfoList = FileUtils.GetFileLockInfoWin32(tempFile);
+			Assert.IsNotEmpty(lockInfoList);
+
+			// Find our process (the current test runner)
+			int currentPid = Environment.ProcessId;
+			FileLockInfoWin32? ourLock = lockInfoList.Find(l => l.ProcessId == currentPid);
+			Assert.IsNotNull(ourLock);
+
+			// Verify command line is present (should contain dotnet or testhost)
+			Assert.IsNotNull(ourLock.CommandLine);
+			Assert.IsGreaterThan(0, ourLock.CommandLine.Length);
+
+			// Verify parent process info is present
+			Assert.IsGreaterThan(0, ourLock.ParentProcessId);
+
+			// Verify ToString() includes all the info in a readable format
+			string lockString = ourLock.ToString();
+			Assert.IsTrue(lockString.Contains(ourLock.ProcessId.ToString(), StringComparison.Ordinal));
+			Assert.IsTrue(lockString.Contains("Command line:", StringComparison.Ordinal));
+			Assert.IsTrue(lockString.Contains("Parent:", StringComparison.Ordinal));
+
+			// Check multi-line formatting
+			string[] lines = lockString.Split('\n');
+			Assert.IsGreaterThanOrEqualTo(2, lines.Length);
+		}
 	}
 }
diff --git a/Engine/Source/Programs/Shared/EpicGames.Core/FileUtils.cs b/Engine/Source/Programs/Shared/EpicGames.Core/FileUtils.cs
--- a/Engine/Source/Programs/Shared/EpicGames.Core/FileUtils.cs
+++ b/Engine/Source/Programs/Shared/EpicGames.Core/FileUtils.cs
@@ -6,6 +6,7 @@ using System.ComponentModel;
 using System.IO;
 using System.Runtime.InteropServices;
 using System.Runtime.InteropServices.ComTypes;
+using System.Runtime.Versioning;
 using System.Text;
 using Microsoft.Win32.SafeHandles;
 
@@ -32,7 +33,19 @@ namespace EpicGames.Core
 	/// <param name="fileName">Path to the process holding the lock</param>
 	/// <param name="appName">Name of the application</param>
 	/// <param name="startTime">Time at which the process started</param>
-	public class FileLockInfoWin32(int processId, string? fileName, string appName, DateTime startTime)
+	/// <param name="commandLine">Command-line arguments (may be null if unavailable)</param>
+	/// <param name="parentProcessId">Parent process id (0 if unavailable)</param>
+	/// <param name="parentFileName">Path to the parent process (may be null)</param>
+	/// <param name="parentCommandLine">Parent process command-line (may be null)</param>
+	public class FileLockInfoWin32(
+		int processId,
+		string? fileName,
+		string appName,
+		DateTime startTime,
+		string? commandLine,
+		int parentProcessId,
+		string? parentFileName,
+		string? parentCommandLine)
 	{
 		/// <summary>
 		/// Process id
@@ -54,8 +67,58 @@ namespace EpicGames.Core
 		/// </summary>
 		public DateTime StartTime => startTime;
 
+		/// <summary>
+		/// Command-line arguments used to start the process (may be null if unavailable)
+		/// </summary>
+		public string? CommandLine => commandLine;
+
+		/// <summary>
+		/// Parent process id (0 if unavailable or parent has exited)
+		/// </summary>
+		public int ParentProcessId => parentProcessId;
+
+		/// <summary>
+		/// Path to the parent process (may be null if parent has exited)
+		/// </summary>
+		public string? ParentFileName => parentFileName;
+
+		/// <summary>
+		/// Parent process command-line (may be null if parent has exited)
+		/// </summary>
+		public string? ParentCommandLine => parentCommandLine;
+
 		/// <inheritdoc/>
-		public override string ToString() => $"{ProcessId}: {FileName ?? AppName} (started {StartTime})";
+		public override string ToString()
+		{
+			StringBuilder sb = new ();
+
+			// Main process info
+			sb.Append($"{ProcessId}: {FileName ?? AppName} (started {StartTime})");
+
+			// Command line on next line, indented
+			if (!String.IsNullOrEmpty(CommandLine))
+			{
+				string cmdLine = CommandLine.Length > 200 ? CommandLine[..197] + "..." : CommandLine;
+				sb.Append($"\n    Command line: {cmdLine}");
+			}
+
+			// Parent process info on next line, indented
+			if (ParentProcessId > 0)
+			{
+				string parentName = ParentFileName != null ? Path.GetFileName(ParentFileName) : "(exited)";
+				sb.Append($"\n    Parent: {ParentProcessId} {parentName}");
+
+				if (!String.IsNullOrEmpty(ParentCommandLine))
+				{
+					string parentCmd = ParentCommandLine.Length > 100
+						? ParentCommandLine[..97] + "..."
+						: ParentCommandLine;
+					sb.Append($" [{parentCmd}]");
+				}
+			}
+
+			return sb.ToString();
+		}
 	}
 
 	/// <summary>
@@ -584,6 +647,7 @@ namespace EpicGames.Core
 		const int ERROR_ACCESS_DENIED = 5;
 #pragma warning restore IDE1006 // Naming Styles
 
+		[SupportedOSPlatform("windows")]
 		private static void ForceDeleteLongDirectoryContentsWin32(string dirName)
 		{
 			WIN32_FIND_DATA findData = new WIN32_FIND_DATA();
@@ -635,6 +699,7 @@ namespace EpicGames.Core
 			}
 		}
 
+		[SupportedOSPlatform("windows")]
 		private static void ForceDeleteLongDirectoryWin32(string dirName)
 		{
 			ForceDeleteLongDirectoryContentsWin32(dirName);
@@ -656,6 +721,7 @@ namespace EpicGames.Core
 		/// Force-delete a file (Windows only)
 		/// </summary>
 		/// <param name="fileName"></param>
+		[SupportedOSPlatform("windows")]
 		public static void ForceDeleteFileWin32(string fileName)
 		{
 			// Try to delete the file normally
@@ -692,23 +758,11 @@ namespace EpicGames.Core
 			StringBuilder message = new StringBuilder($"Unable to delete {fileName} - {ex.Message}");
 			if (ex.NativeErrorCode == ERROR_ACCESS_DENIED)
 			{
-				List<FileLockInfoWin32>? lockInfoList;
-				try
-				{
-					lockInfoList = GetFileLockInfoWin32(fileName);
-				}
-				catch
+				string? lockInfo = GetFileLockInfo(fileName);
+				if (lockInfo != null)
 				{
-					lockInfoList = null;
-				}
-
-				if (lockInfoList != null && lockInfoList.Count > 0)
-				{
-					message.Append("\nProcesses with open handles to file:");
-					foreach (FileLockInfoWin32 lockInfo in lockInfoList)
-					{
-						message.Append($"\n  {lockInfo}");
-					}
+					message.Append('\n');
+					message.Append(lockInfo);
 				}
 			}
 			throw new WrappedFileOrDirectoryException(ex, message.ToString());
@@ -778,8 +832,6 @@ namespace EpicGames.Core
 									[In, Out] RM_PROCESS_INFO[] rgAffectedApps,
 									ref uint lpdwRebootReasons);
 
-		const int PROCESS_QUERY_LIMITED_INFORMATION = 0x1000;
-
 		[DllImport("kernel32.dll", SetLastError = true)]
 		static extern SafeProcessHandle OpenProcess(
 			 int processAccess,
@@ -792,11 +844,6 @@ namespace EpicGames.Core
 		static extern bool GetProcessTimes(SafeProcessHandle hProcess, out FILETIME
 		   lpCreationTime, out FILETIME lpExitTime, out FILETIME lpKernelTime,
 		   out FILETIME lpUserTime);
-
-		[DllImport("kernel32.dll", SetLastError = true)]
-#pragma warning disable CA1838 // Avoid 'StringBuilder' parameters for P/Invokes
-		static extern bool QueryFullProcessImageName([In] SafeProcessHandle hProcess, [In] int dwFlags, [Out] StringBuilder lpExeName, ref int lpdwSize);
-#pragma warning restore CA1838 // Avoid 'StringBuilder' parameters for P/Invokes
 #pragma warning restore IDE1006
 
 		/// <summary>
@@ -804,6 +851,7 @@ namespace EpicGames.Core
 		/// </summary>
 		/// <param name="fileName">File to check</param>
 		/// <returns>List of processes with a lock open</returns>
+		[SupportedOSPlatform("windows")]
 		public static List<FileLockInfoWin32> GetFileLockInfoWin32(string fileName)
 		{
 			uint sessionHandle = 0;
@@ -840,24 +888,41 @@ namespace EpicGames.Core
 					long startTimeTicks = FileTimeToTicks(processInfo.Process.ProcessStartTime);
 
 					string? imageName = null;
-					using (SafeProcessHandle hProcess = OpenProcess(PROCESS_QUERY_LIMITED_INFORMATION, false, processInfo.Process.dwProcessId))
+					string? commandLine = null;
+					int parentPid = 0;
+					string? parentFileName = null;
+					string? parentCommandLine = null;
+
+					using (SafeProcessHandle hProcess = OpenProcess(ProcessUtils.PROCESS_QUERY_LIMITED_INFORMATION | ProcessUtils.PROCESS_VM_READ, false, processInfo.Process.dwProcessId))
 					{
-						if (hProcess != null)
+						if (hProcess is { IsInvalid: false })
 						{
 							FILETIME createTime, exitTime, kernelTime, userTime;
 							if (GetProcessTimes(hProcess, out createTime, out exitTime, out kernelTime, out userTime) && FileTimeToTicks(createTime) == startTimeTicks)
 							{
-								int capacity = 260;
-								StringBuilder imageNameBuilder = new StringBuilder(capacity);
-								if (QueryFullProcessImageName(hProcess, 0, imageNameBuilder, ref capacity))
+								imageName = ProcessUtils.GetProcessImageName(hProcess);
+								commandLine = ProcessUtils.GetProcessCommandLine(hProcess);
+								parentPid = ProcessUtils.GetParentProcessId(hProcess);
+
+								// Get parent process info if parent is still running
+								if (parentPid > 0)
 								{
-									imageName = imageNameBuilder.ToString(0, capacity);
+									parentFileName = ProcessUtils.GetProcessImageName(parentPid);
+									parentCommandLine = ProcessUtils.GetProcessCommandLine(parentPid);
 								}
 							}
 						}
 					}
 
-					fileLocks.Add(new FileLockInfoWin32(processInfo.Process.dwProcessId, imageName, processInfo.strAppName, DateTime.FromFileTime(startTimeTicks)));
+					fileLocks.Add(new FileLockInfoWin32(
+						processInfo.Process.dwProcessId,
+						imageName,
+						processInfo.strAppName,
+						DateTime.FromFileTime(startTimeTicks),
+						commandLine,
+						parentPid,
+						parentFileName,
+						parentCommandLine));
 				}
 				return fileLocks;
 			}
diff --git a/Engine/Source/Programs/Shared/EpicGames.Core/ProcessUtils.cs b/Engine/Source/Programs/Shared/EpicGames.Core/ProcessUtils.cs
--- a/Engine/Source/Programs/Shared/EpicGames.Core/ProcessUtils.cs
+++ b/Engine/Source/Programs/Shared/EpicGames.Core/ProcessUtils.cs
@@ -19,7 +19,8 @@ namespace EpicGames.Core
 	public static class ProcessUtils
 	{
 		const uint PROCESS_TERMINATE = 0x0001;
-		const uint PROCESS_QUERY_LIMITED_INFORMATION = 0x1000;
+		internal const int PROCESS_VM_READ = 0x0010;
+		internal const int PROCESS_QUERY_LIMITED_INFORMATION = 0x1000;
 
 		[DllImport("kernel32.dll", SetLastError = true)]
 		static extern SafeProcessHandle OpenProcess(uint dwDesiredAccess, bool bInheritHandle, uint processId);
@@ -53,9 +54,35 @@ namespace EpicGames.Core
 			public IntPtr _inheritedFromUniqueProcessId;
 		}
 
+		[StructLayout(LayoutKind.Sequential)]
+		struct UnicodeString
+		{
+			public ushort _length;
+			public ushort _maximumLength;
+			public IntPtr _buffer;
+		}
+
+		const int ProcessBasicInformationClass = 0;
+
 		[DllImport("ntdll.dll")]
 		private static extern int NtQueryInformationProcess(IntPtr processHandle, int processInformationClass, ref ProcessBasicInformation processInformation, int processInformationLength, out int returnLength);
 
+		[DllImport("kernel32.dll", SetLastError = true)]
+		static extern bool ReadProcessMemory(
+			IntPtr hProcess,
+			IntPtr lpBaseAddress,
+			byte[] lpBuffer,
+			int dwSize,
+			out int lpNumberOfBytesRead);
+
+		[DllImport("kernel32.dll", SetLastError = true)]
+		static extern bool ReadProcessMemory(
+			IntPtr hProcess,
+			IntPtr lpBaseAddress,
+			ref IntPtr lpBuffer,
+			int dwSize,
+			out int lpNumberOfBytesRead);
+
 		/// <summary>
 		/// Attempts to terminate all processes matching a predicate
 		/// </summary>
@@ -317,5 +344,221 @@ namespace EpicGames.Core
 				parentProcess = GetParentProcess(parentProcess);
 			}
 		}
+
+		/// <summary>
+		/// Gets the command line for a process by reading its PEB (Windows only)
+		/// </summary>
+		/// <param name="hProcess">Handle to the process (must have PROCESS_QUERY_LIMITED_INFORMATION and PROCESS_VM_READ access)</param>
+		/// <returns>The command line string, or null if it could not be retrieved</returns>
+		[SupportedOSPlatform("windows")]
+		public static string? GetProcessCommandLine(SafeProcessHandle hProcess)
+		{
+			if (hProcess == null || hProcess.IsInvalid)
+			{
+				return null;
+			}
+
+			if (!GetProcessBasicInfo(hProcess, out IntPtr pebBaseAddress, out _))
+			{
+				return null;
+			}
+
+			if (pebBaseAddress == IntPtr.Zero)
+			{
+				return null;
+			}
+
+			return GetProcessCommandLineFromPeb(hProcess.DangerousGetHandle(), pebBaseAddress);
+		}
+
+		/// <summary>
+		/// Gets the command line for a process by process ID (Windows only)
+		/// </summary>
+		/// <param name="processId">The process ID</param>
+		/// <returns>The command line string, or null if it could not be retrieved</returns>
+		[SupportedOSPlatform("windows")]
+		public static string? GetProcessCommandLine(int processId)
+		{
+			try
+			{
+				using SafeProcessHandle hProcess = OpenProcess(PROCESS_QUERY_LIMITED_INFORMATION | PROCESS_VM_READ, false, (uint)processId);
+				return GetProcessCommandLine(hProcess);
+			}
+			catch
+			{
+				return null;
+			}
+		}
+
+		/// <summary>
+		/// Gets the full image path for a process (Windows only)
+		/// </summary>
+		/// <param name="hProcess">Handle to the process (must have PROCESS_QUERY_LIMITED_INFORMATION access)</param>
+		/// <returns>The full path to the process executable, or null if it could not be retrieved</returns>
+		[SupportedOSPlatform("windows")]
+		public static string? GetProcessImageName(SafeProcessHandle hProcess)
+		{
+			if (hProcess == null || hProcess.IsInvalid)
+			{
+				return null;
+			}
+
+			int capacity = 260;
+			StringBuilder imageNameBuilder = new StringBuilder(capacity);
+			if (QueryFullProcessImageName(hProcess, 0, imageNameBuilder, ref capacity) > 0)
+			{
+				return imageNameBuilder.ToString(0, capacity);
+			}
+			return null;
+		}
+
+		/// <summary>
+		/// Gets the full image path for a process by process ID (Windows only)
+		/// </summary>
+		/// <param name="processId">The process ID</param>
+		/// <returns>The full path to the process executable, or null if it could not be retrieved</returns>
+		[SupportedOSPlatform("windows")]
+		public static string? GetProcessImageName(int processId)
+		{
+			try
+			{
+				using SafeProcessHandle hProcess = OpenProcess(PROCESS_QUERY_LIMITED_INFORMATION, false, (uint)processId);
+				return GetProcessImageName(hProcess);
+			}
+			catch
+			{
+				return null;
+			}
+		}
+
+		/// <summary>
+		/// Gets the parent process ID for a process (Windows only)
+		/// </summary>
+		/// <param name="hProcess">Handle to the process (must have PROCESS_QUERY_LIMITED_INFORMATION access)</param>
+		/// <returns>The parent process ID, or 0 if it could not be retrieved</returns>
+		[SupportedOSPlatform("windows")]
+		public static int GetParentProcessId(SafeProcessHandle hProcess)
+		{
+			if (hProcess == null || hProcess.IsInvalid)
+			{
+				return 0;
+			}
+
+			if (GetProcessBasicInfo(hProcess, out _, out int parentPid))
+			{
+				return parentPid;
+			}
+			return 0;
+		}
+
+		/// <summary>
+		/// Gets basic process information including PEB address and parent PID
+		/// </summary>
+		[SupportedOSPlatform("windows")]
+		private static bool GetProcessBasicInfo(SafeProcessHandle hProcess, out IntPtr pebBaseAddress, out int parentPid)
+		{
+			pebBaseAddress = IntPtr.Zero;
+			parentPid = 0;
+
+			try
+			{
+				ProcessBasicInformation pbi = new();
+				int status = NtQueryInformationProcess(
+					hProcess.DangerousGetHandle(),
+					ProcessBasicInformationClass,
+					ref pbi,
+					Marshal.SizeOf(pbi),
+					out _);
+
+				if (status != 0)
+				{
+					return false;
+				}
+
+				pebBaseAddress = pbi._pebBaseAddress;
+				parentPid = pbi._inheritedFromUniqueProcessId.ToInt32();
+				return true;
+			}
+			catch
+			{
+				return false;
+			}
+		}
+
+		/// <summary>
+		/// Gets the command line for a process by reading its PEB
+		/// </summary>
+		[SupportedOSPlatform("windows")]
+		private static string? GetProcessCommandLineFromPeb(IntPtr hProcess, IntPtr pebBaseAddress)
+		{
+			try
+			{
+				// PEB layout differs by architecture:
+				// 64-bit: ProcessParameters at offset 0x20
+				// 32-bit: ProcessParameters at offset 0x10
+				int processParametersOffset = IntPtr.Size == 8 ? 0x20 : 0x10;
+
+				IntPtr processParametersPtr = IntPtr.Zero;
+				if (!ReadProcessMemory(
+					hProcess,
+					pebBaseAddress + processParametersOffset,
+					ref processParametersPtr,
+					IntPtr.Size,
+					out _))
+				{
+					return null;
+				}
+
+				// RTL_USER_PROCESS_PARAMETERS layout:
+				// 64-bit: CommandLine UNICODE_STRING at offset 0x70
+				// 32-bit: CommandLine UNICODE_STRING at offset 0x40
+				int commandLineOffset = IntPtr.Size == 8 ? 0x70 : 0x40;
+
+				// Read the UNICODE_STRING structure
+				byte[] unicodeStringBytes = new byte[Marshal.SizeOf<UnicodeString>()];
+				if (!ReadProcessMemory(
+					hProcess,
+					processParametersPtr + commandLineOffset,
+					unicodeStringBytes,
+					unicodeStringBytes.Length,
+					out _))
+				{
+					return null;
+				}
+
+				GCHandle handle = GCHandle.Alloc(unicodeStringBytes, GCHandleType.Pinned);
+				try
+				{
+					UnicodeString unicodeString = Marshal.PtrToStructure<UnicodeString>(handle.AddrOfPinnedObject());
+
+					if (unicodeString._length == 0 || unicodeString._buffer == IntPtr.Zero)
+					{
+						return null;
+					}
+
+					// Read the actual command line string
+					byte[] commandLineBytes = new byte[unicodeString._length];
+					if (!ReadProcessMemory(
+						hProcess,
+						unicodeString._buffer,
+						commandLineBytes,
+						commandLineBytes.Length,
+						out _))
+					{
+						return null;
+					}
+
+					return Encoding.Unicode.GetString(commandLineBytes);
+				}
+				finally
+				{
+					handle.Free();
+				}
+			}
+			catch
+			{
+				return null;
+			}
+		}
 	}
 }
