EpicGames.Core: Fix file access flags in JsonWriter.cs

#rnx

#Synced-CL 50271719

diff --git a/Engine/Source/Programs/Shared/EpicGames.Core/JsonWriter.cs b/Engine/Source/Programs/Shared/EpicGames.Core/JsonWriter.cs
--- a/Engine/Source/Programs/Shared/EpicGames.Core/JsonWriter.cs
+++ b/Engine/Source/Programs/Shared/EpicGames.Core/JsonWriter.cs
@@ -41,7 +41,7 @@ namespace EpicGames.Core
 		/// <param name="fileName">File to write to</param>
 		/// <param name="style">Should use packed JSON or not</param>
 		public JsonWriter(string fileName, JsonWriterStyle style = JsonWriterStyle.Readable)
-			: this(new StreamWriter(fileName))
+			: this(new StreamWriter(new FileStream(fileName, FileMode.Create, FileAccess.Write, FileShare.Read)))
 		{
 			_style = style;
 		}
@@ -52,7 +52,7 @@ namespace EpicGames.Core
 		/// <param name="fileName">File to write to</param>
 		/// <param name="style">Should use packed JSON or not</param>
 		public JsonWriter(FileReference fileName, JsonWriterStyle style = JsonWriterStyle.Readable)
-			: this(new StreamWriter(fileName.FullName))
+			: this(new StreamWriter(new FileStream(fileName.FullName, FileMode.Create, FileAccess.Write, FileShare.Read)))
 		{
 			_style = style;
 		}
