Shared: csproj updates
Remove unneeded msbuild properties
* AllowUnsafeBlocks: Set in several projects that do not need it
* Several other properties that are set in the directory property sheet
Remove no longer needed restricted include in EpicGames.Core

#rnx

#Synced-CL 50360693

diff --git a/Engine/Source/Programs/Shared/EpicGames.Analytics.PerformanceTrends/EpicGames.Analytics.PerformanceTrends.csproj b/Engine/Source/Programs/Shared/EpicGames.Analytics.PerformanceTrends/EpicGames.Analytics.PerformanceTrends.csproj
--- a/Engine/Source/Programs/Shared/EpicGames.Analytics.PerformanceTrends/EpicGames.Analytics.PerformanceTrends.csproj
+++ b/Engine/Source/Programs/Shared/EpicGames.Analytics.PerformanceTrends/EpicGames.Analytics.PerformanceTrends.csproj
@@ -4,7 +4,6 @@
     <Deterministic>true</Deterministic>
     <RootNamespace>EpicGames.Analytics</RootNamespace>
     <GenerateDocumentationFile>True</GenerateDocumentationFile>
-    <AllowUnsafeBlocks>true</AllowUnsafeBlocks>
   </PropertyGroup>
 
   <PropertyGroup Condition="'$(Configuration)|$(Platform)' == 'Development|AnyCPU'">
diff --git a/Engine/Source/Programs/Shared/EpicGames.Analytics.Telemetry/EpicGames.Analytics.Telemetry.csproj b/Engine/Source/Programs/Shared/EpicGames.Analytics.Telemetry/EpicGames.Analytics.Telemetry.csproj
--- a/Engine/Source/Programs/Shared/EpicGames.Analytics.Telemetry/EpicGames.Analytics.Telemetry.csproj
+++ b/Engine/Source/Programs/Shared/EpicGames.Analytics.Telemetry/EpicGames.Analytics.Telemetry.csproj
@@ -2,7 +2,6 @@
   <PropertyGroup>
     <TargetFramework>net8.0</TargetFramework>
     <RootNamespace>EpicGames.Analytics</RootNamespace>
-    <AllowUnsafeBlocks>true</AllowUnsafeBlocks>
   </PropertyGroup>
 
   <PropertyGroup Condition="'$(Configuration)|$(Platform)' == 'Development|AnyCPU'">
diff --git a/Engine/Source/Programs/Shared/EpicGames.Benchmarks/EpicGames.Benchmarks.csproj b/Engine/Source/Programs/Shared/EpicGames.Benchmarks/EpicGames.Benchmarks.csproj
--- a/Engine/Source/Programs/Shared/EpicGames.Benchmarks/EpicGames.Benchmarks.csproj
+++ b/Engine/Source/Programs/Shared/EpicGames.Benchmarks/EpicGames.Benchmarks.csproj
@@ -2,9 +2,6 @@
   <PropertyGroup>
     <TargetFramework>net8.0</TargetFramework>
     <OutputType>Exe</OutputType>
-    <Nullable>enable</Nullable>
-    <Configuration Condition=" '$(Configuration)' == '' ">Development</Configuration>
-    <Configurations>Debug;Release;Development;Analyze</Configurations>
     <NoWarn>$(NoWarn);1591</NoWarn>
   </PropertyGroup>
 
diff --git a/Engine/Source/Programs/Shared/EpicGames.Build/EpicGames.Build.csproj b/Engine/Source/Programs/Shared/EpicGames.Build/EpicGames.Build.csproj
--- a/Engine/Source/Programs/Shared/EpicGames.Build/EpicGames.Build.csproj
+++ b/Engine/Source/Programs/Shared/EpicGames.Build/EpicGames.Build.csproj
@@ -10,7 +10,7 @@
   </PropertyGroup>
 
   <ItemGroup>
-    <ProjectReference Include="..\..\Shared\EpicGames.Core\EpicGames.Core.csproj" />
+    <ProjectReference Include="..\EpicGames.Core\EpicGames.Core.csproj" />
     <ProjectReference Include="..\EpicGames.IoHash\EpicGames.IoHash.csproj" />
     <ProjectReference Include="..\EpicGames.MsBuild\EpicGames.MsBuild.csproj" />
   </ItemGroup>
diff --git a/Engine/Source/Programs/Shared/EpicGames.Core/EpicGames.Core.csproj b/Engine/Source/Programs/Shared/EpicGames.Core/EpicGames.Core.csproj
--- a/Engine/Source/Programs/Shared/EpicGames.Core/EpicGames.Core.csproj
+++ b/Engine/Source/Programs/Shared/EpicGames.Core/EpicGames.Core.csproj
@@ -4,20 +4,6 @@
     <AllowUnsafeBlocks>true</AllowUnsafeBlocks>
   </PropertyGroup>
 
-  <ItemGroup>
-    <Compile Include="..\..\..\..\Restricted\*\Source\Programs\Shared\EpicGames.Core\ScopeProfiler\*.cs">
-      <Link>$([System.Text.RegularExpressions.Regex]::Replace(%(FullPath), '^.+?[\\/]EpicGames.Core[\\/]', ''))</Link>
-    </Compile>
-  </ItemGroup>
-
-  <Choose>
-    <When Condition="Exists('..\..\..\..\Restricted\NoRedist\Source\Programs\Shared\EpicGames.Core\ScopeProfiler\')">
-      <PropertyGroup>
-        <DefineConstants>$(DefineConstants);__SCOPEPROFILER_AVAILABLE__</DefineConstants>
-      </PropertyGroup>
-    </When>
-  </Choose>
-
   <ItemGroup>
     <PackageReference Include="JetBrains.Annotations" Version="2024.3.0" />
     <PackageReference Include="Microsoft.Extensions.Configuration" Version="8.0.0" />
diff --git a/Engine/Source/Programs/Shared/EpicGames.Core/ScopeProfiler.cs b/Engine/Source/Programs/Shared/EpicGames.Core/ScopeProfiler.cs
--- a/Engine/Source/Programs/Shared/EpicGames.Core/ScopeProfiler.cs
+++ b/Engine/Source/Programs/Shared/EpicGames.Core/ScopeProfiler.cs
@@ -4,7 +4,6 @@
 
 namespace EpicGames.Core
 {
-#if !__SCOPEPROFILER_AVAILABLE__
 #pragma warning disable IDE0060 // Remove unused parameter
 
 	/// <summary>
@@ -81,6 +80,5 @@ namespace EpicGames.Core
 		}
 	}
 #pragma warning restore IDE0060 // Remove unused parameter
-#endif
 }
 
diff --git a/Engine/Source/Programs/Shared/EpicGames.Tracing.Tests/EpicGames.Tracing.Tests.csproj b/Engine/Source/Programs/Shared/EpicGames.Tracing.Tests/EpicGames.Tracing.Tests.csproj
--- a/Engine/Source/Programs/Shared/EpicGames.Tracing.Tests/EpicGames.Tracing.Tests.csproj
+++ b/Engine/Source/Programs/Shared/EpicGames.Tracing.Tests/EpicGames.Tracing.Tests.csproj
@@ -1,7 +1,6 @@
 <Project Sdk="Microsoft.NET.Sdk">
   <PropertyGroup>
     <TargetFramework>net8.0</TargetFramework>
-    <OutputType>Library</OutputType>
   </PropertyGroup>
 
   <ItemGroup>
diff --git a/Engine/Source/Programs/Shared/EpicGames.UBA/EpicGames.UBA.csproj b/Engine/Source/Programs/Shared/EpicGames.UBA/EpicGames.UBA.csproj
--- a/Engine/Source/Programs/Shared/EpicGames.UBA/EpicGames.UBA.csproj
+++ b/Engine/Source/Programs/Shared/EpicGames.UBA/EpicGames.UBA.csproj
@@ -1,15 +1,17 @@
 ﻿<Project Sdk="Microsoft.NET.Sdk">
   <PropertyGroup>
     <TargetFramework>net8.0</TargetFramework>
-    <AllowUnsafeBlocks>true</AllowUnsafeBlocks>
     <GenerateDocumentationFile>True</GenerateDocumentationFile>
     <PackageReadmeFile>README.md</PackageReadmeFile>
   </PropertyGroup>
 
-  <ItemGroup>
-    <Compile Include="../../../../Restricted/*/Source/Programs/Shared/EpicGames.UBA/**/*.cs">
-      <Link>Restricted/$([System.Text.RegularExpressions.Regex]::Replace(%(FullPath), '^.+?[\\/]EpicGames.UBA[\\/]', ''))</Link>
-    </Compile>
+  <ItemGroup Condition="'$(DesignTimeBuild)' == 'true'">
+    <Compile Include="../../../../Restricted/*/Source/Programs/Shared/EpicGames.UBA/**/*.cs"
+             Link="Restricted/$([System.String]::Copy('%(RecursiveDir)').Replace('\', '/').Replace('Source/Programs/Shared/EpicGames.UBA/', ''))%(Filename)%(Extension)" />
+  </ItemGroup>
+
+  <ItemGroup Condition="'$(DesignTimeBuild)' != 'true'">
+    <Compile Include="../../../../Restricted/*/Source/Programs/Shared/EpicGames.UBA/**/*.cs" />
   </ItemGroup>
 
   <ItemGroup>
