Shared: Reenable optimizations for Development in some projects

#rnx

#Synced-CL 50360693

diff --git a/Engine/Source/Programs/Shared/EpicGames.Analytics.Generators/EpicGames.Analytics.Generators.csproj b/Engine/Source/Programs/Shared/EpicGames.Analytics.Generators/EpicGames.Analytics.Generators.csproj
--- a/Engine/Source/Programs/Shared/EpicGames.Analytics.Generators/EpicGames.Analytics.Generators.csproj
+++ b/Engine/Source/Programs/Shared/EpicGames.Analytics.Generators/EpicGames.Analytics.Generators.csproj
@@ -6,10 +6,6 @@
     <EnforceExtendedAnalyzerRules>true</EnforceExtendedAnalyzerRules>
   </PropertyGroup>
 
-  <PropertyGroup Condition="'$(Configuration)|$(Platform)' == 'Development|AnyCPU'">
-    <Optimize>false</Optimize>
-  </PropertyGroup>
-
   <ItemGroup>
     <PackageReference Include="Microsoft.CodeAnalysis.Analyzers" Version="3.3.4" PrivateAssets="all" />
     <PackageReference Include="Microsoft.CodeAnalysis.Common" Version="4.8.0" PrivateAssets="all" />
diff --git a/Engine/Source/Programs/Shared/EpicGames.Analytics.PerformanceTrends/EpicGames.Analytics.PerformanceTrends.csproj b/Engine/Source/Programs/Shared/EpicGames.Analytics.PerformanceTrends/EpicGames.Analytics.PerformanceTrends.csproj
--- a/Engine/Source/Programs/Shared/EpicGames.Analytics.PerformanceTrends/EpicGames.Analytics.PerformanceTrends.csproj
+++ b/Engine/Source/Programs/Shared/EpicGames.Analytics.PerformanceTrends/EpicGames.Analytics.PerformanceTrends.csproj
@@ -6,10 +6,6 @@
     <GenerateDocumentationFile>True</GenerateDocumentationFile>
   </PropertyGroup>
 
-  <PropertyGroup Condition="'$(Configuration)|$(Platform)' == 'Development|AnyCPU'">
-    <Optimize>false</Optimize>
-  </PropertyGroup>
-
   <ItemGroup>
     <ProjectReference Include="..\EpicGames.Analytics.Telemetry\EpicGames.Analytics.Telemetry.csproj" />
     <ProjectReference Include="..\EpicGames.Horde\EpicGames.Horde.csproj" />
diff --git a/Engine/Source/Programs/Shared/EpicGames.Analytics.Telemetry/EpicGames.Analytics.Telemetry.csproj b/Engine/Source/Programs/Shared/EpicGames.Analytics.Telemetry/EpicGames.Analytics.Telemetry.csproj
--- a/Engine/Source/Programs/Shared/EpicGames.Analytics.Telemetry/EpicGames.Analytics.Telemetry.csproj
+++ b/Engine/Source/Programs/Shared/EpicGames.Analytics.Telemetry/EpicGames.Analytics.Telemetry.csproj
@@ -3,8 +3,4 @@
     <TargetFramework>net8.0</TargetFramework>
     <RootNamespace>EpicGames.Analytics</RootNamespace>
   </PropertyGroup>
-
-  <PropertyGroup Condition="'$(Configuration)|$(Platform)' == 'Development|AnyCPU'">
-    <Optimize>false</Optimize>
-  </PropertyGroup>
 </Project>
diff --git a/Engine/Source/Programs/Shared/EpicGames.Horde/EpicGames.Horde.csproj b/Engine/Source/Programs/Shared/EpicGames.Horde/EpicGames.Horde.csproj
--- a/Engine/Source/Programs/Shared/EpicGames.Horde/EpicGames.Horde.csproj
+++ b/Engine/Source/Programs/Shared/EpicGames.Horde/EpicGames.Horde.csproj
@@ -5,10 +5,6 @@
     <AllowUnsafeBlocks>True</AllowUnsafeBlocks>
   </PropertyGroup>
 
-  <PropertyGroup Condition="'$(Configuration)|$(Platform)' == 'Development|AnyCPU'">
-    <Optimize>false</Optimize>
-  </PropertyGroup>
-
   <ItemGroup>
     <PackageReference Include="BitFaster.Caching" Version="2.4.1" />
     <PackageReference Include="Dapper" Version="2.1.66" />
diff --git a/Engine/Source/Programs/Shared/EpicGames.Oodle/EpicGames.Oodle.csproj b/Engine/Source/Programs/Shared/EpicGames.Oodle/EpicGames.Oodle.csproj
--- a/Engine/Source/Programs/Shared/EpicGames.Oodle/EpicGames.Oodle.csproj
+++ b/Engine/Source/Programs/Shared/EpicGames.Oodle/EpicGames.Oodle.csproj
@@ -6,10 +6,6 @@
     <AllowUnsafeBlocks>True</AllowUnsafeBlocks>
   </PropertyGroup>
 
-  <PropertyGroup Condition="'$(Configuration)|$(Platform)' == 'Development|AnyCPU'">
-    <Optimize>false</Optimize>
-  </PropertyGroup>
-
   <ItemGroup>
     <None Condition="'$([System.Runtime.InteropServices.RuntimeInformation]::IsOSPlatform($([System.Runtime.InteropServices.OSPlatform]::Linux)))' == 'true' Or Exists('Sdk\2.9.10\linux\lib\liboo2corelinux64.so.9')" Include="Sdk\2.9.10\linux\lib\liboo2corelinux64.so.9" Link="liboo2corelinux64.so.9" CopyToOutputDirectory="PreserveNewest" Visible="false" />
     <None Condition="'$([System.Runtime.InteropServices.RuntimeInformation]::IsOSPlatform($([System.Runtime.InteropServices.OSPlatform]::Linux)))' == 'true' Or Exists('Sdk\2.9.10\linuxarm\lib\liboo2corelinuxarm32.so.9')" Include="Sdk\2.9.10\linuxarm\lib\liboo2corelinuxarm32.so.9" Link="liboo2corelinuxarm32.so.9" CopyToOutputDirectory="PreserveNewest" Visible="false" />
