EpicGames.UHT: Remove unnecessary AllowUnsafeBlocks

#rnx
#Synced-CL 50360693

diff --git a/Engine/Source/Programs/Shared/EpicGames.UHT/EpicGames.UHT.csproj b/Engine/Source/Programs/Shared/EpicGames.UHT/EpicGames.UHT.csproj
--- a/Engine/Source/Programs/Shared/EpicGames.UHT/EpicGames.UHT.csproj
+++ b/Engine/Source/Programs/Shared/EpicGames.UHT/EpicGames.UHT.csproj
@@ -1,14 +1,13 @@
 ﻿<Project Sdk="Microsoft.NET.Sdk">
   <PropertyGroup>
     <TargetFramework>net8.0</TargetFramework>
-    <AllowUnsafeBlocks>true</AllowUnsafeBlocks>
     <GenerateDocumentationFile>True</GenerateDocumentationFile>
     <EnforceCodeStyleInBuild>False</EnforceCodeStyleInBuild>
   </PropertyGroup>
 
   <ItemGroup>
-    <ProjectReference Include="..\..\Shared\EpicGames.Core\EpicGames.Core.csproj" />
     <ProjectReference Include="..\EpicGames.Build\EpicGames.Build.csproj" />
+    <ProjectReference Include="..\EpicGames.Core\EpicGames.Core.csproj" />
   </ItemGroup>
 
   <ItemGroup>
