Add the option to clean into RetrieveCloudArtifact so we can clean the folder for test jobs.

#Synced-CL 50361130

diff --git a/Engine/Source/Programs/AutomationTool/BuildGraph/Tasks/RetrieveCloudArtifactTask.cs b/Engine/Source/Programs/AutomationTool/BuildGraph/Tasks/RetrieveCloudArtifactTask.cs
--- a/Engine/Source/Programs/AutomationTool/BuildGraph/Tasks/RetrieveCloudArtifactTask.cs
+++ b/Engine/Source/Programs/AutomationTool/BuildGraph/Tasks/RetrieveCloudArtifactTask.cs
@@ -153,6 +153,12 @@ namespace AutomationTool.Tasks
 		/// </summary>
 		[TaskParameter(Optional = true)]
 		public bool Append { get; set; } = false;
+
+		/// <summary>
+		/// Delete all data in target folder that is not part of the downloaded content
+		/// </summary>
+		[TaskParameter(Optional = true)]
+		public bool Clean { get; set; } = false;
 	}
 
 	/// <summary>
@@ -412,7 +418,8 @@ namespace AutomationTool.Tasks
 					// boost worker memory if boosting workers as well
 					BoostWorkerMemory = CommandUtils.IsBuildMachine || _parameters.BoostWorkers,
 					TraceFile = traceFile,
-					Append = _parameters.Append
+					Append = _parameters.Append,
+					Clean = _parameters.Clean
 				}, optionsOverride: optionsOverride);
 		}
 
diff --git a/Engine/Source/Programs/Shared/EpicGames.ProjectStore/CloudStorage/CloudStorage.cs b/Engine/Source/Programs/Shared/EpicGames.ProjectStore/CloudStorage/CloudStorage.cs
--- a/Engine/Source/Programs/Shared/EpicGames.ProjectStore/CloudStorage/CloudStorage.cs
+++ b/Engine/Source/Programs/Shared/EpicGames.ProjectStore/CloudStorage/CloudStorage.cs
@@ -191,6 +191,7 @@ namespace EpicGames.ProjectStore
 			string boostWorkerMemoryOption = downloadOptions?.BoostWorkerMemory ?? false ? "--boost-worker-memory" : "";
 			string allowMultipartOption = downloadOptions?.AllowMultipart != null ? $"--allow-multipart={downloadOptions.AllowMultipart}" : "";
 			string append = downloadOptions?.Append ?? false ? "--append=true" : String.Empty;
+			string clean = downloadOptions?.Clean ?? false ? "--clean" : ""; 
 			// trace file option needs to be passed before other arguments including the verbs
 			string traceOption = downloadOptions?.TraceFile != null ? $"--tracefile={downloadOptions.TraceFile.FullName} " : "";
 			string zenFolderPath = downloadOptions?.ZenStateFolder != null ? $"--zen-folder-path={downloadOptions.ZenStateFolder.FullName} " : "";
@@ -198,7 +199,7 @@ namespace EpicGames.ProjectStore
 			string wildcardOption = includeFiles == null ? "" : $"--wildcard \"{String.Join(';', includeFiles.Select(p => ToZenWildcard(outputDir, p)))}\"";
 			string excludeOption = excludeFiles == null ? "" : $"--exclude-wildcard \"{String.Join(';', excludeFiles.Select(p => ToZenWildcard(outputDir, p)))}\"";
 			// pass the access token via environment variable to avoid showing it on the cli
-			string cmdline = $"{traceOption}builds download {_options.Value.GetHostCliArgs(optionsOverride)} --namespace {ns} --bucket {bucketId} {wildcardOption} {excludeOption} {allowMultipartOption} {append} {http2Options} {boostWorkersOption} {boostWorkerMemoryOption} {zenFolderPath} \"{outputDir}\" {buildId}";
+			string cmdline = $"{traceOption}builds download {_options.Value.GetHostCliArgs(optionsOverride)} --namespace {ns} --bucket {bucketId} {wildcardOption} {excludeOption} {allowMultipartOption} {append} {clean} {http2Options} {boostWorkersOption} {boostWorkerMemoryOption} {zenFolderPath} \"{outputDir}\" {buildId}";
 
 			int exitCode = await RunZenCliAsync(cmdline, progress: progress, optionsOverride: optionsOverride, cancellationToken: cancellationToken);
 			EnsureZenExitCodeSuccess(exitCode);
diff --git a/Engine/Source/Programs/Shared/EpicGames.ProjectStore/CloudStorage/ICloudStorage.cs b/Engine/Source/Programs/Shared/EpicGames.ProjectStore/CloudStorage/ICloudStorage.cs
--- a/Engine/Source/Programs/Shared/EpicGames.ProjectStore/CloudStorage/ICloudStorage.cs
+++ b/Engine/Source/Programs/Shared/EpicGames.ProjectStore/CloudStorage/ICloudStorage.cs
@@ -219,6 +219,11 @@ namespace EpicGames.ProjectStore
 		/// Set this to true if you want to download an artifact to a directory, then download an additional artifact from another without wiping state between
 		/// </summary>
 		public bool Append { get; set; }
+
+		/// <summary>
+		/// Delete all data in target folder that is not part of the downloaded content
+		/// </summary>
+		public bool Clean { get; set; }
 	}
 
 	/// <summary>
