Fix assorted code analysis issues in UHT

#rnx

#Synced-CL 50396713

diff --git a/Engine/Source/Programs/Shared/EpicGames.UHT/Exporters/CodeGen/UhtHeaderCodeGeneratorCppFile.cs b/Engine/Source/Programs/Shared/EpicGames.UHT/Exporters/CodeGen/UhtHeaderCodeGeneratorCppFile.cs
--- a/Engine/Source/Programs/Shared/EpicGames.UHT/Exporters/CodeGen/UhtHeaderCodeGeneratorCppFile.cs
+++ b/Engine/Source/Programs/Shared/EpicGames.UHT/Exporters/CodeGen/UhtHeaderCodeGeneratorCppFile.cs
@@ -1639,7 +1639,7 @@ namespace EpicGames.UHT.Exporters.CodeGen
 					}
 					if (isVerseCallable)
 					{
-						AppendVerseCallableFunction(builder, classObj, function, "verse_exe_");
+						AppendVerseCallableFunction(builder, classObj, function);
 					}
 				}
 				if (isRpc)
@@ -1755,9 +1755,9 @@ namespace EpicGames.UHT.Exporters.CodeGen
 
 			// The bodies must differ based on the VM
 			builder.Append("#if WITH_VERSE_VM\r\n");
-			AppendVerseNativeCallableFunctionBody(builder, classObj, function, classNameSuffix, returnProperty, true);
+			AppendVerseNativeCallableFunctionBody(builder, classObj, function, returnProperty, true);
 			builder.Append("#else\r\n");
-			AppendVerseNativeCallableFunctionBody(builder, classObj, function, classNameSuffix, returnProperty, false);
+			AppendVerseNativeCallableFunctionBody(builder, classObj, function, returnProperty, false);
 			builder.Append("#endif\r\n");
 
 			if (isSuspends)
@@ -1768,7 +1768,7 @@ namespace EpicGames.UHT.Exporters.CodeGen
 			return builder;
 		}
 
-		private static StringBuilder AppendVerseNativeCallableFunctionBody(StringBuilder builder, UhtClass classObj, UhtFunction function, string classNameSuffix, UhtProperty? returnProperty, bool isVerseVM)
+		private static StringBuilder AppendVerseNativeCallableFunctionBody(StringBuilder builder, UhtClass classObj, UhtFunction function, UhtProperty? returnProperty, bool isVerseVM)
 		{
 			builder.Append(isVerseVM ? "\tAutoRTFM::Open([&] {\r\n" : "\t{\r\n");
 
@@ -1816,7 +1816,7 @@ namespace EpicGames.UHT.Exporters.CodeGen
 			return builder;
 		}
 
-		private static StringBuilder AppendVerseCallableFunction(StringBuilder builder, UhtClass classObj, UhtFunction function, string functionPrefix)
+		private static StringBuilder AppendVerseCallableFunction(StringBuilder builder, UhtClass classObj, UhtFunction function)
 		{
 			UhtProperty? returnProperty = function.ReturnProperty;
 
diff --git a/Engine/Source/Programs/Shared/EpicGames.UHT/Exporters/CodeGen/UhtHeaderCodeGeneratorHFile.cs b/Engine/Source/Programs/Shared/EpicGames.UHT/Exporters/CodeGen/UhtHeaderCodeGeneratorHFile.cs
--- a/Engine/Source/Programs/Shared/EpicGames.UHT/Exporters/CodeGen/UhtHeaderCodeGeneratorHFile.cs
+++ b/Engine/Source/Programs/Shared/EpicGames.UHT/Exporters/CodeGen/UhtHeaderCodeGeneratorHFile.cs
@@ -890,7 +890,7 @@ namespace EpicGames.UHT.Exporters.CodeGen
 
 				using (UhtMacroCreator macro = new(builder, this, classObj, GeneratedUInterfaceBodyMacroSuffix))
 				{
-					AppendCommonGeneratedBody(builder, classObj, api);
+					AppendCommonGeneratedBody(builder, classObj);
 				}
 
 				if (usesLegacy)
@@ -1565,7 +1565,7 @@ namespace EpicGames.UHT.Exporters.CodeGen
 		/// <returns>Output builder</returns>
 		private StringBuilder AppendClassGeneratedBody(StringBuilder builder, UhtClass classObj, string api)
 		{
-			AppendCommonGeneratedBody(builder, classObj, api);
+			AppendCommonGeneratedBody(builder, classObj);
 
 			// export the class's config name
 			UhtClass? superClass = classObj.SuperClass;
@@ -1600,9 +1600,8 @@ namespace EpicGames.UHT.Exporters.CodeGen
 		/// </summary>
 		/// <param name="builder">Output builder</param>
 		/// <param name="classObj">Class being exported</param>
-		/// <param name="api">API text to be used</param>
 		/// <returns>Output builder</returns>
-		private StringBuilder AppendCommonGeneratedBody(StringBuilder builder, UhtClass classObj, string api)
+		private StringBuilder AppendCommonGeneratedBody(StringBuilder builder, UhtClass classObj)
 		{
 			UhtCodeGenerator.ObjectInfo info = ObjectInfos[classObj.ObjectTypeIndex];
 
diff --git a/Engine/Source/Programs/Shared/EpicGames.UHT/Tables/UhtCodeGeneratorInjectorTable.cs b/Engine/Source/Programs/Shared/EpicGames.UHT/Tables/UhtCodeGeneratorInjectorTable.cs
--- a/Engine/Source/Programs/Shared/EpicGames.UHT/Tables/UhtCodeGeneratorInjectorTable.cs
+++ b/Engine/Source/Programs/Shared/EpicGames.UHT/Tables/UhtCodeGeneratorInjectorTable.cs
@@ -94,7 +94,7 @@ namespace EpicGames.UHT.Tables
 			}
 		}
 
-		internal void OnUhtCodeGeneratorInjectorAttribute(Type type, MethodInfo method, UhtCodeGeneratorInjectorAttribute attribute)
+		internal void OnUhtCodeGeneratorInjectorAttribute(MethodInfo method, UhtCodeGeneratorInjectorAttribute attribute)
 		{
 			// Walk back from each support leaf types up to the declared type
 			// If our leaf type IS of the declared type, then add a delegate for the handler of the leaf type
diff --git a/Engine/Source/Programs/Shared/EpicGames.UHT/Utils/UhtTables.cs b/Engine/Source/Programs/Shared/EpicGames.UHT/Utils/UhtTables.cs
--- a/Engine/Source/Programs/Shared/EpicGames.UHT/Utils/UhtTables.cs
+++ b/Engine/Source/Programs/Shared/EpicGames.UHT/Utils/UhtTables.cs
@@ -733,7 +733,7 @@ namespace EpicGames.UHT.Utils
 					}
 					else if (methodAttribute is UhtCodeGeneratorInjectorAttribute codeGeneratorInjectorAttribute)
 					{
-						CodeGeneratorInjectorTable.OnUhtCodeGeneratorInjectorAttribute(type, methodInfo, codeGeneratorInjectorAttribute);
+						CodeGeneratorInjectorTable.OnUhtCodeGeneratorInjectorAttribute(methodInfo, codeGeneratorInjectorAttribute);
 					}
 				}
 			}
