Fix StateStream enums for structs with > 31 properties.
Add error for structs with > 64 properties.

#rnx
#Synced-CL 50405523

diff --git a/Engine/Source/Programs/Shared/EpicGames.UHT/Exporters/CodeGen/UhtHeaderCodeGeneratorHFile.cs b/Engine/Source/Programs/Shared/EpicGames.UHT/Exporters/CodeGen/UhtHeaderCodeGeneratorHFile.cs
--- a/Engine/Source/Programs/Shared/EpicGames.UHT/Exporters/CodeGen/UhtHeaderCodeGeneratorHFile.cs
+++ b/Engine/Source/Programs/Shared/EpicGames.UHT/Exporters/CodeGen/UhtHeaderCodeGeneratorHFile.cs
@@ -542,9 +542,16 @@ namespace EpicGames.UHT.Exporters.CodeGen
 
 					if (scriptStruct.MetaData.ContainsKey(UhtNames.StateStreamStaticState) || scriptStruct.MetaData.ContainsKey(UhtNames.StateStreamDynamicState))
 					{
-						builder.Append("\tenum Fields \\\r\n");
+						if (scriptStruct.Properties.Count() > 64)
+						{
+							throw new UhtIceException("StateStream states cannot support more than 64 properties.");
+						}
+
+						string modifiedFlagsType = scriptStruct.Properties.Count() > 32 ? "uint64" : "uint32";
+
+						builder.Append("\tenum class Fields : ").Append(modifiedFlagsType).Append(" \\\r\n");
 						builder.Append("\t{ \\\r\n");
-						int enumValue = 1;
+						ulong enumValue = 1;
 						foreach (UhtProperty property in scriptStruct.Properties)
 						{
 							builder.Append("\t\tField_").Append(property.SourceName).Append(" = ").Append(enumValue).Append(", \\\r\n");
@@ -552,7 +559,7 @@ namespace EpicGames.UHT.Exporters.CodeGen
 						}
 						builder.Append("\t}; \\\r\n");
 
-						builder.Append("\tuint32 ModifiedFlags = 0; \\\r\n");
+						builder.Append('\t').Append(modifiedFlagsType).Append(" ModifiedFlags = 0; \\\r\n");
 
 						builder.Append('\t').Append(scriptStruct.SourceName).Append("() = default; \\\r\n");
 
@@ -587,11 +594,11 @@ namespace EpicGames.UHT.Exporters.CodeGen
 							string propName = property.SourceName;
 							bool isBoolWithB = (property is UhtBoolProperty && propName.StartsWith('b'));
 							string cleanPropName = isBoolWithB ? propName[1..] : propName;
-							builder.Append("\tinline bool ").Append(cleanPropName).Append("Modified() const { return !!(ModifiedFlags & Field_").Append(propName).Append("); } \\\r\n");
+							builder.Append("\tinline bool ").Append(cleanPropName).Append("Modified() const { return !!(ModifiedFlags & (").Append(modifiedFlagsType).Append(")Fields::Field_").Append(propName).Append("); } \\\r\n");
 
 							builder.Append("\tinline void Set").Append(cleanPropName).Append("(const ");
 							property.AppendText(builder, UhtPropertyTextType.Generic);
-							builder.Append("& In) { ").Append(propName).Append(" = In; ModifiedFlags |= Field_").Append(propName).Append("; } \\\r\n");
+							builder.Append("& In) { ").Append(propName).Append(" = In; ModifiedFlags |= (").Append(modifiedFlagsType).Append(")Fields::Field_").Append(propName).Append("; } \\\r\n");
 
 							if (isBoolWithB)
 							{
@@ -650,7 +657,7 @@ namespace EpicGames.UHT.Exporters.CodeGen
 							{
 								builder.Append("\t\t\t").Append(propName).Append(" = Other.").Append(propName).Append("; \\\r\n");
 							}
-							builder.Append("\t\t\tModifiedFlags |= Field_").Append(propName).Append("; \\\r\n");
+							builder.Append("\t\t\tModifiedFlags |= (").Append(modifiedFlagsType).Append(")Fields::Field_").Append(propName).Append("; \\\r\n");
 							builder.Append("\t\t} \\\r\n");
 						}
 						builder.Append("\t} \\\r\n");
@@ -667,7 +674,7 @@ namespace EpicGames.UHT.Exporters.CodeGen
 
 							builder.Append("\t\tif (To.").Append(cleanPropName).Append("Modified()) \\\r\n");
 							builder.Append("\t\t{ \\\r\n");
-							builder.Append("\t\t\tModifiedFlags |= Field_").Append(propName).Append("; \\\r\n");
+							builder.Append("\t\t\tModifiedFlags |= (").Append(modifiedFlagsType).Append(")Fields::Field_").Append(propName).Append("; \\\r\n");
 
 							bool assignHandled = false;
 							if (property is UhtStructProperty structProperty)
