Add support for TSubScriptStructOf in UnrealHeaderTool. The customizer supports metadata MetaStruct (like TObjectPtr<UScriptStruct>) and BaseStruct (like FInstanceStruct).

#Synced-CL 50428318

diff --git a/Engine/Source/Editor/DetailCustomizations/Private/SubScriptStructOfCustomization.cpp b/Engine/Source/Editor/DetailCustomizations/Private/SubScriptStructOfCustomization.cpp
--- a/Engine/Source/Editor/DetailCustomizations/Private/SubScriptStructOfCustomization.cpp
+++ b/Engine/Source/Editor/DetailCustomizations/Private/SubScriptStructOfCustomization.cpp
@@ -19,7 +19,8 @@ void FSubScriptStructOfCustomization::CustomizeHeader(TSharedRef<IPropertyHandle
 {
 	PropertyHandle = InPropertyHandle;
 
-	const FString& MetaStructName = PropertyHandle->GetMetaData("MetaStruct");
+	const FString& MetaStructName1 = PropertyHandle->GetMetaData("MetaStruct");
+	const FString& MetaStructName2 = PropertyHandle->GetMetaData("BaseStruct");
 	const bool bAllowNone = !(PropertyHandle->GetMetaDataProperty()->PropertyFlags & CPF_NoClear);
 	const bool bShowTreeView = PropertyHandle->HasMetaData("ShowTreeView");
 	const bool bHideViewOptions = PropertyHandle->HasMetaData("HideViewOptions");
@@ -27,9 +28,17 @@ void FSubScriptStructOfCustomization::CustomizeHeader(TSharedRef<IPropertyHandle
 	
 
 	const UScriptStruct* MetaScriptStruct = nullptr;
-	if (!MetaStructName.IsEmpty())
+	if (!MetaStructName1.IsEmpty())
 	{
-		TArray<const UScriptStruct*> Structs = PropertyCustomizationHelpers::GetStructsFromMetadataString(MetaStructName);
+		TArray<const UScriptStruct*> Structs = PropertyCustomizationHelpers::GetStructsFromMetadataString(MetaStructName1);
+		if (Structs.Num() >= 1)
+		{
+			MetaScriptStruct = Structs[0];
+		}
+	}
+	else if (!MetaStructName2.IsEmpty())
+	{
+		TArray<const UScriptStruct*> Structs = PropertyCustomizationHelpers::GetStructsFromMetadataString(MetaStructName2);
 		if (Structs.Num() >= 1)
 		{
 			MetaScriptStruct = Structs[0];
diff --git a/Engine/Source/Programs/Shared/EpicGames.UHT/Types/Properties/UhtStructProperty.cs b/Engine/Source/Programs/Shared/EpicGames.UHT/Types/Properties/UhtStructProperty.cs
--- a/Engine/Source/Programs/Shared/EpicGames.UHT/Types/Properties/UhtStructProperty.cs
+++ b/Engine/Source/Programs/Shared/EpicGames.UHT/Types/Properties/UhtStructProperty.cs
@@ -638,6 +638,34 @@ namespace EpicGames.UHT.Types
 			return new UhtTemplateStructProperty(args.PropertySettings, baseScriptStruct.Session.FInstancedStruct, "TInstancedStruct", baseScriptStruct);
 		}
 
+		[UhtPropertyType(Keyword = "TSubScriptStructOf")]
+		[SuppressMessage("CodeQuality", "IDE0051:Remove unused private members", Justification = "Attribute accessed method")]
+		[SuppressMessage("Style", "IDE0060:Remove unused parameter", Justification = "Attribute accessed method")]
+		private static UhtTemplateStructProperty? SubScriptStructOfProperty(UhtPropertyResolveArgs args)
+		{
+			UhtScriptStruct? baseScriptStruct = args.ParseTemplateScriptStruct();
+			if (baseScriptStruct == null || baseScriptStruct.Session.FSubScriptStructOf == null)
+			{
+				return null;
+			}
+
+			UhtPropertySettings rootSettings = args.PropertySettings.RootSettings;
+			if (rootSettings.MetaData.ContainsKey("MetaStruct"))
+			{
+				args.TokenReader.LogError("MetaStruct metadata is implicitly set from the TSubScriptStructOf template argument and should not be explicitly specified.");
+				return null;
+			}
+			if (rootSettings.MetaData.ContainsKey("BaseStruct"))
+			{
+				args.TokenReader.LogError("BaseStruct metadata is implicitly set from the TSubScriptStructOf template argument and should not be explicitly specified.");
+				return null;
+			}
+
+			// With TSubScriptStructOf, MetaStruct or BaseStruct is used as a type limiter.
+			rootSettings.MetaData.Add("MetaStruct", baseScriptStruct.PathName);
+			return new UhtTemplateStructProperty(args.PropertySettings, baseScriptStruct.Session.FSubScriptStructOf, "TSubScriptStructOf", baseScriptStruct);
+		}
+
 		[UhtPropertyType(Keyword = "TStateTreePropertyRef")]
 		[SuppressMessage("CodeQuality", "IDE0051:Remove unused private members", Justification = "Attribute accessed method")]
 		[SuppressMessage("Style", "IDE0060:Remove unused parameter", Justification = "Attribute accessed method")]
diff --git a/Engine/Source/Programs/Shared/EpicGames.UHT/Utils/UhtSession.cs b/Engine/Source/Programs/Shared/EpicGames.UHT/Utils/UhtSession.cs
--- a/Engine/Source/Programs/Shared/EpicGames.UHT/Utils/UhtSession.cs
+++ b/Engine/Source/Programs/Shared/EpicGames.UHT/Utils/UhtSession.cs
@@ -988,6 +988,12 @@ namespace EpicGames.UHT.Utils
 		/// </summary>
 		public UhtScriptStruct? FInstancedStruct { get; set; } = null;
 
+		/// <summary>
+		/// After headers are parsed, returns the FSubScriptStructOf script struct.  Unlike such properties as 
+		/// "UObject", there is no requirement for FSubScriptStructOf to be defined.  May be null.
+		/// </summary>
+		public UhtScriptStruct? FSubScriptStructOf { get; set; } = null;
+
 		/// <summary>
 		/// After headers are parsed, returns the FStateTreePropertyRef script struct.
 		/// There is no requirement for FStateTreePropertyRef to be defined. May be null.
@@ -1045,6 +1051,7 @@ namespace EpicGames.UHT.Utils
 			AActor = (UhtClass?)FindType(null, UhtFindOptions.SourceName | UhtFindOptions.Class, "AActor");
 			INotifyFieldValueChanged = (UhtClass?)FindType(null, UhtFindOptions.SourceName | UhtFindOptions.Class, "INotifyFieldValueChanged");
 			FInstancedStruct = (UhtScriptStruct?)FindType(null, UhtFindOptions.SourceName | UhtFindOptions.ScriptStruct, "FInstancedStruct");
+			FSubScriptStructOf = (UhtScriptStruct?)FindType(null, UhtFindOptions.SourceName | UhtFindOptions.ScriptStruct, "FSubScriptStructOf");
 			FStateTreePropertyRef = (UhtScriptStruct?)FindType(null, UhtFindOptions.SourceName | UhtFindOptions.ScriptStruct, "FStateTreePropertyRef");
 		}
 		#endregion
diff --git a/Engine/Source/Runtime/CoreUObject/Tests/SubScriptStructOfTest.h b/Engine/Source/Runtime/CoreUObject/Tests/SubScriptStructOfTest.h
--- a/Engine/Source/Runtime/CoreUObject/Tests/SubScriptStructOfTest.h
+++ b/Engine/Source/Runtime/CoreUObject/Tests/SubScriptStructOfTest.h
@@ -2,6 +2,7 @@
 
 #pragma once
 
+#include "Templates/SubScriptStructOf.h"
 #include "UObject/Object.h"
 
 #include "SubScriptStructOfTest.generated.h"
@@ -12,35 +13,37 @@ USTRUCT()
 struct FSubScriptStructOfTest_Base
 {
 	GENERATED_BODY()
-
-public:
-	
 };
 
 USTRUCT()
 struct FSubScriptStructOfTest_ChildA : public FSubScriptStructOfTest_Base
 {
 	GENERATED_BODY()
-
-public:
-	
 };
 
 USTRUCT()
 struct FSubScriptStructOfTest_ChildB : public FSubScriptStructOfTest_Base
 {
 	GENERATED_BODY()
-
-public:
-	
 };
 
 USTRUCT()
 struct FSubScriptStructOfTest_GrandChildA : public FSubScriptStructOfTest_ChildA
+{
+	GENERATED_BODY()
+};
+
+UCLASS()
+class USubScriptStructOfTest : public UObject
 {
 	GENERATED_BODY()
 
 public:
+	UPROPERTY(EditAnywhere, Category = "Properties")
+	TSubScriptStructOf<FSubScriptStructOfTest_ChildA> ChildA;
+
+	UPROPERTY(EditAnywhere, Category = "Properties", meta=(MetaStruct="/Script/CoreUObject.SubScriptStructOfTest_ChildB"))
+	FSubScriptStructOf ChildB;
 };
 
 #endif
\ No newline at end of file
