Horde: Refactor SlackClient error handling to use enum-based behavior instead of boolean predicates

* Logs "not_in_channel" error message at log level warning instead of error. This error can be common if the Horde Slack bot has not been invited to a channel.
* For throw error behavior, it's assumed that the caller will take care of logging any exceptions (previously the exception was logged, and then re-thrown)

diff --git a/Engine/Source/Programs/Shared/EpicGames.Slack/SlackClient.cs b/Engine/Source/Programs/Shared/EpicGames.Slack/SlackClient.cs
--- a/Engine/Source/Programs/Shared/EpicGames.Slack/SlackClient.cs
+++ b/Engine/Source/Programs/Shared/EpicGames.Slack/SlackClient.cs
@@ -36,6 +36,24 @@ namespace EpicGames.Slack
 		}
 	}
 
+	/// <summary>
+	/// Specifies how to handle a Slack API error
+	/// </summary>
+	enum SlackErrorBehavior
+	{
+		/// <summary>Throw a SlackException (caller handles logging)</summary>
+		Throw,
+
+		/// <summary>Log at error level, don't throw</summary>
+		LogError,
+
+		/// <summary>Log at warning level, don't throw</summary>
+		LogWarning,
+
+		/// <summary>Don't log, don't throw</summary>
+		Ignore
+	}
+
 	/// <summary>
 	/// Response from posting a Slack message
 	/// </summary>
@@ -142,45 +160,66 @@ namespace EpicGames.Slack
 			return new SlackClient(httpClient, logger);
 		}
 
-		static bool ShouldLogError<TResponse>(TResponse response) where TResponse : SlackResponse => !response.Ok;
+		static SlackErrorBehavior GetErrorBehavior<TResponse>(TResponse response) where TResponse : SlackResponse
+		{
+			if (response.Ok)
+			{
+				return SlackErrorBehavior.Ignore;
+			}
+
+			return response.Error switch
+			{
+				"already_reacted" => SlackErrorBehavior.Ignore,
+				"no_reaction" => SlackErrorBehavior.Ignore,
+				"already_in_channel" => SlackErrorBehavior.Ignore,
+				"users_not_found" => SlackErrorBehavior.Ignore,
+				_ => SlackErrorBehavior.Throw
+			};
+		}
 
 		private Task<TResponse> SendRequestAsync<TResponse>(string requestUrl, object request, CancellationToken cancellationToken) where TResponse : SlackResponse
 		{
-			return SendRequestAsync<TResponse>(requestUrl, request, ShouldLogError, cancellationToken);
+			return SendRequestAsync<TResponse>(requestUrl, request, GetErrorBehavior, cancellationToken);
 		}
 
-		private async Task<TResponse> SendRequestAsync<TResponse>(string requestUrl, object request, Func<TResponse, bool> shouldLogError, CancellationToken cancellationToken) where TResponse : SlackResponse
+		private async Task<TResponse> SendRequestAsync<TResponse>(string requestUrl, object request, Func<TResponse, SlackErrorBehavior> getErrorBehavior, CancellationToken cancellationToken) where TResponse : SlackResponse
 		{
 			using HttpRequestMessage sendMessageRequest = new(HttpMethod.Post, requestUrl);
 			string requestJson = JsonSerializer.Serialize(request, _serializerOptions);
 			using StringContent messageContent = new(requestJson, Encoding.UTF8, "application/json");
 			sendMessageRequest.Content = messageContent;
-			return await SendRequestAsync<TResponse>(sendMessageRequest, requestJson, shouldLogError, cancellationToken);
+			return await SendRequestAsync<TResponse>(sendMessageRequest, requestJson, getErrorBehavior, cancellationToken);
 		}
 
 		private Task<TResponse> SendRequestAsync<TResponse>(HttpRequestMessage request, string requestJson, CancellationToken cancellationToken) where TResponse : SlackResponse
 		{
-			return SendRequestAsync<TResponse>(request, requestJson, ShouldLogError, cancellationToken);
+			return SendRequestAsync<TResponse>(request, requestJson, GetErrorBehavior, cancellationToken);
 		}
 
-		private async Task<TResponse> SendRequestAsync<TResponse>(HttpRequestMessage request, string requestJson, Func<TResponse, bool> shouldLogError, CancellationToken cancellationToken) where TResponse : SlackResponse
+		private async Task<TResponse> SendRequestAsync<TResponse>(HttpRequestMessage request, string requestJson, Func<TResponse, SlackErrorBehavior> getErrorBehavior, CancellationToken cancellationToken) where TResponse : SlackResponse
 		{
 			HttpResponseMessage response = await _httpClient.SendAsync(request, cancellationToken);
 			byte[] responseBytes = await response.Content.ReadAsByteArrayAsync(cancellationToken);
 
 			TResponse responseObject = JsonSerializer.Deserialize<TResponse>(responseBytes)!;
-			if (shouldLogError(responseObject))
+			SlackErrorBehavior behavior = getErrorBehavior(responseObject);
+
+			switch (behavior)
 			{
-				try
-				{
+				case SlackErrorBehavior.Throw:
 					throw new SlackException(responseObject.Error ?? "unknown");
-				}
-				catch (Exception ex)
-				{
-					_logger.LogWarning(ex, "Failed to send Slack message to {Url} ({Error}). Request: {Request}. Response: {Response}", request.RequestUri, responseObject.Error, requestJson, Encoding.UTF8.GetString(responseBytes));
-					throw;
-				}
+				case SlackErrorBehavior.Ignore:
+					break;
+				case SlackErrorBehavior.LogError:
+				case SlackErrorBehavior.LogWarning:
+					LogLevel level = behavior == SlackErrorBehavior.LogError ? LogLevel.Error : LogLevel.Warning;
+					_logger.Log(level, new SlackException(responseObject.Error ?? "unknown"), "Slack request to {Url} failed ({Error}). Request: {Request}. Response: {Response}",
+						request.RequestUri, responseObject.Error, requestJson, Encoding.UTF8.GetString(responseBytes));
+					break;
+				default:
+					throw new InvalidOperationException($"Unhandled error behavior: {behavior}");
 			}
+
 			return responseObject;
 		}
 
@@ -385,6 +424,16 @@ namespace EpicGames.Slack
 			public string? Name { get; set; }
 		}
 
+		static SlackErrorBehavior GetReactionErrorBehavior<TResponse>(TResponse response) where TResponse : SlackResponse
+		{
+			if (response.Error == "not_in_channel")
+			{
+				return SlackErrorBehavior.LogWarning;
+			}
+
+			return GetErrorBehavior(response);
+		}
+
 		/// <summary>
 		/// Adds a reaction to a posted message
 		/// </summary>
@@ -398,9 +447,7 @@ namespace EpicGames.Slack
 			message.Ts = messageId.Ts;
 			message.Name = name;
 
-			static bool ShouldLogError(SlackResponse response) => !response.Ok && !String.Equals(response.Error, "already_reacted", StringComparison.Ordinal);
-
-			await SendRequestAsync<SlackResponse>(AddReactionUrl, message, ShouldLogError, cancellationToken);
+			await SendRequestAsync<SlackResponse>(AddReactionUrl, message, GetReactionErrorBehavior, cancellationToken);
 		}
 
 		/// <summary>
@@ -416,9 +463,7 @@ namespace EpicGames.Slack
 			message.Ts = messageId.Ts;
 			message.Name = name;
 
-			static bool ShouldLogError(SlackResponse response) => !response.Ok && !String.Equals(response.Error, "no_reaction", StringComparison.Ordinal);
-
-			await SendRequestAsync<SlackResponse>(RemoveReactionUrl, message, ShouldLogError, cancellationToken);
+			await SendRequestAsync<SlackResponse>(RemoveReactionUrl, message, GetReactionErrorBehavior, cancellationToken);
 		}
 
 		#endregion
@@ -456,9 +501,7 @@ namespace EpicGames.Slack
 			message.Channel = channel;
 			message.Users = String.Join(",", userIds);
 
-			static bool ShouldLogError(SlackResponse response) => !response.Ok && !String.Equals(response.Error, "already_in_channel", StringComparison.Ordinal);
-
-			await SendRequestAsync<SlackResponse>(ConversationsInviteUrl, message, ShouldLogError, cancellationToken);
+			await SendRequestAsync<SlackResponse>(ConversationsInviteUrl, message, cancellationToken);
 		}
 
 		/// <summary>
@@ -473,9 +516,8 @@ namespace EpicGames.Slack
 			message.Channel = channel;
 			message.Users = String.Join(",", userIds);
 
-			static bool ShouldLogError(SlackResponse response) => false;
-
-			SlackResponse response = await SendRequestAsync<SlackResponse>(ConversationsInviteUrl, message, ShouldLogError, cancellationToken);
+			static SlackErrorBehavior IgnoreAllErrors(SlackResponse response) => SlackErrorBehavior.Ignore;
+			SlackResponse response = await SendRequestAsync<SlackResponse>(ConversationsInviteUrl, message, IgnoreAllErrors, cancellationToken);
 			return response.Ok ? null : response.Error;
 		}
 
@@ -491,9 +533,7 @@ namespace EpicGames.Slack
 		{
 			object message = new { channel_id = channel, user_ids = String.Join(",", userIds) };
 
-			static bool ShouldLogError(SlackResponse response) => !response.Ok && !String.Equals(response.Error, "already_in_channel", StringComparison.Ordinal);
-
-			await SendRequestAsync<SlackResponse>(AdminConversationsInviteUrl, message, ShouldLogError, cancellationToken);
+			await SendRequestAsync<SlackResponse>(AdminConversationsInviteUrl, message, cancellationToken);
 		}
 
 		#endregion
@@ -517,11 +557,13 @@ namespace EpicGames.Slack
 		/// <returns>User profile</returns>
 		public async Task<SlackUser> GetUserAsync(string userId, CancellationToken cancellationToken = default)
 		{
-			using (HttpRequestMessage request = new HttpRequestMessage(HttpMethod.Post, $"{UsersInfoUrl}?user={userId}"))
+			using HttpRequestMessage request = new (HttpMethod.Post, $"{UsersInfoUrl}?user={userId}");
+			UserResponse response = await SendRequestAsync<UserResponse>(request, "", cancellationToken);
+			if (!response.Ok || response.User == null)
 			{
-				UserResponse response = await SendRequestAsync<UserResponse>(request, "", ShouldLogError, cancellationToken);
-				return response.User!;
+				throw new SlackException(response.Error ?? "unknown");
 			}
+			return response.User;
 		}
 
 		/// <summary>
@@ -534,13 +576,9 @@ namespace EpicGames.Slack
 		{
 			try
 			{
-				using (HttpRequestMessage request = new HttpRequestMessage(HttpMethod.Post, $"{UsersLookupByEmailUrl}?email={email}"))
-				{
-					static bool ShouldLogError(UserResponse response) => !response.Ok && !String.Equals(response.Error, "users_not_found", StringComparison.Ordinal);
-
-					UserResponse response = await SendRequestAsync<UserResponse>(request, "", ShouldLogError, cancellationToken);
-					return response.User;
-				}
+				using HttpRequestMessage request = new (HttpMethod.Post, $"{UsersLookupByEmailUrl}?email={email}");
+				UserResponse response = await SendRequestAsync<UserResponse>(request, "", cancellationToken);
+				return response.User;
 			}
 			catch (Exception ex)
 			{
