Catch and ignore ObjectDisposedException in the single-threaded task scheduler, as it may trigger before the thread shuts down.

#Synced-CL 50482548

diff --git a/Engine/Source/Programs/Shared/EpicGames.Core/SingleThreadTaskScheduler.cs b/Engine/Source/Programs/Shared/EpicGames.Core/SingleThreadTaskScheduler.cs
--- a/Engine/Source/Programs/Shared/EpicGames.Core/SingleThreadTaskScheduler.cs
+++ b/Engine/Source/Programs/Shared/EpicGames.Core/SingleThreadTaskScheduler.cs
@@ -77,6 +77,10 @@ public sealed class SingleThreadTaskScheduler : TaskScheduler, IDisposable
 			{
 				// Don't do anything, we'll exit the loop immediately.
 			}
+			catch (ObjectDisposedException)
+			{
+				// Don't do anything, we'll exit the loop immediately.
+			}
 		}
 	}
 
