CbObject Serialization - Fixed issue with serializing byte[] properties (readonlymemory<byte> was fine and any type with custom converters were fine).

Issue was when generating IL we can not rely on implicit conversion from byte[] to readonlymemory.

Added test.

#Synced-CL 50250157

diff --git a/Engine/Source/Programs/Shared/EpicGames.Serialization.Tests/SerializerTests.cs b/Engine/Source/Programs/Shared/EpicGames.Serialization.Tests/SerializerTests.cs
--- a/Engine/Source/Programs/Shared/EpicGames.Serialization.Tests/SerializerTests.cs
+++ b/Engine/Source/Programs/Shared/EpicGames.Serialization.Tests/SerializerTests.cs
@@ -218,5 +218,23 @@ namespace EpicGames.Serialization.Tests
 			Assert.AreEqual(123, output2.Dictionary["hello"]);
 			Assert.AreEqual(456, output2.Dictionary["world"]);
 		}
+
+		class ObjectWithByteArray
+		{
+			[CbField("b")]
+			public byte[] Bytes { get; set; } = null!;
+		}
+
+		[TestMethod]
+		public void ObjectWithByteArrayTest()
+		{
+			ObjectWithByteArray input = new ObjectWithByteArray();
+			byte[] b = Encoding.ASCII.GetBytes("test-string");
+			input.Bytes = b;
+			CbObject obj = CbSerializer.Serialize(input);
+
+			ObjectWithByteArray output = CbSerializer.Deserialize<ObjectWithByteArray>(obj.AsField());
+			CollectionAssert.AreEqual(b, output.Bytes.ToArray());
+		}
 	}
 }
diff --git a/Engine/Source/Programs/Shared/EpicGames.Serialization/CbWriter.cs b/Engine/Source/Programs/Shared/EpicGames.Serialization/CbWriter.cs
--- a/Engine/Source/Programs/Shared/EpicGames.Serialization/CbWriter.cs
+++ b/Engine/Source/Programs/Shared/EpicGames.Serialization/CbWriter.cs
@@ -1175,7 +1175,7 @@ namespace EpicGames.Serialization
 		/// </summary>
 		/// <param name="writer">Writer for output data</param>
 		/// <param name="value">Value to be written</param>
-		public static void WriteBinaryArrayValue(this ICbWriter writer, ReadOnlyMemory<byte> value) => writer.WriteBinarySpanValue(value.Span);
+		public static void WriteBinaryArrayValue(this ICbWriter writer, byte[] value) => writer.WriteBinarySpanValue(value);
 
 		/// <summary>
 		/// Writes a named binary value
@@ -1183,6 +1183,6 @@ namespace EpicGames.Serialization
 		/// <param name="writer">Writer for output data</param>
 		/// <param name="name">Name of the field</param>
 		/// <param name="value">Value to be written</param>
-		public static void WriteBinaryArray(this ICbWriter writer, CbFieldName name, ReadOnlyMemory<byte> value) => writer.WriteBinarySpan(name, value.Span);
+		public static void WriteBinaryArray(this ICbWriter writer, CbFieldName name, byte[] value) => writer.WriteBinarySpan(name, value);
 	}
 }
\ No newline at end of file
