[Horde][[Analytics] Updating key stats record.

#Synced-CL 50425380

diff --git a/Engine/Source/Programs/Shared/EpicGames.Analytics.PerformanceTrends/KeyStatTelemetryRecord.cs b/Engine/Source/Programs/Shared/EpicGames.Analytics.PerformanceTrends/KeyStatTelemetryRecord.cs
--- a/Engine/Source/Programs/Shared/EpicGames.Analytics.PerformanceTrends/KeyStatTelemetryRecord.cs
+++ b/Engine/Source/Programs/Shared/EpicGames.Analytics.PerformanceTrends/KeyStatTelemetryRecord.cs
@@ -10,7 +10,7 @@ namespace EpicGames.Analytics
 	/// Telemetry record that describes key stats for performance tests.
 	/// </summary>
 	[AnalyticsTableGen]
-	[Table("personal_schema.jgamble.state_performance_key_stat")]
+	[Table("ingest.horde.state_performance_key_stats")]
 	public record KeyStatTelemetryRecord : PerformanceTrendTelemetry
 	{
 		/// <summary>
@@ -28,7 +28,7 @@ namespace EpicGames.Analytics
 		/// <summary>
 		/// Average GPU time.
 		/// </summary>
-		[Column("gpu_time_avg")]
+		[Column("gp_utime_avg")]
 		public float GpuTimeAvg { get; init; }
 
 		/// <summary>
@@ -43,12 +43,6 @@ namespace EpicGames.Analytics
 		[Column("dynamic_resolution_vsm_nanite")]
 		public int DynamicResolutionVsmNanite { get; init; }
 
-		/// <summary>
-		/// Render thread time on critical path.
-		/// </summary>
-		[Column("render_thread_time_critical_path")]
-		public float RenderThreadTimeCriticalPath { get; init; }
-
 		/// <summary>
 		/// MVP.
 		/// </summary>
@@ -79,12 +73,6 @@ namespace EpicGames.Analytics
 		[Column("render_threadtime_avg")]
 		public float RenderThreadtimeAvg { get; init; }
 
-		/// <summary>
-		/// Game thead time on critical path.
-		/// </summary>
-		[Column("game_thread_time_critical_path")]
-		public float GameThreadTimeCriticalPath { get; init; }
-
 		/// <summary>
 		/// Dynamic reoslution percentage average.
 		/// </summary>
