EpicGames.Oodle: Restructure to follow the standard pattern for importing native dependencies in dotnet and remove the custom assembly loader code now it's no longer required. Generate as much marshalling code at compile-time as possible. Only include the relevant platform native binaries if we're building for a specific RuntimeIdentifier.

#Synced-CL 50648979

diff --git a/Engine/Source/Programs/Shared/EpicGames.Oodle/EpicGames.Oodle.csproj b/Engine/Source/Programs/Shared/EpicGames.Oodle/EpicGames.Oodle.csproj
--- a/Engine/Source/Programs/Shared/EpicGames.Oodle/EpicGames.Oodle.csproj
+++ b/Engine/Source/Programs/Shared/EpicGames.Oodle/EpicGames.Oodle.csproj
@@ -4,15 +4,18 @@
     <RootNamespace>EpicGames.Compression</RootNamespace>
     <GenerateDocumentationFile>True</GenerateDocumentationFile>
     <AllowUnsafeBlocks>True</AllowUnsafeBlocks>
+
+    <OodleVersion>2.9.10</OodleVersion>
   </PropertyGroup>
 
-  <ItemGroup>
-    <None Condition="'$([System.Runtime.InteropServices.RuntimeInformation]::IsOSPlatform($([System.Runtime.InteropServices.OSPlatform]::Linux)))' == 'true' Or Exists('Sdk\2.9.10\linux\lib\liboo2corelinux64.so.9')" Include="Sdk\2.9.10\linux\lib\liboo2corelinux64.so.9" Link="liboo2corelinux64.so.9" CopyToOutputDirectory="PreserveNewest" Visible="false" />
-    <None Condition="'$([System.Runtime.InteropServices.RuntimeInformation]::IsOSPlatform($([System.Runtime.InteropServices.OSPlatform]::Linux)))' == 'true' Or Exists('Sdk\2.9.10\linuxarm\lib\liboo2corelinuxarm32.so.9')" Include="Sdk\2.9.10\linuxarm\lib\liboo2corelinuxarm32.so.9" Link="liboo2corelinuxarm32.so.9" CopyToOutputDirectory="PreserveNewest" Visible="false" />
-    <None Condition="'$([System.Runtime.InteropServices.RuntimeInformation]::IsOSPlatform($([System.Runtime.InteropServices.OSPlatform]::Linux)))' == 'true' Or Exists('Sdk\2.9.10\linuxarm\lib\liboo2corelinuxarm64.so.9')" Include="Sdk\2.9.10\linuxarm\lib\liboo2corelinuxarm64.so.9" Link="liboo2corelinuxarm64.so.9" CopyToOutputDirectory="PreserveNewest" Visible="false" />
-    <None Condition="'$([System.Runtime.InteropServices.RuntimeInformation]::IsOSPlatform($([System.Runtime.InteropServices.OSPlatform]::OSX)))' == 'true' Or Exists('Sdk\2.9.10\mac\lib\liboo2coremac64.2.9.10.dylib')" Include="Sdk\2.9.10\mac\lib\liboo2coremac64.2.9.10.dylib" Link="liboo2coremac64.2.9.10.dylib" CopyToOutputDirectory="PreserveNewest" Visible="false" />
-    <None Condition="'$([System.Runtime.InteropServices.RuntimeInformation]::IsOSPlatform($([System.Runtime.InteropServices.OSPlatform]::Windows)))' == 'true' Or Exists('Sdk\2.9.10\win\redist\oo2core_9_win64.dll')" Include="Sdk\2.9.10\win\redist\oo2core_9_win64.dll" Link="oo2core_9_win64.dll" CopyToOutputDirectory="PreserveNewest" Visible="false" />
-    <None Condition="'$([System.Runtime.InteropServices.RuntimeInformation]::IsOSPlatform($([System.Runtime.InteropServices.OSPlatform]::Windows)))' == 'true' Or Exists('Sdk\2.9.10\win\redist\oo2core_9_win32.dll')" Include="Sdk\2.9.10\win\redist\oo2core_9_win32.dll" Link="oo2core_9_win32.dll" CopyToOutputDirectory="PreserveNewest" Visible="false" />
-    <None Condition="'$([System.Runtime.InteropServices.RuntimeInformation]::IsOSPlatform($([System.Runtime.InteropServices.OSPlatform]::Windows)))' == 'true' Or Exists('Sdk\2.9.10\winuwp\redist\oo2core_9_winuwparm64.dll')" Include="Sdk\2.9.10\winuwp\redist\oo2core_9_winuwparm64.dll" Link="oo2core_9_winuwparm64.dll" CopyToOutputDirectory="PreserveNewest" Visible="false" />
+  <ItemGroup Condition="'$(RuntimeIdentifier)' == ''">
+    <None Include="$(OodleVersion)\runtimes\**"
+          Link="runtimes\%(RecursiveDir)%(Filename)%(Extension)"
+          CopyToOutputDirectory="PreserveNewest" />
+  </ItemGroup>
+  <ItemGroup Condition="'$(RuntimeIdentifier)' != ''">
+    <None Include="$(OodleVersion)\runtimes\$(RuntimeIdentifier)\native\*"
+          Link="runtimes\$(RuntimeIdentifier)\native\%(Filename)%(Extension)"
+          CopyToOutputDirectory="PreserveNewest" />
   </ItemGroup>
 </Project>
diff --git a/Engine/Source/Programs/Shared/EpicGames.Oodle/Oodle.cs b/Engine/Source/Programs/Shared/EpicGames.Oodle/Oodle.cs
--- a/Engine/Source/Programs/Shared/EpicGames.Oodle/Oodle.cs
+++ b/Engine/Source/Programs/Shared/EpicGames.Oodle/Oodle.cs
@@ -1,7 +1,6 @@
 // Copyright Epic Games, Inc. All Rights Reserved.
 
 using System;
-using System.Reflection;
 using System.Runtime.InteropServices;
 
 namespace EpicGames.Compression
@@ -154,7 +153,7 @@ namespace EpicGames.Compression
 	/// <summary>
 	/// Wraps an instance of the Oodle compressor
 	/// </summary>
-	public sealed class Oodle
+	public sealed partial class Oodle
 	{
 #pragma warning disable CA1712 // Do not prefix enum values with type name
 		enum OodleLZ_Decode_ThreadPhase
@@ -205,106 +204,39 @@ namespace EpicGames.Compression
 		}
 #pragma warning restore IDE1006
 
-		[DllImport("oo2core")]
-		static extern unsafe long OodleLZ_Compress(OodleCompressorType compressor, byte* rawBuf, long rawLen, byte* compBuf, OodleCompressionLevel level, long pOptions, byte[]? dictionaryBase, byte[]? lrm, byte* scratchMem, long scratchSize);
+		[LibraryImport("oo2core")]
+		private static unsafe partial long OodleLZ_Compress(OodleCompressorType compressor, byte* rawBuf, long rawLen, byte* compBuf, OodleCompressionLevel level, long pOptions, byte[]? dictionaryBase, byte[]? lrm, byte* scratchMem, long scratchSize);
 
-		[DllImport("oo2core")]
-		static extern long OodleLZ_GetCompressedBufferSizeNeeded(OodleCompressorType compressor, long rawSize);
+		[LibraryImport("oo2core")]
+		private static partial long OodleLZ_GetCompressedBufferSizeNeeded(OodleCompressorType compressor, long rawSize);
 
-		[DllImport("oo2core")]
-		static extern int OodleLZDecoder_MemorySizeNeeded(OodleCompressorType compressor, long rawLen = -1);
+		[LibraryImport("oo2core")]
+		private static partial int OodleLZDecoder_MemorySizeNeeded(OodleCompressorType compressor, long rawLen = -1);
 
 		[DllImport("oo2core")]
-		static extern int OodleLZ_GetCompressScratchMemBound(OodleCompressorType compressor, OodleCompressionLevel level, int rawLen, OodleConfigValues configValues);
+		private static extern int OodleLZ_GetCompressScratchMemBound(OodleCompressorType compressor, OodleCompressionLevel level, int rawLen, OodleConfigValues configValues);
 
-		[DllImport("oo2core")]
-		static extern unsafe long OodleLZ_Decompress(byte* compBuf, long compBufSize, byte* rawBuf, long rawLen, OodleLZ_FuzzSafe fuzzSafe = OodleLZ_FuzzSafe.OodleLZ_FuzzSafe_Yes, OodleLZ_CheckCRC checkCRC = OodleLZ_CheckCRC.OodleLZ_CheckCRC_No, OodleLZ_Verbosity verbosity = OodleLZ_Verbosity.OodleLZ_Verbosity_None, byte* decBufBase = null, long decBufSize = 0, long fpCallback = 0, long callbackUserData = 0, byte* decoderMemory = null, long decoderMemorySize = 0, OodleLZ_Decode_ThreadPhase threadPhase = OodleLZ_Decode_ThreadPhase.OodleLZ_Decode_Unthreaded);
+		[LibraryImport("oo2core")]
+		private static unsafe partial long OodleLZ_Decompress(byte* compBuf, long compBufSize, byte* rawBuf, long rawLen, OodleLZ_FuzzSafe fuzzSafe = OodleLZ_FuzzSafe.OodleLZ_FuzzSafe_Yes, OodleLZ_CheckCRC checkCRC = OodleLZ_CheckCRC.OodleLZ_CheckCRC_No, OodleLZ_Verbosity verbosity = OodleLZ_Verbosity.OodleLZ_Verbosity_None, byte* decBufBase = null, long decBufSize = 0, long fpCallback = 0, long callbackUserData = 0, byte* decoderMemory = null, long decoderMemorySize = 0, OodleLZ_Decode_ThreadPhase threadPhase = OodleLZ_Decode_ThreadPhase.OodleLZ_Decode_Unthreaded);
 
 		[DllImport("oo2core")]
-		static extern void Oodle_GetConfigValues(OodleConfigValues configValues);
+		private static extern void Oodle_GetConfigValues(OodleConfigValues configValues);
 
 		[DllImport("oo2core")]
-		static extern void Oodle_SetConfigValues(OodleConfigValues configValues);
+		private static extern void Oodle_SetConfigValues(OodleConfigValues configValues);
 
-		static readonly OodleConfigValues s_configValues = new OodleConfigValues();
+		private static readonly OodleConfigValues s_configValues = new OodleConfigValues();
 
 		/// <summary>
-		/// Static constructor. Registers the import resolver for the native Oodle library.
+		/// Static constructor. Sets the backwards compatibility version for Oodle.
 		/// </summary>
 		static Oodle()
 		{
-			NativeLibrary.SetDllImportResolver(typeof(Oodle).Assembly, ImportResolver);
-
 			Oodle_GetConfigValues(s_configValues);
 			s_configValues.m_OodleLZ_BackwardsCompatible_MajorVersion = 9;
 			Oodle_SetConfigValues(s_configValues);
 		}
 
-		private static IntPtr ImportResolver(string libraryName, Assembly assembly, DllImportSearchPath? searchpath)
-		{
-			IntPtr libHandle = IntPtr.Zero;
-			if (libraryName == "oo2core")
-			{
-				// we manually load the assembly as the file name is not consistent between platforms thus the automatic searching will not work
-				string assemblyFilename;
-				if (RuntimeInformation.IsOSPlatform(OSPlatform.Windows))
-				{
-					switch (RuntimeInformation.ProcessArchitecture)
-					{
-						case Architecture.X86:
-							assemblyFilename = "oo2core_9_win32.dll";
-							break;
-						case Architecture.X64:
-							assemblyFilename = "oo2core_9_win64.dll";
-							break;
-						case Architecture.Arm64:
-							assemblyFilename = "oo2core_9_winuwparm64.dll";
-							break;
-						default:
-							throw new PlatformNotSupportedException($"Oodle support is not currently implemented for {RuntimeInformation.RuntimeIdentifier}");
-					}
-				}
-				else if (RuntimeInformation.IsOSPlatform(OSPlatform.OSX))
-				{
-					switch (RuntimeInformation.ProcessArchitecture)
-					{
-						case Architecture.X64:
-							assemblyFilename = "liboo2coremac64.2.9.10.dylib";
-							break;
-						case Architecture.Arm64:
-							assemblyFilename = "liboo2coremac64.2.9.10.dylib";
-							break;
-						default:
-							throw new PlatformNotSupportedException($"Oodle support is not currently implemented for {RuntimeInformation.RuntimeIdentifier}");
-					}
-				}
-				else if (RuntimeInformation.IsOSPlatform(OSPlatform.Linux))
-				{
-					switch (RuntimeInformation.ProcessArchitecture)
-					{
-						case Architecture.X64:
-							assemblyFilename = "liboo2corelinux64.so.9";
-							break;
-						case Architecture.Arm:
-							assemblyFilename = "liboo2corelinuxarm32.so.9";
-							break;
-						case Architecture.Arm64:
-							assemblyFilename = "liboo2corelinuxarm64.so.9";
-							break;
-						default:
-							throw new PlatformNotSupportedException($"Oodle support is not currently implemented for {RuntimeInformation.RuntimeIdentifier}");
-					}
-				}
-				else
-				{
-					throw new PlatformNotSupportedException($"Oodle support is not currently implemented for {RuntimeInformation.RuntimeIdentifier}");
-				}
-				IntPtr handle = NativeLibrary.Load(assemblyFilename, assembly, searchpath);
-				libHandle = handle;
-			}
-			return libHandle;
-		}
-
 		/// <summary>
 		/// Compress a block of data
 		/// </summary>
