Horde:Fix CA1307 and IDE0059 warnings in ShaderIssueHandler and related tests
#rnx

#Synced-CL 0

diff --git a/Engine/Source/Programs/Shared/EpicGames.Horde/Issues/Handlers/ShaderIssueHandler.cs b/Engine/Source/Programs/Shared/EpicGames.Horde/Issues/Handlers/ShaderIssueHandler.cs
--- a/Engine/Source/Programs/Shared/EpicGames.Horde/Issues/Handlers/ShaderIssueHandler.cs
+++ b/Engine/Source/Programs/Shared/EpicGames.Horde/Issues/Handlers/ShaderIssueHandler.cs
@@ -109,7 +109,7 @@ namespace EpicGames.Horde.Issues.Handlers
 			int lastSlash = materialPath.LastIndexOf('/');
 			string name = lastSlash >= 0 ? materialPath.Substring(lastSlash + 1) : materialPath;
 
-			int dotIndex = name.IndexOf('.');
+			int dotIndex = name.IndexOf('.', System.StringComparison.Ordinal);
 			if (dotIndex > 0)
 			{
 				name = name.Substring(0, dotIndex);
diff --git a/Engine/Source/Programs/Shared/EpicGames.LogMatchers.Tests/AutomationToolLogEventParserTests.cs b/Engine/Source/Programs/Shared/EpicGames.LogMatchers.Tests/AutomationToolLogEventParserTests.cs
--- a/Engine/Source/Programs/Shared/EpicGames.LogMatchers.Tests/AutomationToolLogEventParserTests.cs
+++ b/Engine/Source/Programs/Shared/EpicGames.LogMatchers.Tests/AutomationToolLogEventParserTests.cs
@@ -783,8 +783,7 @@ namespace EpicGames.LogMatchers.Tests
 
 				// Verify NO material property is set for non-material errors
 				LogEvent shaderEvent = logEvents[0];
-				object? materialValue;
-				bool hasMaterial = shaderEvent.TryGetProperty("material", out materialValue);
+				bool hasMaterial = shaderEvent.TryGetProperty("material", out _);
 				Assert.IsFalse(hasMaterial);
 			}
 		}
