Prevented stepping in and out of StaticClass/StaticStruct and other generated functions in the debugger.

diff --git a/Engine/Source/Programs/Shared/EpicGames.UHT/Exporters/CodeGen/UhtHeaderCodeGeneratorHFile.cs b/Engine/Source/Programs/Shared/EpicGames.UHT/Exporters/CodeGen/UhtHeaderCodeGeneratorHFile.cs
--- a/Engine/Source/Programs/Shared/EpicGames.UHT/Exporters/CodeGen/UhtHeaderCodeGeneratorHFile.cs
+++ b/Engine/Source/Programs/Shared/EpicGames.UHT/Exporters/CodeGen/UhtHeaderCodeGeneratorHFile.cs
@@ -510,6 +510,7 @@ namespace EpicGames.UHT.Exporters.CodeGen
 				{
 					builder.Append("\tfriend struct ::Z_Construct_UScriptStruct_").Append(scriptStruct.SourceName).Append("_Statics; \\\r\n");
 					builder.Append('\t');
+					builder.Append("UE_NODEBUG ");
 					if (!scriptStruct.ScriptStructFlags.HasAnyFlags(EStructFlags.RequiredAPI) &&
 						(Module.Module.AlwaysExportStructs || scriptStruct.ScriptStructExportFlags.HasAnyFlags(UhtScriptStructExportFlags.MinimalAPI)))
 					{
@@ -790,11 +791,11 @@ namespace EpicGames.UHT.Exporters.CodeGen
 				// Forward declare the StaticEnum<> specialization for enum classes
 				if (Module.Module.AlwaysExportEnums || enumObj.EnumExportFlags.HasAnyFlags(UhtEnumExportFlags.MinimalAPI))
 				{
-					builder.Append("template<> ").Append(Module.NonAttributedApi).Append("UEnum* StaticEnum<").Append(enumObj.Namespace.FullSourceName).Append(enumObj.CppType).Append(">();\r\n");
+					builder.Append("template<> UE_NODEBUG ").Append(Module.NonAttributedApi).Append("UEnum* StaticEnum<").Append(enumObj.Namespace.FullSourceName).Append(enumObj.CppType).Append(">();\r\n");
 				}
 				else
 				{
-					builder.Append("template<> ").Append("UEnum* StaticEnum<").Append(enumObj.Namespace.FullSourceName).Append(enumObj.CppType).Append(">();\r\n");
+					builder.Append("template<> UE_NODEBUG ").Append("UEnum* StaticEnum<").Append(enumObj.Namespace.FullSourceName).Append(enumObj.CppType).Append(">();\r\n");
 				}
 			}
 
diff --git a/Engine/Source/Runtime/CoreUObject/Public/UObject/Field.h b/Engine/Source/Runtime/CoreUObject/Public/UObject/Field.h
--- a/Engine/Source/Runtime/CoreUObject/Public/UObject/Field.h
+++ b/Engine/Source/Runtime/CoreUObject/Public/UObject/Field.h
@@ -310,14 +310,14 @@ public: \
 		: Super(EC_InternalUseOnlyConstructor, InClass) \
 	{ \
 	} \
-	static TRequiredAPI FFieldClass* StaticClass(); \
+	UE_NODEBUG static TRequiredAPI FFieldClass* StaticClass(); \
 	UE_DECLARE_FIELD_CLASS_STATIC_CONSTINIT(TClass, TRequiredAPI) \
 	static TRequiredAPI FField* Construct(const FFieldVariant& InOwner, const FName& InName, EObjectFlags InObjectFlags); \
-	inline static constexpr EClassCastFlags StaticClassCastFlagsPrivate() \
+	UE_REWRITE static constexpr EClassCastFlags StaticClassCastFlagsPrivate() \
 	{ \
 		return TStaticFlags; \
 	} \
-	inline static constexpr EClassCastFlags StaticClassCastFlags() \
+	UE_REWRITE static constexpr EClassCastFlags StaticClassCastFlags() \
 	{ \
 		return TStaticFlags | Super::StaticClassCastFlags(); \
 	} \
diff --git a/Engine/Source/Runtime/CoreUObject/Public/UObject/ObjectMacros.h b/Engine/Source/Runtime/CoreUObject/Public/UObject/ObjectMacros.h
--- a/Engine/Source/Runtime/CoreUObject/Public/UObject/ObjectMacros.h
+++ b/Engine/Source/Runtime/CoreUObject/Public/UObject/ObjectMacros.h
@@ -1898,22 +1898,22 @@ public: \
 	/** Typedef for {{ typedef-type }}. */ \
 	typedef TClass ThisClass;\
 	/** Returns a UClass object representing this class at runtime */ \
-	inline static UClass* StaticClass() \
+	UE_REWRITE static UClass* StaticClass() \
 	{ \
 		return TPrivateAccessor(); \
 	} \
 	/** Returns the package this class belongs in */ \
-	inline static const TCHAR* StaticPackage() \
+	UE_REWRITE static const TCHAR* StaticPackage() \
 	{ \
 		return TPackage; \
 	} \
 	/** Returns the unique cast flag for this class if any */ \
-	inline constexpr static EClassCastFlags StaticClassCastFlags() \
+	UE_REWRITE constexpr static EClassCastFlags StaticClassCastFlags() \
 	{ \
 		return TStaticCastFlags; \
 	} \
 	/** Returns all cast flags for this class, including those inherited from parent classes */ \
-	inline consteval static EClassCastFlags StaticAllClassCastFlags() \
+	UE_REWRITE consteval static EClassCastFlags StaticAllClassCastFlags() \
 	{ \
 		if constexpr (std::is_same_v<ThisClass, Super>) { return StaticClassCastFlags(); } \
 		return StaticClassCastFlags() | Super::StaticAllClassCastFlags(); \
