Refactor out the HTTP delegating handler support for logging into and refreshing OIDC tokens automatically from EpicGames.Horde to EpicGames.OIDC.

#Synced-CL 50766644

#Synced-CL 50766644

diff --git a/Engine/Source/Programs/Horde/HordeCmd/Commands/LoginCommand.cs b/Engine/Source/Programs/Horde/HordeCmd/Commands/LoginCommand.cs
--- a/Engine/Source/Programs/Horde/HordeCmd/Commands/LoginCommand.cs
+++ b/Engine/Source/Programs/Horde/HordeCmd/Commands/LoginCommand.cs
@@ -4,7 +4,6 @@ using System.ComponentModel;
 using EpicGames.Core;
 using EpicGames.Horde;
 using EpicGames.Horde.Server;
-using Microsoft.Extensions.DependencyInjection;
 using Microsoft.Extensions.Logging;
 
 namespace Horde.Commands
@@ -24,12 +23,10 @@ namespace Horde.Commands
 		[Description("Echo the bearer token acquired from the server to stdout")]
 		public bool Token { get; set; }
 
-		readonly IServiceProvider _serviceProvider;
 		readonly IHordeClient _hordeClient;
 
-		public LoginCommand(IServiceProvider serviceProvider, IHordeClient hordeClient)
+		public LoginCommand(IHordeClient hordeClient)
 		{
-			_serviceProvider = serviceProvider;
 			_hordeClient = hordeClient;
 		}
 
@@ -43,9 +40,7 @@ namespace Horde.Commands
 
 			if (Token)
 			{
-				HordeHttpAuthHandlerState state = _serviceProvider.GetRequiredService<HordeHttpAuthHandlerState>();
-
-				string? accessToken = await state.GetAccessTokenAsync(true, CancellationToken.None);
+				string? accessToken = await _hordeClient.GetAccessTokenAsync(true, CancellationToken.None);
 				if (accessToken != null)
 				{
 					Console.WriteLine($"Bearer {accessToken}");
diff --git a/Engine/Source/Programs/OidcToken/TokenService.cs b/Engine/Source/Programs/OidcToken/TokenService.cs
--- a/Engine/Source/Programs/OidcToken/TokenService.cs
+++ b/Engine/Source/Programs/OidcToken/TokenService.cs
@@ -119,15 +119,15 @@ namespace OidcToken
 			{
 				case TokenServiceOptions.TokenServiceMode.Query:
 				{
-					await OutputStatus(providerName, TokenManager.GetStatusForProvider(providerName));
+					await OutputStatus(providerName, await TokenManager.GetStatusForProvider(providerName));
 					break;
 				}
 				case TokenServiceOptions.TokenServiceMode.GetToken:
 				{
-					OidcStatus status = TokenManager.GetStatusForProvider(providerName);
+					OidcStatus status = await TokenManager.GetStatusForProvider(providerName);
 					Logger.LogInformation("Determined status of provider {ProviderName} was {Status}", providerName, status);
 
-					OidcTokenInfo? tokenInfo;
+					IOidcTokenInfo? tokenInfo;
 					if (status == OidcStatus.NotLoggedIn && !Settings.CurrentValue.Unattended)
 					{
 						Logger.LogInformation("Logging in to provider {ProviderName}", providerName);
@@ -240,7 +240,7 @@ namespace OidcToken
 			{
 				// verify that the refresh token is valid and actually able to generate a access token
 				await TokenManager.GetAccessToken(service);
-				OidcStatus refreshedStatus = TokenManager.GetStatusForProvider(service);
+				OidcStatus refreshedStatus = await TokenManager.GetStatusForProvider(service);
 				Logger.LogInformation("Determined status of provider {ProviderName} was {Status}", service, refreshedStatus);
 			}
 
diff --git a/Engine/Source/Programs/Shared/EpicGames.Horde.Tests/HordeHttpTests.cs b/Engine/Source/Programs/Shared/EpicGames.Horde.Tests/HordeHttpTests.cs
--- a/Engine/Source/Programs/Shared/EpicGames.Horde.Tests/HordeHttpTests.cs
+++ b/Engine/Source/Programs/Shared/EpicGames.Horde.Tests/HordeHttpTests.cs
@@ -13,7 +13,6 @@ using EpicGames.Core;
 using EpicGames.Horde.Server;
 using EpicGames.OIDC;
 using Microsoft.Extensions.Logging;
-using Microsoft.Extensions.Options;
 using Microsoft.VisualStudio.TestTools.UnitTesting;
 
 namespace EpicGames.Horde.Tests;
@@ -122,7 +121,7 @@ public class HordeHttpAuthHandlerTests
 	}
 
 	[SuppressMessage("Reliability", "CA2000:Dispose objects before losing scope")]
-	private static (HttpClient client, FakeServerWithAuth server) CreateClientServer(FakeOidcTokenManager oidcTokenManager, HordeOptions? hordeOptions = null)
+	private (HttpClient client, FakeServerWithAuth server) CreateClientServer(FakeOidcTokenManager oidcTokenManager)
 	{
 		FakeServerWithAuth server = new(oidcTokenManager);
 		using ILoggerFactory loggerFactory = LoggerFactory.Create(builder =>
@@ -131,11 +130,10 @@ public class HordeHttpAuthHandlerTests
 			builder.AddSimpleConsole(options => { options.SingleLine = true; });
 		});
 
-		ILogger<HordeHttpAuthHandlerState> logger = loggerFactory.CreateLogger<HordeHttpAuthHandlerState>();
-		OptionsWrapper<HordeOptions> options = new(hordeOptions ?? new HordeOptions());
-		InMemoryTokenStore inMemoryTokenStore = new();
-		HordeHttpAuthHandlerState state = new HordeHttpAuthHandlerState(server, new Uri("http://fake-server"), options, logger, inMemoryTokenStore, oidcTokenManager);
-		HordeHttpAuthHandler authHandler = new HordeHttpAuthHandler(server, state, options);
+		ILogger<OidcAuthState> logger = loggerFactory.CreateLogger<OidcAuthState>();
+		IOidcAuthStateConfig oidcAuthStateConfig = new OidcAuthStateConfig("Horde");
+		OidcAuthState state = new OidcAuthState(oidcTokenManager, new NullStaticAccessTokenProvider(), _clock, oidcAuthStateConfig, logger);
+		HttpOidcAuthHandler authHandler = new HttpOidcAuthHandler(server, state, new HttpOidcAuthHandlerConfig("Horde-AllowInteractiveLogin"));
 
 		authHandler.InnerHandler = server;
 		HttpClient client = new(authHandler);
diff --git a/Engine/Source/Programs/Shared/EpicGames.Horde/Auth/HordeAuthConfigProvider.cs b/Engine/Source/Programs/Shared/EpicGames.Horde/Auth/HordeAuthConfigProvider.cs
new file mode 100644
--- /dev/null
+++ b/Engine/Source/Programs/Shared/EpicGames.Horde/Auth/HordeAuthConfigProvider.cs
@@ -0,0 +1,71 @@
+// Copyright Epic Games, Inc. All Rights Reserved.
+
+using System;
+using System.Diagnostics.CodeAnalysis;
+using System.Net.Http;
+using System.Net.Http.Json;
+using System.Text.Json;
+using System.Threading;
+using System.Threading.Tasks;
+using EpicGames.Horde.Server;
+using EpicGames.OIDC;
+using Microsoft.Extensions.Caching.Memory;
+using Microsoft.Extensions.Logging;
+using Microsoft.Extensions.Options;
+
+namespace EpicGames.Horde.Auth;
+
+internal sealed class HordeAuthConfigProvider(
+	HttpMessageHandler httpMessageHandler,
+	Uri serverUri,
+	ILogger<HordeAuthConfigProvider> logger
+) : IOidcAuthStateConfig, IDisposable
+{
+	public async Task<string> GetOidcProviderAsync(CancellationToken cancellationToken)
+	{
+		GetAuthConfigResponse authConfig = await GetAuthConfigResponseAsync(cancellationToken);
+		return GetOidcProvider(authConfig);
+	}
+
+	public static string GetOidcProvider(GetAuthConfigResponse authConfig) => authConfig.ProfileName ?? "Horde";
+
+	public async Task<GetAuthConfigResponse> GetAuthConfigResponseAsync(CancellationToken cancellationToken)
+	{
+		const string CacheKey = "auth";
+		if (_memoryCache.TryGetValue(CacheKey, out GetAuthConfigResponse? auth))
+		{
+			return auth!;
+		}
+
+		GetAuthConfigResponse? authConfig;
+		using (HttpClient httpClient = new HttpClient(_httpMessageHandler, false))
+		{
+			httpClient.BaseAddress = _serverUri;
+			_logger.LogDebug("Retrieving auth configuration for {Server}", _serverUri);
+
+			JsonSerializerOptions jsonOptions = new JsonSerializerOptions();
+			HordeHttpClient.ConfigureJsonSerializer(jsonOptions);
+
+			authConfig = await httpClient.GetFromJsonAsync<GetAuthConfigResponse>("api/v1/server/auth", jsonOptions, cancellationToken);
+			if (authConfig == null)
+			{
+				throw new Exception($"Invalid response from server");
+			}
+		}
+
+		TimeSpan expiresAfter = TimeSpan.FromMinutes(1);
+		_memoryCache.Set(CacheKey, authConfig, expiresAfter);
+		return authConfig;
+	}
+
+	public void Dispose()
+	{
+		_memoryCache.Dispose();
+	}
+
+	[SuppressMessage("Usage", "CA2213:Disposable fields should be disposed", Justification = "Lifetime is managed by caller")]
+	private readonly HttpMessageHandler _httpMessageHandler = httpMessageHandler;
+	private readonly Uri _serverUri = serverUri;
+	private readonly ILogger<HordeAuthConfigProvider> _logger = logger;
+	private readonly MemoryCache _memoryCache = new MemoryCache(Options.Create(new MemoryCacheOptions()));
+}
diff --git a/Engine/Source/Programs/Shared/EpicGames.Horde/Auth/HordeOidcTokenManager.cs b/Engine/Source/Programs/Shared/EpicGames.Horde/Auth/HordeOidcTokenManager.cs
new file mode 100644
--- /dev/null
+++ b/Engine/Source/Programs/Shared/EpicGames.Horde/Auth/HordeOidcTokenManager.cs
@@ -0,0 +1,104 @@
+// Copyright Epic Games, Inc. All Rights Reserved.
+
+using System;
+using System.Collections.Generic;
+using System.Linq;
+using System.Threading;
+using System.Threading.Tasks;
+using EpicGames.Horde.Server;
+using EpicGames.OIDC;
+using Microsoft.Extensions.Configuration;
+
+namespace EpicGames.Horde.Auth;
+
+internal sealed class HordeOidcTokenManager(
+	HordeAuthConfigProvider hordeAuthConfigProvider,
+	ITokenStore tokenStore
+) : IOidcTokenManager
+{
+	public async Task<IOidcTokenInfo> GetAccessToken(string providerIdentifier, CancellationToken cancellationToken)
+	{
+		GetAuthConfigResponse authConfig = await _hordeAuthConfigProvider.GetAuthConfigResponseAsync(cancellationToken);
+
+		if (authConfig.Method == AuthMethod.Anonymous)
+		{
+			return new AnonymousOidcTokenInfo();
+		}
+
+		OidcTokenManager oidcTokenManager = CreateOidcTokenManager(providerIdentifier, authConfig);
+		return await oidcTokenManager.GetAccessToken(providerIdentifier, cancellationToken);
+	}
+
+	public async Task<OidcStatus> GetStatusForProvider(string providerIdentifier, CancellationToken cancellationToken)
+	{
+		GetAuthConfigResponse authConfig = await _hordeAuthConfigProvider.GetAuthConfigResponseAsync(cancellationToken);
+
+		if (authConfig.Method == AuthMethod.Anonymous)
+		{
+			return OidcStatus.Connected;
+		}
+
+		OidcTokenManager oidcTokenManager = CreateOidcTokenManager(providerIdentifier, authConfig);
+		return await oidcTokenManager.GetStatusForProvider(providerIdentifier, cancellationToken);
+	}
+
+	public async Task<IOidcTokenInfo> LoginAsync(string providerIdentifier, CancellationToken cancellationToken)
+	{
+		GetAuthConfigResponse authConfig = await _hordeAuthConfigProvider.GetAuthConfigResponseAsync(cancellationToken);
+
+		if (authConfig.Method == AuthMethod.Anonymous)
+		{
+			return new AnonymousOidcTokenInfo();
+		}
+
+		OidcTokenManager oidcTokenManager = CreateOidcTokenManager(providerIdentifier, authConfig);
+		return await oidcTokenManager.LoginAsync(providerIdentifier, cancellationToken);
+	}
+
+	public async Task<IOidcTokenInfo?> TryGetAccessToken(string providerIdentifier, CancellationToken cancellationToken)
+	{
+		GetAuthConfigResponse authConfig = await _hordeAuthConfigProvider.GetAuthConfigResponseAsync(cancellationToken);
+
+		if (authConfig.Method == AuthMethod.Anonymous)
+		{
+			return new AnonymousOidcTokenInfo();
+		}
+
+		OidcTokenManager oidcTokenManager = CreateOidcTokenManager(providerIdentifier, authConfig);
+		return await oidcTokenManager.TryGetAccessToken(providerIdentifier, cancellationToken);
+	}
+
+	private readonly HordeAuthConfigProvider _hordeAuthConfigProvider = hordeAuthConfigProvider;
+	private readonly ITokenStore _tokenStore = tokenStore;
+
+	private OidcTokenManager CreateOidcTokenManager(string providerIdentifier, GetAuthConfigResponse authConfig)
+	{
+		string? localRedirectUrl = authConfig.LocalRedirectUrls?.FirstOrDefault();
+		if (String.IsNullOrEmpty(authConfig.ServerUrl) || String.IsNullOrEmpty(localRedirectUrl))
+		{
+			throw new Exception("No auth server configuration found");
+		}
+		string oidcProvider = HordeAuthConfigProvider.GetOidcProvider(authConfig);
+
+		if (!String.Equals(oidcProvider, providerIdentifier, StringComparison.Ordinal))
+		{
+			throw new Exception("Provider identifier mismatch");
+		}
+
+		Dictionary<string, string?> values = [];
+		values[$"Providers:{oidcProvider}:DisplayName"] = oidcProvider;
+		values[$"Providers:{oidcProvider}:ServerUri"] = authConfig.ServerUrl;
+		values[$"Providers:{oidcProvider}:ClientId"] = authConfig.ClientId;
+		values[$"Providers:{oidcProvider}:RedirectUri"] = localRedirectUrl;
+		if (authConfig.Scopes is { Length: > 0 })
+		{
+			values[$"Providers:{oidcProvider}:Scopes"] = String.Join(" ", authConfig.Scopes);
+		}
+
+		ConfigurationBuilder builder = new();
+		builder.AddInMemoryCollection(values);
+		IConfiguration configuration = builder.Build();
+
+		return OidcTokenManager.CreateTokenManager(configuration, _tokenStore, [oidcProvider]);
+	}
+}
diff --git a/Engine/Source/Programs/Shared/EpicGames.Horde/Auth/HordeStaticAccessTokenProvider.cs b/Engine/Source/Programs/Shared/EpicGames.Horde/Auth/HordeStaticAccessTokenProvider.cs
new file mode 100644
--- /dev/null
+++ b/Engine/Source/Programs/Shared/EpicGames.Horde/Auth/HordeStaticAccessTokenProvider.cs
@@ -0,0 +1,42 @@
+// Copyright Epic Games, Inc. All Rights Reserved.
+
+using System;
+using EpicGames.OIDC;
+using Microsoft.Extensions.Options;
+
+namespace EpicGames.Horde.Auth;
+
+internal sealed class HordeStaticAccessTokenProvider(IOptionsSnapshot<HordeOptions> options) : IStaticAccessTokenProvider
+{
+	public string? AccessToken
+	{
+		get
+		{
+			// If an explicit access token is specified, just use that
+			if (_options.Value.AccessToken != null)
+			{
+				return _options.Value.AccessToken;
+			}
+
+			// Check environment variables for an access token matching the current server
+			string? hordeUrlEnvVar = Environment.GetEnvironmentVariable(HordeHttpClient.HordeUrlEnvVarName);
+			if (!String.IsNullOrEmpty(hordeUrlEnvVar))
+			{
+				Uri hordeUrl = new Uri(hordeUrlEnvVar);
+				if (_options.Value.ServerUrl == null || String.Equals(_options.Value.ServerUrl.Host, hordeUrl.Host, StringComparison.OrdinalIgnoreCase))
+				{
+					string? hordeToken = Environment.GetEnvironmentVariable(HordeHttpClient.HordeTokenEnvVarName);
+					if (!String.IsNullOrEmpty(hordeToken))
+					{
+						return hordeToken;
+					}
+				}
+			}
+
+			// Otherwise we need to find the access token dynamically
+			return null;
+		}
+	}
+
+	private readonly IOptionsSnapshot<HordeOptions> _options = options;
+}
diff --git a/Engine/Source/Programs/Shared/EpicGames.Horde/HordeClient.cs b/Engine/Source/Programs/Shared/EpicGames.Horde/HordeClient.cs
--- a/Engine/Source/Programs/Shared/EpicGames.Horde/HordeClient.cs
+++ b/Engine/Source/Programs/Shared/EpicGames.Horde/HordeClient.cs
@@ -9,6 +9,7 @@ using System.Threading;
 using System.Threading.Tasks;
 using EpicGames.Core;
 using EpicGames.Horde.Artifacts;
+using EpicGames.Horde.Auth;
 using EpicGames.Horde.Compute;
 using EpicGames.Horde.Compute.Clients;
 using EpicGames.Horde.Logs;
@@ -19,6 +20,7 @@ using EpicGames.Horde.Storage;
 using EpicGames.Horde.Storage.Backends;
 using EpicGames.Horde.Storage.Bundles;
 using EpicGames.Horde.Tools;
+using EpicGames.OIDC;
 using Grpc.Core;
 using Grpc.Net.Client;
 using Grpc.Net.Client.Configuration;
@@ -324,8 +326,13 @@ namespace EpicGames.Horde
 	/// </summary>
 	class HordeClientWithDynamicCredentials : HordeClient
 	{
-		readonly HordeHttpAuthHandlerState _authHandlerState;
+		private const string AllowInteractiveLoginHeaderKey = "Horde-AllowInteractiveLogin";
+		private static readonly HttpRequestOptionsKey<bool> s_allowInteractiveLogin = new(AllowInteractiveLoginHeaderKey);
+
+		readonly OidcAuthState _authHandlerState;
 		readonly IHordeHttpMessageHandler _baseHttpMessageHandler;
+		private readonly ITokenStore _tokenStore;
+		private readonly HordeAuthConfigProvider _hordeAuthConfigProvider;
 
 		[SuppressMessage("Usage", "CA2213:Disposable fields should be disposed", Justification = "Lifetime is managed by caller")]
 		readonly HttpMessageHandler _authHttpMessageHandler;
@@ -338,8 +345,15 @@ namespace EpicGames.Horde
 		{
 			_baseHttpMessageHandler = baseHttpMessageHandler;
 
-			_authHandlerState = new HordeHttpAuthHandlerState(_baseHttpMessageHandler.Instance, serverUrl, hordeOptions, loggerFactory.CreateLogger<HordeHttpAuthHandlerState>());
-			_authHttpMessageHandler = new HordeHttpAuthHandler(_baseHttpMessageHandler.Instance, _authHandlerState, hordeOptions);
+			_tokenStore = TokenStoreFactory.CreateTokenStore();
+			_hordeAuthConfigProvider =
+				new HordeAuthConfigProvider(_baseHttpMessageHandler.Instance, serverUrl, loggerFactory.CreateLogger<HordeAuthConfigProvider>());
+			IOidcTokenManager oidcTokenManager = new HordeOidcTokenManager(_hordeAuthConfigProvider, _tokenStore);
+			IStaticAccessTokenProvider staticAccessTokenProvider = new HordeStaticAccessTokenProvider(hordeOptions);
+			HttpOidcAuthHandlerConfig httpOidcAuthHandlerConfig = HordeOptionsToIHttpOidcAuthHandlerConfig(hordeOptions.Value);
+			IClock clock = new DefaultClock();
+			_authHandlerState = new OidcAuthState(oidcTokenManager, staticAccessTokenProvider, clock, _hordeAuthConfigProvider, loggerFactory.CreateLogger<OidcAuthState>());
+			_authHttpMessageHandler = new HttpOidcAuthHandler(_baseHttpMessageHandler.Instance, _authHandlerState, httpOidcAuthHandlerConfig);
 
 			_authHandlerState.OnStateChanged += NotifyAuthStateChanged;
 		}
@@ -350,6 +364,8 @@ namespace EpicGames.Horde
 			_authHandlerState.OnStateChanged -= NotifyAuthStateChanged;
 
 			await _authHandlerState.DisposeAsync();
+			_hordeAuthConfigProvider.Dispose();
+			_tokenStore.Dispose();
 			await base.DisposeAsync();
 		}
 
@@ -363,7 +379,7 @@ namespace EpicGames.Horde
 			using (HttpClient httpClient = CreateAuthenticatedHttpClient())
 			{
 				using HttpRequestMessage request = new HttpRequestMessage(HttpMethod.Get, "api/v1/dashboard/challenge");
-				request.Options.Set(HordeHttpAuthHandler.AllowInteractiveLogin, allowLogin);
+				request.Options.Set(s_allowInteractiveLogin, allowLogin);
 
 				using HttpResponseMessage response = await httpClient.SendAsync(request, cancellationToken);
 				return response.IsSuccessStatusCode;
@@ -375,7 +391,7 @@ namespace EpicGames.Horde
 		{
 			try
 			{
-				return _authHandlerState.IsLoggedIn();
+				return _authHandlerState.IsLoggedIn;
 			}
 			catch
 			{
@@ -402,6 +418,9 @@ namespace EpicGames.Horde
 			httpClient.BaseAddress = ServerUrl;
 			return httpClient;
 		}
+
+		private static HttpOidcAuthHandlerConfig HordeOptionsToIHttpOidcAuthHandlerConfig(HordeOptions options) =>
+			new HttpOidcAuthHandlerConfig(AllowInteractiveLoginHeaderKey, options.AllowAuthPrompt);
 	}
 
 	/// <summary>
diff --git a/Engine/Source/Programs/Shared/EpicGames.Horde/HordeHttpAuthHandler.cs b/Engine/Source/Programs/Shared/EpicGames.Horde/HordeHttpAuthHandler.cs
deleted file mode 100644
diff --git a/Engine/Source/Programs/Shared/EpicGames.OIDC/EpicGames.OIDC.csproj b/Engine/Source/Programs/Shared/EpicGames.OIDC/EpicGames.OIDC.csproj
--- a/Engine/Source/Programs/Shared/EpicGames.OIDC/EpicGames.OIDC.csproj
+++ b/Engine/Source/Programs/Shared/EpicGames.OIDC/EpicGames.OIDC.csproj
@@ -12,4 +12,8 @@
     <PackageReference Include="System.Text.Json" Version="8.0.6" />
   </ItemGroup>
 
+  <ItemGroup>
+    <ProjectReference Include="..\EpicGames.Core\EpicGames.Core.csproj" />
+  </ItemGroup>
+
 </Project>
diff --git a/Engine/Source/Programs/Shared/EpicGames.OIDC/FakeOidcTokenManager.cs b/Engine/Source/Programs/Shared/EpicGames.OIDC/FakeOidcTokenManager.cs
--- a/Engine/Source/Programs/Shared/EpicGames.OIDC/FakeOidcTokenManager.cs
+++ b/Engine/Source/Programs/Shared/EpicGames.OIDC/FakeOidcTokenManager.cs
@@ -49,7 +49,7 @@ public class FakeOidcTokenManager : IOidcTokenManager
 	}
 
 	/// <inheritdoc/>
-	public Task<OidcTokenInfo> LoginAsync(string providerIdentifier, CancellationToken cancellationToken = default)
+	public Task<IOidcTokenInfo> LoginAsync(string providerIdentifier, CancellationToken cancellationToken = default)
 	{
 		RefreshToken = "fakeRefreshToken-" + _refreshCounter++;
 		RefreshAccessToken();
@@ -57,38 +57,38 @@ public class FakeOidcTokenManager : IOidcTokenManager
 	}
 
 	/// <inheritdoc/>
-	public Task<OidcTokenInfo> GetAccessToken(string providerIdentifier, CancellationToken cancellationToken = default)
+	public Task<IOidcTokenInfo> GetAccessToken(string providerIdentifier, CancellationToken cancellationToken = default)
 	{
 		throw new NotImplementedException("Method not in use");
 	}
 
 	/// <inheritdoc/>
-	public Task<OidcTokenInfo?> TryGetAccessToken(string providerIdentifier, CancellationToken cancellationToken = default)
+	public Task<IOidcTokenInfo?> TryGetAccessToken(string providerIdentifier, CancellationToken cancellationToken = default)
 	{
 		if (String.IsNullOrEmpty(RefreshToken))
 		{
 			throw new NotLoggedInException();
 		}
 
-		OidcTokenInfo tokenInfo = GetOidcTokenInfo();
+		IOidcTokenInfo tokenInfo = GetOidcTokenInfo();
 
 		// Ensure token is valid at least for another two minutes
 		if (tokenInfo.IsValid(_utcNow().AddMinutes(-2)))
 		{
-			return Task.FromResult<OidcTokenInfo?>(tokenInfo);
+			return Task.FromResult<IOidcTokenInfo?>(tokenInfo);
 		}
 
 		// Access token not valid, simulate a refresh with provider
 		RefreshAccessToken();
 
 		// Assume the updated access token is valid
-		return Task.FromResult<OidcTokenInfo?>(GetOidcTokenInfo());
+		return Task.FromResult<IOidcTokenInfo?>(GetOidcTokenInfo());
 	}
 
 	/// <inheritdoc/>
-	public OidcStatus GetStatusForProvider(string providerIdentifier)
+	public Task<OidcStatus> GetStatusForProvider(string providerIdentifier, CancellationToken cancellationToken)
 	{
-		return OidcTokenClient.GetStatus(RefreshToken, AccessToken, AccessTokenExpiry);
+		return Task.FromResult(OidcTokenClient.GetStatus(RefreshToken, AccessToken, AccessTokenExpiry));
 	}
 
 	/// <summary>
@@ -133,7 +133,7 @@ public class FakeOidcTokenManager : IOidcTokenManager
 		AccessToken = $"{AccessTokenPrefix}-{_accessCounter++}-{AccessTokenExpiry.ToUnixTimeMilliseconds()}-{AccessTokenExpiry.ToString()}";
 	}
 
-	private OidcTokenInfo GetOidcTokenInfo()
+	private IOidcTokenInfo GetOidcTokenInfo()
 	{
 		return new OidcTokenInfo { RefreshToken = RefreshToken, AccessToken = AccessToken, TokenExpiry = AccessTokenExpiry };
 	}
diff --git a/Engine/Source/Programs/Shared/EpicGames.OIDC/HttpOidcAuthHandler.cs b/Engine/Source/Programs/Shared/EpicGames.OIDC/HttpOidcAuthHandler.cs
new file mode 100644
--- /dev/null
+++ b/Engine/Source/Programs/Shared/EpicGames.OIDC/HttpOidcAuthHandler.cs
@@ -0,0 +1,78 @@
+// Copyright Epic Games, Inc. All Rights Reserved.
+
+using System;
+using System.Net;
+using System.Net.Http;
+using System.Net.Http.Headers;
+using System.Threading;
+using System.Threading.Tasks;
+
+namespace EpicGames.OIDC;
+
+/// <summary>
+/// HTTP message handler which automatically refreshes access tokens as required.
+/// </summary>
+public sealed class HttpOidcAuthHandler : DelegatingHandler
+{
+	/// <summary>
+	/// Constructor
+	/// </summary>
+	public HttpOidcAuthHandler(HttpMessageHandler innerHandler, IOidcAuthState authState, HttpOidcAuthHandlerConfig config)
+		: base(innerHandler)
+	{
+		_authState = authState;
+		_config = config;
+		_allowInteractiveLogin = new(_config.AllowInteractiveLoginHeaderKey);
+	}
+
+	private readonly IOidcAuthState _authState;
+	private readonly HttpOidcAuthHandlerConfig _config;
+	private readonly HttpRequestOptionsKey<bool> _allowInteractiveLogin;
+
+	/// <inheritdoc/>
+	protected override async Task<HttpResponseMessage> SendAsync(HttpRequestMessage request, CancellationToken cancellationToken)
+	{
+		// If the request already has a custom auth header, send the request as it is.
+		if (request.Headers.Authorization != null)
+		{
+			return await base.SendAsync(request, cancellationToken);
+		}
+
+		// Check whether the request specifically allows interactive auth, otherwise fall back to the default.
+		bool allowInteractiveLogin;
+		if (!request.Options.TryGetValue(_allowInteractiveLogin, out allowInteractiveLogin))
+		{
+			allowInteractiveLogin = _config.AllowInteractiveLogin;
+		}
+
+		// Get the current access token and send the request with that.
+		string? accessToken = await _authState.GetAccessTokenAsync(allowInteractiveLogin, cancellationToken);
+		for (int attempt = 0; ; attempt++)
+		{
+			// Attempt to perform the request with this access token.
+			request.Headers.Authorization = (accessToken == null) ? null : new AuthenticationHeaderValue(_config.AuthenticationScheme, accessToken);
+			HttpResponseMessage response = await base.SendAsync(request, cancellationToken);
+
+			if (response.StatusCode != HttpStatusCode.Unauthorized || attempt >= _config.MaxAttempts)
+			{
+				return response;
+			}
+
+			// Mark this access token as invalid.
+			if (accessToken != null)
+			{
+				_authState.Invalidate(accessToken);
+			}
+
+			// Get the next token, and quit out if it's the same.
+			string? nextAccessToken = await _authState.GetAccessTokenAsync(allowInteractiveLogin, cancellationToken);
+			if (String.Equals(accessToken, nextAccessToken, StringComparison.Ordinal))
+			{
+				return response;
+			}
+
+			// Otherwise update the token and try again.
+			accessToken = nextAccessToken;
+		}
+	}
+}
diff --git a/Engine/Source/Programs/Shared/EpicGames.OIDC/HttpOidcAuthHandlerConfig.cs b/Engine/Source/Programs/Shared/EpicGames.OIDC/HttpOidcAuthHandlerConfig.cs
new file mode 100644
--- /dev/null
+++ b/Engine/Source/Programs/Shared/EpicGames.OIDC/HttpOidcAuthHandlerConfig.cs
@@ -0,0 +1,13 @@
+// Copyright Epic Games, Inc. All Rights Reserved.
+
+namespace EpicGames.OIDC;
+
+/// <summary>
+/// Configuration for a <see cref="HttpOidcAuthHandler"/>.
+/// </summary>
+public sealed record HttpOidcAuthHandlerConfig(
+	string AllowInteractiveLoginHeaderKey,
+	bool AllowInteractiveLogin = true,
+	string AuthenticationScheme = "Bearer",
+	int MaxAttempts = 3
+);
diff --git a/Engine/Source/Programs/Shared/EpicGames.OIDC/OidcAuthState.cs b/Engine/Source/Programs/Shared/EpicGames.OIDC/OidcAuthState.cs
new file mode 100644
--- /dev/null
+++ b/Engine/Source/Programs/Shared/EpicGames.OIDC/OidcAuthState.cs
@@ -0,0 +1,308 @@
+// Copyright Epic Games, Inc. All Rights Reserved.
+
+using System;
+using System.Diagnostics.CodeAnalysis;
+using System.Threading;
+using System.Threading.Tasks;
+using EpicGames.Core;
+using Microsoft.Extensions.Logging;
+
+namespace EpicGames.OIDC;
+
+/// <summary>
+/// Interface to manage the state of OIDC authentication.
+/// </summary>
+public interface IOidcAuthState
+{
+	/// <summary>
+	/// Event handler for the auth state changes.
+	/// </summary>
+	event Action? OnStateChanged;
+
+	/// <summary>
+	/// Is there a valid auth header available?
+	/// </summary>
+	bool IsLoggedIn { get; }
+
+	/// <summary>
+	/// Get the current access token, if one exists.
+	/// </summary>
+	string? CurrentAccessToken { get; }
+
+	/// <summary>
+	/// Reset the current auth state.
+	/// </summary>
+	void Reset();
+
+	/// <summary>
+	/// Marks the given access token as invalid, typically after having attempted to use it and got an unauthorized response.
+	/// </summary>
+	/// <param name="accessToken">The access token to invalidate.</param>
+	void Invalidate(string? accessToken);
+
+	/// <summary>
+	/// Get a valid access token, if possible.
+	/// </summary>
+	Task<string?> GetAccessTokenAsync(bool interactive, CancellationToken cancellationToken);
+}
+
+/// <summary>
+/// Class used to track the last obtained token for a given OIDC provider.
+/// </summary>
+public sealed class OidcAuthState(
+	IOidcTokenManager oidcTokenManager,
+	IStaticAccessTokenProvider staticAccessTokenProvider,
+	IClock clock,
+	IOidcAuthStateConfig config,
+	ILogger<OidcAuthState> logger
+) : IOidcAuthState, IAsyncDisposable
+{
+	/// <inheritdoc/>
+	public event Action? OnStateChanged;
+
+	/// <inheritdoc/>
+	public async ValueTask DisposeAsync()
+	{
+		if (_currentAuthWorker != null)
+		{
+			await _currentAuthWorker.DisposeAsync();
+			_currentAuthWorker = null;
+		}
+	}
+
+	/// <inheritdoc/>
+	public bool IsLoggedIn
+	{
+		get
+		{
+			if (_staticAccessTokenProvider.AccessToken != null)
+			{
+				return true;
+			}
+
+			if (TryGetCurrentAuthState(out AuthState? authState) && authState.IsAuthorized())
+			{
+				return true;
+			}
+
+			return false;
+		}
+	}
+
+	/// <inheritdoc/>
+	public string? CurrentAccessToken
+	{
+		get
+		{
+			string? accessToken = _staticAccessTokenProvider.AccessToken;
+			if (accessToken != null)
+			{
+				return accessToken;
+			}
+
+			if (TryGetCurrentAuthState(out AuthState? authState) && authState.IsAuthorized())
+			{
+				return authState.TokenInfo?.AccessToken;
+			}
+
+			return null;
+		}
+	}
+
+	/// <inheritdoc/>
+	public void Reset()
+	{
+		if (_staticAccessTokenProvider.AccessToken != null)
+		{
+			return;
+		}
+
+		lock (_lockObject)
+		{
+			_currentAuthResult = null;
+		}
+	}
+
+	/// <summary>
+	/// Marks the given access token as invalid, having attempted to use it and got an unauthorized response
+	/// </summary>
+	/// <param name="accessToken">The access  header to invalidate</param>
+	public void Invalidate(string? accessToken)
+	{
+		if (_staticAccessTokenProvider.AccessToken != null)
+		{
+			return;
+		}
+
+		lock (_lockObject)
+		{
+			if (TryGetCurrentAuthState(out AuthState? authState) && Equals(authState.TokenInfo?.AccessToken, accessToken))
+			{
+				_currentAuthResult = null;
+			}
+		}
+	}
+
+	/// <summary>
+	/// Gets the current access token
+	/// </summary>
+	public async Task<string?> GetAccessTokenAsync(bool interactive, CancellationToken cancellationToken)
+	{
+		// If there's a static access token configured, use that
+		string? accessToken = _staticAccessTokenProvider.AccessToken;
+		if (accessToken != null)
+		{
+			return accessToken;
+		}
+
+		// Otherwise check if we need to start a background task to figure it out
+		Task<AuthState> authStateTask;
+		lock (_lockObject)
+		{
+			if (_currentAuthResult == null || (interactive && !_currentAuthInteractive))
+			{
+				int authTaskId = ++_currentAuthTaskId;
+
+				if (_currentAuthResult == null || _currentAuthResult.Task.IsCompleted)
+				{
+					_currentAuthResult = new TaskCompletionSource<AuthState>(TaskCreationOptions.RunContinuationsAsynchronously);
+				}
+
+				BackgroundTask? prevAuthWorker = _currentAuthWorker;
+				_currentAuthWorker = BackgroundTask.StartNew(ctx => GetAuthStateHandlerAsync(authTaskId, interactive, prevAuthWorker, ctx));
+
+				_currentAuthInteractive = interactive;
+			}
+			authStateTask = _currentAuthResult.Task;
+		}
+
+		// Wait for the task to complete
+		AuthState authState = await authStateTask.WaitAsync(cancellationToken);
+		return authState?.TokenInfo?.AccessToken;
+	}
+
+	private readonly object _lockObject = new object();
+	private readonly IOidcTokenManager _oidcTokenManager = oidcTokenManager;
+	private readonly IStaticAccessTokenProvider _staticAccessTokenProvider = staticAccessTokenProvider;
+	private readonly IClock _clock = clock;
+	private readonly IOidcAuthStateConfig _config = config;
+	private readonly ILogger _logger = logger;
+
+	// Background auth process
+	private int _currentAuthTaskId;
+	private bool _currentAuthInteractive;
+	private BackgroundTask? _currentAuthWorker;
+	private TaskCompletionSource<AuthState>? _currentAuthResult;
+
+	/// <summary>
+	/// Gets the current auth state instance. Fails if the current auth task has not finished.
+	/// </summary>
+	private bool TryGetCurrentAuthState([NotNullWhen(true)] out AuthState? authState)
+	{
+		AuthState? state;
+		if (_currentAuthResult != null && _currentAuthResult.Task.IsCompletedSuccessfully && _currentAuthResult.Task.TryGetResult(out state))
+		{
+			authState = state;
+			return true;
+		}
+		else
+		{
+			authState = null;
+			return false;
+		}
+	}
+
+	private async Task GetAuthStateHandlerAsync(int authTaskId, bool interactive, BackgroundTask? prevAuthTask, CancellationToken cancellationToken)
+	{
+		Task? disposeTask = null;
+		try
+		{
+			// Start disposing of the previous auth task asynchronously
+			if (prevAuthTask != null)
+			{
+				disposeTask = Task.Run(() => prevAuthTask.DisposeAsync().AsTask(), CancellationToken.None);
+			}
+
+			// Get the new auth state
+			bool stateHasChanged = false;
+			try
+			{
+				AuthState authState = await GetAuthStateInternalAsync(interactive, cancellationToken);
+				lock (_lockObject)
+				{
+					if (_currentAuthTaskId == authTaskId)
+					{
+						_logger.LogDebug("Auth task complete (interactive: {Interactive}, authorized: {Authorized})", authState.Interactive, authState.IsAuthorized());
+						stateHasChanged = _currentAuthResult?.TrySetResult(authState) ?? false;
+					}
+				}
+			}
+			catch (Exception ex)
+			{
+				lock (_lockObject)
+				{
+					if (_currentAuthTaskId == authTaskId)
+					{
+						_logger.LogDebug(ex, "Exception while attempting auth: {Message}", ex.Message);
+						stateHasChanged = _currentAuthResult?.TrySetException(ex) ?? false;
+					}
+				}
+			}
+
+			// Send notifications about the updated state
+			if (stateHasChanged)
+			{
+				try
+				{
+					OnStateChanged?.Invoke();
+				}
+				catch (Exception ex)
+				{
+					_logger.LogError(ex, "Exception while sending state change notifications: {Message}", ex.Message);
+				}
+			}
+		}
+		finally
+		{
+			// Wait for the child task to finish being disposed
+			if (disposeTask != null)
+			{
+				await disposeTask;
+			}
+		}
+	}
+
+	private async Task<AuthState> GetAuthStateInternalAsync(bool interactive, CancellationToken cancellationToken)
+	{
+		IOidcTokenInfo? result = null;
+		string oidcProvider = await _config.GetOidcProviderAsync(cancellationToken);
+
+		OidcStatus providerStatus = await _oidcTokenManager.GetStatusForProvider(oidcProvider, cancellationToken);
+		if (providerStatus != OidcStatus.NotLoggedIn)
+		{
+			try
+			{
+				result = await _oidcTokenManager.TryGetAccessToken(oidcProvider, cancellationToken);
+			}
+			catch (Exception ex)
+			{
+				_logger.LogTrace(ex, "Unable to get access token; attempting login: {Message}", ex.Message);
+			}
+		}
+		if (result == null && interactive)
+		{
+			_logger.LogInformation("Logging in to provider '{Provider}'...", oidcProvider);
+			result = await _oidcTokenManager.LoginAsync(oidcProvider, cancellationToken);
+		}
+
+		return new AuthState(_clock, result, interactive);
+	}
+
+	private record AuthState(IClock Clock, IOidcTokenInfo? TokenInfo, bool Interactive)
+	{
+		public bool IsAuthorized()
+		{
+			return TokenInfo != null && TokenInfo.IsValid(Clock.UtcNow);
+		}
+	}
+}
diff --git a/Engine/Source/Programs/Shared/EpicGames.OIDC/OidcAuthStateConfig.cs b/Engine/Source/Programs/Shared/EpicGames.OIDC/OidcAuthStateConfig.cs
new file mode 100644
--- /dev/null
+++ b/Engine/Source/Programs/Shared/EpicGames.OIDC/OidcAuthStateConfig.cs
@@ -0,0 +1,26 @@
+// Copyright Epic Games, Inc. All Rights Reserved.
+
+using System.Threading;
+using System.Threading.Tasks;
+
+namespace EpicGames.OIDC;
+
+/// <summary>
+/// Interface defining configuration for an <see cref="OidcAuthState"/>.
+/// </summary>
+public interface IOidcAuthStateConfig
+{
+	/// <summary>
+	/// The OIDC provider associated with this auth state.
+	/// </summary>
+	Task<string> GetOidcProviderAsync(CancellationToken cancellationToken);
+}
+
+/// <summary>
+/// Configuration for an <see cref="OidcAuthState"/>.
+/// </summary>
+public sealed class OidcAuthStateConfig(string oidcProvider) : IOidcAuthStateConfig
+{
+	/// <inheritdoc/>
+	public Task<string> GetOidcProviderAsync(CancellationToken cancellationToken) => Task.FromResult(oidcProvider);
+}
diff --git a/Engine/Source/Programs/Shared/EpicGames.OIDC/OidcTokenInfo.cs b/Engine/Source/Programs/Shared/EpicGames.OIDC/OidcTokenInfo.cs
new file mode 100644
--- /dev/null
+++ b/Engine/Source/Programs/Shared/EpicGames.OIDC/OidcTokenInfo.cs
@@ -0,0 +1,80 @@
+// Copyright Epic Games, Inc. All Rights Reserved.
+
+using System;
+
+namespace EpicGames.OIDC;
+
+/// <summary>
+/// Interface describing an OIDC token.
+/// </summary>
+public interface IOidcTokenInfo
+{
+	/// <summary>
+	/// The refresh token.
+	/// </summary>
+	public string? RefreshToken { get; }
+	/// <summary>
+	/// The access token.
+	/// </summary>
+	public string? AccessToken { get; }
+
+	/// <summary>
+	/// When the token expires.
+	/// </summary>
+	public DateTimeOffset TokenExpiry { get; }
+
+	/// <summary>
+	/// Determine if the token is valid at the given time.
+	/// </summary>
+	public bool IsValid(DateTimeOffset currentTime);
+}
+
+/// <summary>
+/// An OIDC access token, refresh token, and expiry.
+/// </summary>
+public class OidcTokenInfo : IOidcTokenInfo
+{
+	/// <inheritdoc/>
+	public string? RefreshToken { get; set; }
+
+	/// <inheritdoc/>
+	public string? AccessToken { get; set; }
+
+	/// <inheritdoc/>
+	public DateTimeOffset TokenExpiry { get; set; }
+
+	/// <inheritdoc/>
+	public bool IsValid(DateTimeOffset currentTime)
+	{
+		if (String.IsNullOrEmpty(RefreshToken) || String.IsNullOrEmpty(AccessToken))
+		{
+			return false;
+		}
+
+		// An expiry of MinValue means it has no expiration time
+		if (TokenExpiry == DateTimeOffset.MinValue)
+		{
+			return true;
+		}
+
+		return currentTime <= TokenExpiry;
+	}
+}
+
+/// <summary>
+/// A fake OIDC token info representing calls that allow anonymous callers.
+/// </summary>
+public sealed class AnonymousOidcTokenInfo : IOidcTokenInfo
+{
+	/// <inheritdoc/>
+	public string? RefreshToken => null;
+
+	/// <inheritdoc/>
+	public string? AccessToken => null;
+
+	/// <inheritdoc/>
+	public DateTimeOffset TokenExpiry => DateTimeOffset.MaxValue;
+
+	/// <inheritdoc/>
+	public bool IsValid(DateTimeOffset currentTime) => true;
+}
diff --git a/Engine/Source/Programs/Shared/EpicGames.OIDC/OidcTokenManager.cs b/Engine/Source/Programs/Shared/EpicGames.OIDC/OidcTokenManager.cs
--- a/Engine/Source/Programs/Shared/EpicGames.OIDC/OidcTokenManager.cs
+++ b/Engine/Source/Programs/Shared/EpicGames.OIDC/OidcTokenManager.cs
@@ -27,13 +27,13 @@ namespace EpicGames.OIDC
 {
 	public interface IOidcTokenManager
 	{
-		public Task<OidcTokenInfo> LoginAsync(string providerIdentifier, CancellationToken cancellationToken = default);
+		public Task<IOidcTokenInfo> LoginAsync(string providerIdentifier, CancellationToken cancellationToken = default);
 
-		public Task<OidcTokenInfo> GetAccessToken(string providerIdentifier, CancellationToken cancellationToken = default);
+		public Task<IOidcTokenInfo> GetAccessToken(string providerIdentifier, CancellationToken cancellationToken = default);
 
-		public Task<OidcTokenInfo?> TryGetAccessToken(string providerIdentifier, CancellationToken cancellationToken = default);
+		public Task<IOidcTokenInfo?> TryGetAccessToken(string providerIdentifier, CancellationToken cancellationToken = default);
 
-		public OidcStatus GetStatusForProvider(string providerIdentifier);
+		public Task<OidcStatus> GetStatusForProvider(string providerIdentifier, CancellationToken cancellationToken = default);
 	}
 
 	public class OidcTokenManager : IOidcTokenManager
@@ -139,27 +139,27 @@ namespace EpicGames.OIDC
 		}
 
 		/// <inheritdoc/>
-		public Task<OidcTokenInfo> LoginAsync(string providerIdentifier, CancellationToken cancellationToken = default)
+		public Task<IOidcTokenInfo> LoginAsync(string providerIdentifier, CancellationToken cancellationToken = default)
 		{
 			return _tokenClients[providerIdentifier].LoginAsync(cancellationToken);
 		}
 
 		/// <inheritdoc/>
-		public Task<OidcTokenInfo> GetAccessToken(string providerIdentifier, CancellationToken cancellationToken = default)
+		public Task<IOidcTokenInfo> GetAccessToken(string providerIdentifier, CancellationToken cancellationToken = default)
 		{
 			return _tokenClients[providerIdentifier].GetAccessTokenAsync(cancellationToken);
 		}
 
 		/// <inheritdoc/>
-		public Task<OidcTokenInfo?> TryGetAccessToken(string providerIdentifier, CancellationToken cancellationToken = default)
+		public Task<IOidcTokenInfo?> TryGetAccessToken(string providerIdentifier, CancellationToken cancellationToken = default)
 		{
 			return _tokenClients[providerIdentifier].TryGetAccessTokenAsync(cancellationToken);
 		}
 
 		/// <inheritdoc/>
-		public OidcStatus GetStatusForProvider(string providerIdentifier)
+		public Task<OidcStatus> GetStatusForProvider(string providerIdentifier, CancellationToken cancellationToken = default)
 		{
-			return _tokenClients[providerIdentifier].GetStatus();
+			return Task.FromResult(_tokenClients[providerIdentifier].GetStatus());
 		}
 	}
 
@@ -194,33 +194,6 @@ namespace EpicGames.OIDC
 		TokenRefreshRequired
 	}
 
-	public class OidcTokenInfo
-	{
-		public string? RefreshToken { get; set; }
-		public string? AccessToken { get; set; }
-
-		/// <summary>
-		/// When access token expires
-		/// </summary>
-		public DateTimeOffset TokenExpiry { get; set; }
-
-		public bool IsValid(DateTimeOffset currentTime)
-		{
-			if (String.IsNullOrEmpty(RefreshToken) || String.IsNullOrEmpty(AccessToken))
-			{
-				return false;
-			}
-
-			// An expiry of MinValue means it has no expiration time
-			if (TokenExpiry == DateTimeOffset.MinValue)
-			{
-				return true;
-			}
-
-			return currentTime <= TokenExpiry;
-		}
-	}
-
 	/// <summary>
 	/// The data read from a http request into the oidc client, mostly a wrapper around the http.listener context
 	/// </summary>
@@ -482,9 +455,9 @@ namespace EpicGames.OIDC
 			return options;
 		}
 
-		public async Task<OidcTokenInfo> LoginAsync(CancellationToken cancellationToken)
+		public async Task<IOidcTokenInfo> LoginAsync(CancellationToken cancellationToken)
 		{
-			OidcTokenInfo tokenInfo = await DoLoginAsync(cancellationToken);
+			IOidcTokenInfo tokenInfo = await DoLoginAsync(cancellationToken);
 
 			if (String.IsNullOrEmpty(tokenInfo.RefreshToken))
 			{
@@ -504,7 +477,7 @@ namespace EpicGames.OIDC
 			return tokenInfo;
 		}
 
-		public virtual async Task<OidcTokenInfo> DoLoginAsync(CancellationToken cancellationToken)
+		public virtual async Task<IOidcTokenInfo> DoLoginAsync(CancellationToken cancellationToken)
 		{
 			// setup a local http server to listen for the result of the login
 			LoginResult? loginResult = null;
@@ -831,12 +804,12 @@ namespace EpicGames.OIDC
 			return disco;
 		}
 
-		public async Task<OidcTokenInfo> GetAccessTokenAsync(CancellationToken cancellationToken)
+		public async Task<IOidcTokenInfo> GetAccessTokenAsync(CancellationToken cancellationToken)
 		{
 			return await TryGetAccessTokenAsync(cancellationToken) ?? throw new NotLoggedInException();
 		}
 
-		public async Task<OidcTokenInfo?> TryGetAccessTokenAsync(CancellationToken cancellationToken)
+		public async Task<IOidcTokenInfo?> TryGetAccessTokenAsync(CancellationToken cancellationToken)
 		{
 			if (String.IsNullOrEmpty(_refreshToken))
 			{
diff --git a/Engine/Source/Programs/Shared/EpicGames.OIDC/StaticAccessTokenProvider.cs b/Engine/Source/Programs/Shared/EpicGames.OIDC/StaticAccessTokenProvider.cs
new file mode 100644
--- /dev/null
+++ b/Engine/Source/Programs/Shared/EpicGames.OIDC/StaticAccessTokenProvider.cs
@@ -0,0 +1,23 @@
+// Copyright Epic Games, Inc. All Rights Reserved.
+
+namespace EpicGames.OIDC;
+
+/// <summary>
+/// Interface to support getting a static access token.
+/// </summary>
+public interface IStaticAccessTokenProvider
+{
+	/// <summary>
+	/// The static access token, if one exists, or <see langword="null"/>.
+	/// </summary>
+	string? AccessToken { get; }
+}
+
+/// <summary>
+/// Static access token provider that never provides a static access token.
+/// </summary>
+public sealed class NullStaticAccessTokenProvider : IStaticAccessTokenProvider
+{
+	/// <inheritdoc/>
+	public string? AccessToken => null;
+}
diff --git a/Engine/Source/Programs/Shared/EpicGames.ProjectStore/CloudStorage/CloudStorage.cs b/Engine/Source/Programs/Shared/EpicGames.ProjectStore/CloudStorage/CloudStorage.cs
--- a/Engine/Source/Programs/Shared/EpicGames.ProjectStore/CloudStorage/CloudStorage.cs
+++ b/Engine/Source/Programs/Shared/EpicGames.ProjectStore/CloudStorage/CloudStorage.cs
@@ -127,7 +127,7 @@ namespace EpicGames.ProjectStore
 			bool allowInteractiveLogin = !unattended || _options.Value.AllowAuthPrompt;
 			try
 			{
-				OidcTokenInfo tokenInfo = await tokenManager.GetAccessToken(authConfig.DefaultProvider);
+				IOidcTokenInfo tokenInfo = await tokenManager.GetAccessToken(authConfig.DefaultProvider);
 
 				return tokenInfo.AccessToken;
 			}
@@ -138,7 +138,7 @@ namespace EpicGames.ProjectStore
 					throw;
 				}
 
-				OidcTokenInfo tokenInfo = await tokenManager.LoginAsync(authConfig.DefaultProvider);
+				IOidcTokenInfo tokenInfo = await tokenManager.LoginAsync(authConfig.DefaultProvider);
 
 				return tokenInfo.AccessToken;
 			}
diff --git a/Engine/Source/Programs/Shared/EpicGames.ProjectStore/ProjectStoreTypes.cs b/Engine/Source/Programs/Shared/EpicGames.ProjectStore/ProjectStoreTypes.cs
--- a/Engine/Source/Programs/Shared/EpicGames.ProjectStore/ProjectStoreTypes.cs
+++ b/Engine/Source/Programs/Shared/EpicGames.ProjectStore/ProjectStoreTypes.cs
@@ -819,7 +819,7 @@ namespace EpicGames.ProjectStore
 		/// <param name="ns">The Unreal Cloud namespace to search.</param>
 		/// <param name="accessTokenInfo"></param>
 		/// <param name="runApp">A delegate used to run command-line apps.</param>
-		public UnrealCloudDDCBuildIndex(string host, string ns, OidcTokenInfo accessTokenInfo, RunAppDelegate runApp)
+		public UnrealCloudDDCBuildIndex(string host, string ns, IOidcTokenInfo accessTokenInfo, RunAppDelegate runApp)
 			: base(runApp)
 		{
 			_host = host;
@@ -852,7 +852,7 @@ namespace EpicGames.ProjectStore
 
 				try
 				{
-					OidcTokenInfo tokenInfo = await tokenManager.GetAccessToken(authConfig.DefaultProvider);
+					IOidcTokenInfo tokenInfo = await tokenManager.GetAccessToken(authConfig.DefaultProvider);
 					return new UnrealCloudDDCBuildIndex(host, ns, tokenInfo, runApp);
 				}
 				catch (NotLoggedInException)
@@ -862,7 +862,7 @@ namespace EpicGames.ProjectStore
 						throw;
 					}
 
-					OidcTokenInfo tokenInfo = await tokenManager.LoginAsync(authConfig.DefaultProvider);
+					IOidcTokenInfo tokenInfo = await tokenManager.LoginAsync(authConfig.DefaultProvider);
 					return new UnrealCloudDDCBuildIndex(host, ns, tokenInfo, runApp);
 				}
 			}
diff --git a/Engine/Source/Programs/UnrealGameSync/UnrealGameSync/ProgramApplicationContext.cs b/Engine/Source/Programs/UnrealGameSync/UnrealGameSync/ProgramApplicationContext.cs
--- a/Engine/Source/Programs/UnrealGameSync/UnrealGameSync/ProgramApplicationContext.cs
+++ b/Engine/Source/Programs/UnrealGameSync/UnrealGameSync/ProgramApplicationContext.cs
@@ -208,7 +208,7 @@ namespace UnrealGameSync
 				}
 			}
 
-			HordeHttpAuthHandlerState? hordeAuthState = _serviceProvider.GetService<HordeHttpAuthHandlerState>();
+			IHordeClient? hordeAuthState = _serviceProvider.GetService<IHordeClient>();
 			_startupTask = Task.Run(() => WaitForStartupTasks(startupTasks, hordeAuthState));
 
 			_startupWindow = new ModalTaskWindow("Opening Projects", "Opening projects, please wait...", FormStartPosition.CenterScreen, _startupTask, _startupCancellationSource);
@@ -253,7 +253,7 @@ namespace UnrealGameSync
 			}
 		}
 
-		static async Task WaitForStartupTasks(List<(UserSelectedProjectSettings, ModalTask<OpenProjectInfo>)> startupTasks, HordeHttpAuthHandlerState? hordeAuthState)
+		static async Task WaitForStartupTasks(List<(UserSelectedProjectSettings, ModalTask<OpenProjectInfo>)> startupTasks, IHordeClient? hordeClient)
 		{
 			foreach ((_, ModalTask<OpenProjectInfo> modalTask) in startupTasks)
 			{
@@ -276,9 +276,9 @@ namespace UnrealGameSync
 				}
 			}
 
-			if (hordeAuthState != null)
+			if (hordeClient != null)
 			{
-				await hordeAuthState.GetAccessTokenAsync(false, CancellationToken.None);
+				await hordeClient.GetAccessTokenAsync(false, CancellationToken.None);
 			}
 		}
 
diff --git a/Engine/Source/Programs/UnrealGameSync/UnrealGameSyncCmd/Commands/LoginCommand.cs b/Engine/Source/Programs/UnrealGameSync/UnrealGameSyncCmd/Commands/LoginCommand.cs
--- a/Engine/Source/Programs/UnrealGameSync/UnrealGameSyncCmd/Commands/LoginCommand.cs
+++ b/Engine/Source/Programs/UnrealGameSync/UnrealGameSyncCmd/Commands/LoginCommand.cs
@@ -36,7 +36,7 @@ namespace UnrealGameSyncCmd.Commands
 			using ITokenStore tokenStore = TokenStoreFactory.CreateTokenStore();
 			IConfiguration providerConfiguration = ProviderConfigurationFactory.ReadConfiguration(engineDir, gameDir);
 			OidcTokenManager oidcTokenManager = OidcTokenManager.CreateTokenManager(providerConfiguration, tokenStore, new List<string>() { providerIdentifier });
-			OidcTokenInfo result = await oidcTokenManager.LoginAsync(providerIdentifier);
+			IOidcTokenInfo result = await oidcTokenManager.LoginAsync(providerIdentifier);
 
 			logger.LogInformation("Logged in to provider {ProviderIdentifier}", providerIdentifier);
 		}
diff --git a/Engine/Source/Programs/UnrealGameSync/UnrealGameSyncShared/CloudStorage.cs b/Engine/Source/Programs/UnrealGameSync/UnrealGameSyncShared/CloudStorage.cs
--- a/Engine/Source/Programs/UnrealGameSync/UnrealGameSyncShared/CloudStorage.cs
+++ b/Engine/Source/Programs/UnrealGameSync/UnrealGameSyncShared/CloudStorage.cs
@@ -326,7 +326,7 @@ namespace UnrealGameSync
 				}
 			}
 
-			OidcTokenInfo tokenInfo;
+			IOidcTokenInfo tokenInfo;
 			try
 			{
 				tokenInfo = await _tokenManager!.GetAccessToken(_providerName!, CancellationToken.None);
