Shared: Enable IDE0028, IDE0031, IDE0034, IDE0054, IDE0062, IDE0071, IDE0074, IDE0075, IDE0280, IDE0300, IDE0301 IDE0305

#rnx

#Synced-CL 50769868

diff --git a/Engine/Source/Programs/Shared/EpicGames.Analytics.Generators/AnalyticsTableGenerator.cs b/Engine/Source/Programs/Shared/EpicGames.Analytics.Generators/AnalyticsTableGenerator.cs
--- a/Engine/Source/Programs/Shared/EpicGames.Analytics.Generators/AnalyticsTableGenerator.cs
+++ b/Engine/Source/Programs/Shared/EpicGames.Analytics.Generators/AnalyticsTableGenerator.cs
@@ -102,7 +102,7 @@ namespace EpicGames.Analytics.Generators
 
 		private class SyntaxReceiver : ISyntaxReceiver
 		{
-			public List<TypeDeclarationSyntax> CandidateTypes { get; } = new List<TypeDeclarationSyntax>();
+			public List<TypeDeclarationSyntax> CandidateTypes { get; } = [];
 
 			public void OnVisitSyntaxNode(SyntaxNode syntaxNode)
 			{
@@ -133,7 +133,7 @@ namespace EpicGames.Analytics.Generators
 			string telemetryName = !String.IsNullOrEmpty(fullNamepsace) ? $"{fullNamepsace}_{telemetrySymbol.Name}" : telemetrySymbol.Name;
 			string? tableName = null;
 
-			List<Tuple<string, string>> columns = new List<Tuple<string, string>>();
+			List<Tuple<string, string>> columns = [];
 
 			foreach (AttributeData attr in telemetrySymbol.GetAttributes())
 			{
diff --git a/Engine/Source/Programs/Shared/EpicGames.Analytics.Telemetry/AbstractTelemetryRecord.cs b/Engine/Source/Programs/Shared/EpicGames.Analytics.Telemetry/AbstractTelemetryRecord.cs
--- a/Engine/Source/Programs/Shared/EpicGames.Analytics.Telemetry/AbstractTelemetryRecord.cs
+++ b/Engine/Source/Programs/Shared/EpicGames.Analytics.Telemetry/AbstractTelemetryRecord.cs
@@ -38,7 +38,7 @@ namespace EpicGames.Analytics
 		/// </remarks>
 		public Hashtable ToHashtable()
 		{
-			Hashtable result = new Hashtable();
+			Hashtable result = [];
 			WriteProperties(result);
 			return result;
 		}
diff --git a/Engine/Source/Programs/Shared/EpicGames.AspNet/CbConvertersAspnet.cs b/Engine/Source/Programs/Shared/EpicGames.AspNet/CbConvertersAspnet.cs
--- a/Engine/Source/Programs/Shared/EpicGames.AspNet/CbConvertersAspnet.cs
+++ b/Engine/Source/Programs/Shared/EpicGames.AspNet/CbConvertersAspnet.cs
@@ -101,18 +101,18 @@ namespace EpicGames.AspNet
 				throw new CbException($"Error converting field \"{field.Name}\" to ValidationProblemDetails. Expected CbObject.");
 			}
 
-			Dictionary<string, string[]> errors = new Dictionary<string, string[]>();
+			Dictionary<string, string[]> errors = [];
 
 			foreach (CbField error in field[new Utf8String("errors")].AsObject())
 			{
-				List<string> errorValuesList = new List<string>();
+				List<string> errorValuesList = [];
 				
 				foreach (CbField errorValue in error.AsArray())
 				{
 					errorValuesList.Add(errorValue.AsString());
 				}
 				
-				errors[error.Name.ToString()] = errorValuesList.ToArray();
+				errors[error.Name.ToString()] = [.. errorValuesList];
 			}
 
 			ValidationProblemDetails result = new ValidationProblemDetails
diff --git a/Engine/Source/Programs/Shared/EpicGames.AspNet/ServerTimings.cs b/Engine/Source/Programs/Shared/EpicGames.AspNet/ServerTimings.cs
--- a/Engine/Source/Programs/Shared/EpicGames.AspNet/ServerTimings.cs
+++ b/Engine/Source/Programs/Shared/EpicGames.AspNet/ServerTimings.cs
@@ -86,7 +86,7 @@ namespace EpicGames.AspNet
 
 	public class ServerTiming : IServerTiming
 	{
-		private readonly ConcurrentBag<ServerTimingMetric> _metrics = new ConcurrentBag<ServerTimingMetric>();
+		private readonly ConcurrentBag<ServerTimingMetric> _metrics = [];
 
 		public void AddServerTimingMetric(string metricName, double? duration, string? description)
 		{
diff --git a/Engine/Source/Programs/Shared/EpicGames.Build/Platform/SDKCollection.cs b/Engine/Source/Programs/Shared/EpicGames.Build/Platform/SDKCollection.cs
--- a/Engine/Source/Programs/Shared/EpicGames.Build/Platform/SDKCollection.cs
+++ b/Engine/Source/Programs/Shared/EpicGames.Build/Platform/SDKCollection.cs
@@ -16,7 +16,7 @@ namespace EpicGames.Core
 
 		public UEBuildPlatformSDK PlatformSDK { get; init; } = inPlatformSDK;
 		public List<SDKDescriptor> Sdks { get; init; } = [];
-		public List<SDKDescriptor> FullSdks => Sdks.Where(x => !String.Equals(x.Name, "AutoSdk", StringComparison.OrdinalIgnoreCase)).ToList();
+		public List<SDKDescriptor> FullSdks => [.. Sdks.Where(x => !String.Equals(x.Name, "AutoSdk", StringComparison.OrdinalIgnoreCase))];
 		public SDKDescriptor? AutoSdk => Sdks.FirstOrDefault(x => String.Equals(x.Name, "AutoSdk", StringComparison.OrdinalIgnoreCase));
 
 		public SDKCollection(SDKCollection other, UEBuildPlatformSDK platformSDK)
diff --git a/Engine/Source/Programs/Shared/EpicGames.Build/Platform/UEBuildPlatformSDK.cs b/Engine/Source/Programs/Shared/EpicGames.Build/Platform/UEBuildPlatformSDK.cs
--- a/Engine/Source/Programs/Shared/EpicGames.Build/Platform/UEBuildPlatformSDK.cs
+++ b/Engine/Source/Programs/Shared/EpicGames.Build/Platform/UEBuildPlatformSDK.cs
@@ -657,7 +657,7 @@ namespace EpicGames.Core
 		private static List<FileReference> ProjectsToCheckForSDKOverrides = [];
 		public static void InitializePerProjectSDKVersions(IEnumerable<FileReference> projectsToCheckForOverrides)
 		{
-			ProjectsToCheckForSDKOverrides = projectsToCheckForOverrides.ToList();
+			ProjectsToCheckForSDKOverrides = [.. projectsToCheckForOverrides];
 		}
 
 		#endregion
diff --git a/Engine/Source/Programs/Shared/EpicGames.Build/System/CompileScriptModules.cs b/Engine/Source/Programs/Shared/EpicGames.Build/System/CompileScriptModules.cs
--- a/Engine/Source/Programs/Shared/EpicGames.Build/System/CompileScriptModules.cs
+++ b/Engine/Source/Programs/Shared/EpicGames.Build/System/CompileScriptModules.cs
@@ -58,7 +58,7 @@ namespace UnrealBuildBase
 			private ILogger Logger;
 			private ConcurrentDictionary<string, DateTime> WriteTimes = [];
 			private string BuildRecordDirectory;
-			private Dictionary<FileReference, CsProjBuildRecordEntry> ValidBuildRecords = new();
+			private Dictionary<FileReference, CsProjBuildRecordEntry> ValidBuildRecords = [];
 
 			public Hook(ILogger InLogger, Rules.RulesFileType InRulesFileType)
 			{
@@ -160,7 +160,7 @@ namespace UnrealBuildBase
 				.. GameDirectories,
 				.. AdditionalDirectories,
 			];
-			BaseDirectories = BaseDirectories.Distinct().ToList();
+			BaseDirectories = [.. BaseDirectories.Distinct()];
 			Logger.LogDebug("Find {Count} directories time: {TimeSeconds:0.00} s", BaseDirectories.Count, Timer.Elapsed.TotalSeconds);
 
 			Timer.Restart();
@@ -171,7 +171,7 @@ namespace UnrealBuildBase
 				ForeignPlugins: null, AdditionalSearchPaths: [.. AdditionalDirectories, .. GameBuildDirectories])];
 			if (AutomationProjectIncludeFilter != null && AutomationProjectIncludeFilter.Count > 0)
 			{
-				FoundProjects = FoundProjects.Where(x => (AutomationProjectIncludeFilter.Contains(x.GetFileNameWithoutAnyExtensions(), StringComparer.OrdinalIgnoreCase))).ToHashSet();
+				FoundProjects = [.. FoundProjects.Where(x => (AutomationProjectIncludeFilter.Contains(x.GetFileNameWithoutAnyExtensions(), StringComparer.OrdinalIgnoreCase)))];
 			}
 			Logger.LogDebug("Find {Count} projects time: {TimeSeconds:0.00} s", FoundProjects.Count, Timer.Elapsed.TotalSeconds);
 
@@ -206,8 +206,8 @@ namespace UnrealBuildBase
 				bool bEngineBuildSuccess = false;
 				Dictionary<FileReference, CsProjBuildRecordEntry> EngineBuildResults = Build(
 					RulesFileType,
-					FoundProjects.Where(x => x.IsUnderDirectory(Unreal.EngineDirectory)).ToHashSet(),
-					BaseDirectories.Where(x => x.IsUnderDirectory(Unreal.EngineDirectory)).ToList(),
+					[.. FoundProjects.Where(x => x.IsUnderDirectory(Unreal.EngineDirectory))],
+					[.. BaseDirectories.Where(x => x.IsUnderDirectory(Unreal.EngineDirectory))],
 					null,
 					EngineBuildFlags,
 					out bEngineBuildSuccess,
@@ -222,7 +222,7 @@ namespace UnrealBuildBase
 				}
 				Dictionary<FileReference, CsProjBuildRecordEntry> ProjectBuildResults = Build(
 					RulesFileType,
-					FoundProjects.Where(x => !x.IsUnderDirectory(Unreal.EngineDirectory)).ToHashSet(),
+					[.. FoundProjects.Where(x => !x.IsUnderDirectory(Unreal.EngineDirectory))],
 					BaseDirectories,
 					null,
 					ProjectBuildFlags,
@@ -674,7 +674,7 @@ namespace UnrealBuildBase
 
 			if (String.IsNullOrEmpty(ScriptsForProjectFileName))
 			{
-				GameDirectories = NativeProjectsBase.EnumerateProjectFiles(Logger).Select(x => x.Directory).ToList();
+				GameDirectories = [.. NativeProjectsBase.EnumerateProjectFiles(Logger).Select(x => x.Directory)];
 			}
 			else
 			{
@@ -687,10 +687,10 @@ namespace UnrealBuildBase
 
 		static List<DirectoryReference> GetAdditionalDirectories(List<string>? AdditionalScriptsFolders) =>
 			AdditionalScriptsFolders == null ? [] :
-				AdditionalScriptsFolders.Select(x => DirectoryReference.FindCorrectCase(new DirectoryReference(x))).ToList();
+				[.. AdditionalScriptsFolders.Select(x => DirectoryReference.FindCorrectCase(new DirectoryReference(x)))];
 
 		static List<DirectoryReference> GetAdditionalBuildDirectories(List<DirectoryReference> GameDirectories) =>
-			GameDirectories.Select(x => DirectoryReference.Combine(x, "Build")).Where(x => DirectoryReference.Exists(x)).ToList();
+			[.. GameDirectories.Select(x => DirectoryReference.Combine(x, "Build")).Where(x => DirectoryReference.Exists(x))];
 	}
 }
 
diff --git a/Engine/Source/Programs/Shared/EpicGames.Build/System/EnumeratePlugins.cs b/Engine/Source/Programs/Shared/EpicGames.Build/System/EnumeratePlugins.cs
--- a/Engine/Source/Programs/Shared/EpicGames.Build/System/EnumeratePlugins.cs
+++ b/Engine/Source/Programs/Shared/EpicGames.Build/System/EnumeratePlugins.cs
@@ -45,7 +45,7 @@ namespace UnrealBuildBase
 				BaseDirs.AddRange(Unreal.GetExtensionDirs(ProjectFile.Directory, "Plugins"));
 				BaseDirs.AddRange(Unreal.GetExtensionDirs(ProjectFile.Directory, "Mods"));
 			}
-			return BaseDirs.SelectMany(x => EnumeratePlugins(x)).ToList();
+			return [.. BaseDirs.SelectMany(x => EnumeratePlugins(x))];
 		}
 
 		/// <summary>
@@ -150,10 +150,10 @@ namespace UnrealBuildBase
 				Plugins = null;
 
 				// Collect the list of plugins already build
-				List<FileReference> Built = FoundPlugins.Where(P => UbtPluginFileCache.ContainsKey(P)).ToList();
+				List<FileReference> Built = [.. FoundPlugins.Where(P => UbtPluginFileCache.ContainsKey(P))];
 
 				// Collect the list of plugins not build
-				List<FileReference> NotBuilt = FoundPlugins.Where(P => !UbtPluginFileCache.ContainsKey(P)).ToList();
+				List<FileReference> NotBuilt = [.. FoundPlugins.Where(P => !UbtPluginFileCache.ContainsKey(P))];
 
 				// If we have anything left to build, then build it
 				if (NotBuilt.Count > 0)
diff --git a/Engine/Source/Programs/Shared/EpicGames.Build/System/Parallel2.cs b/Engine/Source/Programs/Shared/EpicGames.Build/System/Parallel2.cs
--- a/Engine/Source/Programs/Shared/EpicGames.Build/System/Parallel2.cs
+++ b/Engine/Source/Programs/Shared/EpicGames.Build/System/Parallel2.cs
@@ -115,10 +115,7 @@ namespace UnrealBuildBase
 
 			context.Handle.Wait();
 
-			if (context.Exception != null)
-			{
-				context.Exception.Throw();
-			}
+			context.Exception?.Throw();
 		}
 
 		/// <summary>
@@ -202,10 +199,7 @@ namespace UnrealBuildBase
 
 			context.Handle.Wait();
 
-			if (context.Exception != null)
-			{
-				context.Exception.Throw();
-			}
+			context.Exception?.Throw();
 		}
 
 		private class ForEachContext<TSource>
@@ -322,10 +316,7 @@ namespace UnrealBuildBase
 
 			WorkerLoop();
 
-			if (_exception != null)
-			{
-				_exception.Throw();
-			}
+			_exception?.Throw();
 		}
 
 		public void Dispose()
diff --git a/Engine/Source/Programs/Shared/EpicGames.Build/System/Timeline.cs b/Engine/Source/Programs/Shared/EpicGames.Build/System/Timeline.cs
--- a/Engine/Source/Programs/Shared/EpicGames.Build/System/Timeline.cs
+++ b/Engine/Source/Programs/Shared/EpicGames.Build/System/Timeline.cs
@@ -126,7 +126,7 @@ namespace UnrealBuildBase
 		/// <summary>
 		/// The recorded events
 		/// </summary>
-		static List<Event> Events = new List<Event>();
+		static List<Event> Events = [];
 
 		/// <summary>
 		/// Property for the total time elapsed
@@ -212,10 +212,10 @@ namespace UnrealBuildBase
 			// Create the root event
 			TimeSpan FinishTime = Stopwatch.Elapsed;
 
-			List<Event> OuterEvents = new List<Event>
-			{
+			List<Event> OuterEvents =
+			[
 				new Event("<Root>", TimeSpan.Zero, FinishTime)
-			};
+			];
 
 			// Print out all the child events
 			TimeSpan LastTime = TimeSpan.Zero;
diff --git a/Engine/Source/Programs/Shared/EpicGames.BuildGraph.Tests/OpcodeTests.cs b/Engine/Source/Programs/Shared/EpicGames.BuildGraph.Tests/OpcodeTests.cs
--- a/Engine/Source/Programs/Shared/EpicGames.BuildGraph.Tests/OpcodeTests.cs
+++ b/Engine/Source/Programs/Shared/EpicGames.BuildGraph.Tests/OpcodeTests.cs
@@ -307,7 +307,7 @@ namespace EpicGames.BuildGraph.Tests
 		public void StrSplit()
 		{
 			object result = Evaluate(((BgString)"1+2+3").Split("+"));
-			Assert.IsTrue(Enumerable.SequenceEqual((IEnumerable<object>)result, new[] { "1", "2", "3" }));
+			Assert.IsTrue(Enumerable.SequenceEqual((IEnumerable<object>)result, ["1", "2", "3"]));
 		}
 
 		[TestMethod]
@@ -349,8 +349,8 @@ namespace EpicGames.BuildGraph.Tests
 			Assert.AreEqual("value", result.DefaultValue);
 			Assert.AreEqual("abc", result.Pattern);
 			Assert.AreEqual("def", result.PatternFailed);
-			Assert.IsTrue(result.Values!.SequenceEqual(new[] { "value1", "value2" }));
-			Assert.IsTrue(result.ValueDescriptions!.SequenceEqual(new[] { "desc1", "desc2" }));
+			Assert.IsTrue(result.Values!.SequenceEqual(["value1", "value2"]));
+			Assert.IsTrue(result.ValueDescriptions!.SequenceEqual(["desc1", "desc2"]));
 		}
 
 		#endregion
@@ -403,7 +403,7 @@ namespace EpicGames.BuildGraph.Tests
 		public void ListEmpty()
 		{
 			object result = Evaluate(BgList<BgInt>.Empty);
-			Assert.IsTrue(Enumerable.SequenceEqual(Enumerable.Empty<object>(), (IEnumerable<object>)result));
+			Assert.IsTrue(Enumerable.SequenceEqual([], (IEnumerable<object>)result));
 		}
 
 		[TestMethod]
@@ -424,58 +424,58 @@ namespace EpicGames.BuildGraph.Tests
 		public void ListPush()
 		{
 			object result1 = Evaluate(BgList<BgInt>.Empty.Add(1));
-			Assert.IsTrue(Enumerable.SequenceEqual(new object[] { 1 }, (IEnumerable<object>)result1));
+			Assert.IsTrue(Enumerable.SequenceEqual([1], (IEnumerable<object>)result1));
 
 			object result2 = Evaluate(BgList<BgInt>.Empty.Add(1, 2).Add(3));
-			Assert.IsTrue(Enumerable.SequenceEqual(new object[] { 1, 2, 3 }, (IEnumerable<object>)result2));
+			Assert.IsTrue(Enumerable.SequenceEqual([1, 2, 3], (IEnumerable<object>)result2));
 		}
 
 		[TestMethod]
 		public void ListConcat()
 		{
 			object result = Evaluate(BgList<BgInt>.Concat(BgList<BgInt>.Create(1, 2), BgList<BgInt>.Create(3)));
-			Assert.IsTrue(Enumerable.SequenceEqual(new object[] { 1, 2, 3 }, (IEnumerable<object>)result));
+			Assert.IsTrue(Enumerable.SequenceEqual([1, 2, 3], (IEnumerable<object>)result));
 		}
 
 		[TestMethod]
 		public void ListUnion()
 		{
 			object result1 = Evaluate(BgList<BgInt>.Create(1, 2).Union(BgList<BgInt>.Create(3)));
-			Assert.IsTrue(Enumerable.SequenceEqual(new object[] { 1, 2, 3 }, (IEnumerable<object>)result1));
+			Assert.IsTrue(Enumerable.SequenceEqual([1, 2, 3], (IEnumerable<object>)result1));
 
 			object result2 = Evaluate(BgList<BgInt>.Create(1, 2, 3).Union(BgList<BgInt>.Create(3)));
-			Assert.IsTrue(Enumerable.SequenceEqual(new object[] { 1, 2, 3 }, (IEnumerable<object>)result2));
+			Assert.IsTrue(Enumerable.SequenceEqual([1, 2, 3], (IEnumerable<object>)result2));
 		}
 
 		[TestMethod]
 		public void ListExcept()
 		{
 			object result1 = Evaluate(BgList<BgInt>.Create(1, 2).Except(BgList<BgInt>.Create(3)));
-			Assert.IsTrue(Enumerable.SequenceEqual(new object[] { 1, 2 }, (IEnumerable<object>)result1));
+			Assert.IsTrue(Enumerable.SequenceEqual([1, 2], (IEnumerable<object>)result1));
 
 			object result2 = Evaluate(BgList<BgInt>.Create(1, 2, 3).Except(BgList<BgInt>.Create(3)));
-			Assert.IsTrue(Enumerable.SequenceEqual(new object[] { 1, 2 }, (IEnumerable<object>)result2));
+			Assert.IsTrue(Enumerable.SequenceEqual([1, 2], (IEnumerable<object>)result2));
 		}
 
 		[TestMethod]
 		public void ListSelect()
 		{
 			object result = Evaluate(BgList<BgInt>.Create(1, 1, 2, 2, 3).Select(x => x * 9));
-			Assert.IsTrue(Enumerable.SequenceEqual(new object[] { 9, 9, 18, 18, 27 }, (IEnumerable<object>)result));
+			Assert.IsTrue(Enumerable.SequenceEqual([9, 9, 18, 18, 27], (IEnumerable<object>)result));
 		}
 
 		[TestMethod]
 		public void ListWhere()
 		{
 			object result = Evaluate(BgList<BgInt>.Create(1, 1, 2, 2, 3).Where(x => x >= 2));
-			Assert.IsTrue(Enumerable.SequenceEqual(new object[] { 2, 2, 3 }, (IEnumerable<object>)result));
+			Assert.IsTrue(Enumerable.SequenceEqual([2, 2, 3], (IEnumerable<object>)result));
 		}
 
 		[TestMethod]
 		public void ListDistinct()
 		{
 			object result = Evaluate(BgList<BgInt>.Concat(BgList<BgInt>.Empty.Add(1, 1, 2, 2, 3), BgList<BgInt>.Empty.Add(3)).Distinct());
-			Assert.IsTrue(Enumerable.SequenceEqual(new object[] { 1, 2, 3 }, (IEnumerable<object>)result));
+			Assert.IsTrue(Enumerable.SequenceEqual([1, 2, 3], (IEnumerable<object>)result));
 		}
 
 		[TestMethod]
@@ -485,8 +485,8 @@ namespace EpicGames.BuildGraph.Tests
 			Assert.AreEqual("stringOption", result.Name);
 			Assert.AreEqual("description", result.Description);
 			Assert.AreEqual("value", result.DefaultValue);
-			Assert.IsTrue(result.Values!.SequenceEqual(new[] { "value1", "value2" }));
-			Assert.IsTrue(result.ValueDescriptions!.SequenceEqual(new[] { "desc1", "desc2" }));
+			Assert.IsTrue(result.Values!.SequenceEqual(["value1", "value2"]));
+			Assert.IsTrue(result.ValueDescriptions!.SequenceEqual(["desc1", "desc2"]));
 		}
 
 		#endregion
diff --git a/Engine/Source/Programs/Shared/EpicGames.BuildGraph.Tests/SerializationTests.cs b/Engine/Source/Programs/Shared/EpicGames.BuildGraph.Tests/SerializationTests.cs
--- a/Engine/Source/Programs/Shared/EpicGames.BuildGraph.Tests/SerializationTests.cs
+++ b/Engine/Source/Programs/Shared/EpicGames.BuildGraph.Tests/SerializationTests.cs
@@ -17,7 +17,7 @@ namespace EpicGames.BuildGraph.Tests
 		{
 			public NestedObject? Nested { get; set; }
 			public List<string>? StringList { get; set; }
-			public List<NestedObject> NestedList { get; } = new List<NestedObject>();
+			public List<NestedObject> NestedList { get; } = [];
 		}
 
 		[TestMethod]
diff --git a/Engine/Source/Programs/Shared/EpicGames.BuildGraph/BgAgentDef.cs b/Engine/Source/Programs/Shared/EpicGames.BuildGraph/BgAgentDef.cs
--- a/Engine/Source/Programs/Shared/EpicGames.BuildGraph/BgAgentDef.cs
+++ b/Engine/Source/Programs/Shared/EpicGames.BuildGraph/BgAgentDef.cs
@@ -21,17 +21,17 @@ namespace EpicGames.BuildGraph
 		/// Array of valid agent types that these nodes may run on. When running in the build system, this determines the class of machine that should
 		/// be selected to run these nodes. The first defined agent type for this branch will be used.
 		/// </summary>
-		public List<string> PossibleTypes { get; } = new List<string>();
+		public List<string> PossibleTypes { get; } = [];
 
 		/// <summary>
 		/// List of nodes in this agent group.
 		/// </summary>
-		public List<BgNodeDef> Nodes { get; set; } = new List<BgNodeDef>();
+		public List<BgNodeDef> Nodes { get; set; } = [];
 
 		/// <summary>
 		/// Diagnostics to output if executing this agent
 		/// </summary>
-		public List<BgDiagnosticDef> Diagnostics { get; } = new List<BgDiagnosticDef>();
+		public List<BgDiagnosticDef> Diagnostics { get; } = [];
 
 		/// <summary>
 		/// Constructor
diff --git a/Engine/Source/Programs/Shared/EpicGames.BuildGraph/BgAggregateDef.cs b/Engine/Source/Programs/Shared/EpicGames.BuildGraph/BgAggregateDef.cs
--- a/Engine/Source/Programs/Shared/EpicGames.BuildGraph/BgAggregateDef.cs
+++ b/Engine/Source/Programs/Shared/EpicGames.BuildGraph/BgAggregateDef.cs
@@ -17,7 +17,7 @@ namespace EpicGames.BuildGraph
 		/// <summary>
 		/// Set of nodes that must be run for this label to be shown.
 		/// </summary>
-		public HashSet<BgNodeDef> RequiredNodes { get; } = new HashSet<BgNodeDef>();
+		public HashSet<BgNodeDef> RequiredNodes { get; } = [];
 
 		/// <summary>
 		/// Constructor
@@ -48,7 +48,7 @@ namespace EpicGames.BuildGraph
 		public string Name { get; set; }
 
 		/// <inheritdoc cref="BgAggregateDef.RequiredNodes"/>
-		public List<BgNodeExpressionDef> RequiredNodes { get; } = new List<BgNodeExpressionDef>();
+		public List<BgNodeExpressionDef> RequiredNodes { get; } = [];
 
 		/// <summary>
 		/// Labels to add this aggregate to
diff --git a/Engine/Source/Programs/Shared/EpicGames.BuildGraph/BgBadgeDef.cs b/Engine/Source/Programs/Shared/EpicGames.BuildGraph/BgBadgeDef.cs
--- a/Engine/Source/Programs/Shared/EpicGames.BuildGraph/BgBadgeDef.cs
+++ b/Engine/Source/Programs/Shared/EpicGames.BuildGraph/BgBadgeDef.cs
@@ -27,7 +27,7 @@ namespace EpicGames.BuildGraph
 		/// <summary>
 		/// Set of nodes that this badge reports the status of
 		/// </summary>
-		public HashSet<BgNodeDef> Nodes { get; } = new HashSet<BgNodeDef>();
+		public HashSet<BgNodeDef> Nodes { get; } = [];
 
 		/// <summary>
 		/// Constructor
diff --git a/Engine/Source/Programs/Shared/EpicGames.BuildGraph/BgCompiler.cs b/Engine/Source/Programs/Shared/EpicGames.BuildGraph/BgCompiler.cs
--- a/Engine/Source/Programs/Shared/EpicGames.BuildGraph/BgCompiler.cs
+++ b/Engine/Source/Programs/Shared/EpicGames.BuildGraph/BgCompiler.cs
@@ -112,7 +112,7 @@ namespace EpicGames.BuildGraph
 			readonly Dictionary<BgExpr, int> _exprToFragmentIdx;
 			readonly List<string> _names;
 			readonly List<BgThunkDef> _thunks;
-			readonly Dictionary<string, int> _nameToIndex = new Dictionary<string, int>();
+			readonly Dictionary<string, int> _nameToIndex = [];
 
 			public ForwardWriter(ByteArrayBuilder builder, Dictionary<BgExpr, int> exprToFragmentIdx, List<string> names, List<BgThunkDef> thunks)
 			{
@@ -199,7 +199,7 @@ namespace EpicGames.BuildGraph
 		/// <returns>Compiled bytecode for the expression, suitable for passing to <see cref="BgInterpreter"/></returns>
 		public static (byte[], BgThunkDef[]) Compile(BgExpr expr)
 		{
-			List<BgExpr> fragments = new List<BgExpr>();
+			List<BgExpr> fragments = [];
 			Dictionary<BgExpr, int> exprToFragmentIndex = new Dictionary<BgExpr, int>(ReferenceEqualityComparer.Instance);
 
 			// Figure out which expressions need to be compiled into separate fragments
@@ -210,8 +210,8 @@ namespace EpicGames.BuildGraph
 			ByteArrayBuilder code = new ByteArrayBuilder();
 			List<int> fragmentLengths = new List<int>(fragments.Count);
 
-			List<BgThunkDef> thunks = new List<BgThunkDef>();
-			List<string> names = new List<string>();
+			List<BgThunkDef> thunks = [];
+			List<string> names = [];
 
 			ForwardWriter writer = new ForwardWriter(code, exprToFragmentIndex, names, thunks);
 			for (int idx = 0; idx < fragments.Count; idx++)
diff --git a/Engine/Source/Programs/Shared/EpicGames.BuildGraph/BgContext.cs b/Engine/Source/Programs/Shared/EpicGames.BuildGraph/BgContext.cs
--- a/Engine/Source/Programs/Shared/EpicGames.BuildGraph/BgContext.cs
+++ b/Engine/Source/Programs/Shared/EpicGames.BuildGraph/BgContext.cs
@@ -40,7 +40,7 @@ namespace EpicGames.BuildGraph
 		/// <summary>
 		/// All outputs for the node
 		/// </summary>
-		public HashSet<FileReference> BuildProducts { get; } = new HashSet<FileReference>();
+		public HashSet<FileReference> BuildProducts { get; } = [];
 
 		readonly Dictionary<string, FileSet> _tagNameToFileSet;
 
@@ -87,7 +87,7 @@ namespace EpicGames.BuildGraph
 		/// <typeparam name="TEnum">The enum type</typeparam>
 		/// <param name="expr">Enum expression</param>
 		/// <returns>The enum value</returns>
-		public List<TEnum> Get<TEnum>(BgList<BgEnum<TEnum>> expr) where TEnum : struct => ((BgListConstantExpr<BgEnum<TEnum>>)expr).Value.Select(x => ((BgEnumConstantExpr<TEnum>)x).Value).ToList();
+		public List<TEnum> Get<TEnum>(BgList<BgEnum<TEnum>> expr) where TEnum : struct => [.. ((BgListConstantExpr<BgEnum<TEnum>>)expr).Value.Select(x => ((BgEnumConstantExpr<TEnum>)x).Value)];
 
 		/// <summary>
 		/// Resolve a list of strings
diff --git a/Engine/Source/Programs/Shared/EpicGames.BuildGraph/BgGraphDef.cs b/Engine/Source/Programs/Shared/EpicGames.BuildGraph/BgGraphDef.cs
--- a/Engine/Source/Programs/Shared/EpicGames.BuildGraph/BgGraphDef.cs
+++ b/Engine/Source/Programs/Shared/EpicGames.BuildGraph/BgGraphDef.cs
@@ -48,12 +48,12 @@ namespace EpicGames.BuildGraph
 		/// <summary>
 		/// List of options, in the order they were specified
 		/// </summary>
-		public List<BgOptionDef> Options { get; } = new List<BgOptionDef>();
+		public List<BgOptionDef> Options { get; } = [];
 
 		/// <summary>
 		/// List of agents containing nodes to execute
 		/// </summary>
-		public List<BgAgentDef> Agents { get; } = new List<BgAgentDef>();
+		public List<BgAgentDef> Agents { get; } = [];
 
 		/// <summary>
 		/// Mapping from name to agent
@@ -83,22 +83,22 @@ namespace EpicGames.BuildGraph
 		/// <summary>
 		/// Mapping of artifact names to their definitions. Artifacts will be produced from matching tag names.
 		/// </summary>
-		public List<BgArtifactDef> Artifacts { get; } = new List<BgArtifactDef>();
+		public List<BgArtifactDef> Artifacts { get; } = [];
 
 		/// <summary>
 		/// List of badges that can be displayed for this build
 		/// </summary>
-		public List<BgBadgeDef> Badges { get; } = new List<BgBadgeDef>();
+		public List<BgBadgeDef> Badges { get; } = [];
 
 		/// <summary>
 		/// List of labels that can be displayed for this build
 		/// </summary>
-		public List<BgLabelDef> Labels { get; } = new List<BgLabelDef>();
+		public List<BgLabelDef> Labels { get; } = [];
 
 		/// <summary>
 		/// Diagnostics at graph scope
 		/// </summary>
-		public List<BgDiagnosticDef> Diagnostics { get; } = new List<BgDiagnosticDef>();
+		public List<BgDiagnosticDef> Diagnostics { get; } = [];
 
 		/// <summary>
 		/// Default constructor
@@ -123,7 +123,7 @@ namespace EpicGames.BuildGraph
 		/// <returns>List of diagnostics</returns>
 		public List<BgDiagnosticDef> GetAllDiagnostics()
 		{
-			List<BgDiagnosticDef> diagnostics = new List<BgDiagnosticDef>(Diagnostics);
+			List<BgDiagnosticDef> diagnostics = [.. Diagnostics];
 			foreach (BgAgentDef agent in Agents)
 			{
 				diagnostics.AddRange(agent.Diagnostics);
@@ -150,7 +150,7 @@ namespace EpicGames.BuildGraph
 				BgNodeOutput? output;
 				if (TagNameToNodeOutput.TryGetValue(name, out output))
 				{
-					outNodes = new BgNodeDef[] { output.ProducingNode };
+					outNodes = [output.ProducingNode];
 					return true;
 				}
 			}
@@ -160,7 +160,7 @@ namespace EpicGames.BuildGraph
 				BgNodeDef? node;
 				if (NameToNode.TryGetValue(name, out node))
 				{
-					outNodes = new BgNodeDef[] { node };
+					outNodes = [node];
 					return true;
 				}
 
@@ -168,7 +168,7 @@ namespace EpicGames.BuildGraph
 				BgAggregateDef? aggregate;
 				if (NameToAggregate.TryGetValue(name, out aggregate))
 				{
-					outNodes = aggregate.RequiredNodes.ToArray();
+					outNodes = [.. aggregate.RequiredNodes];
 					return true;
 				}
 
@@ -176,7 +176,7 @@ namespace EpicGames.BuildGraph
 				BgAgentDef? agent;
 				if (NameToAgent.TryGetValue(name, out agent))
 				{
-					outNodes = agent.Nodes.ToArray();
+					outNodes = [.. agent.Nodes];
 					return true;
 				}
 			}
@@ -201,7 +201,7 @@ namespace EpicGames.BuildGraph
 				BgNodeOutput? output;
 				if (TagNameToNodeOutput.TryGetValue(name, out output))
 				{
-					outOutputs = new BgNodeOutput[] { output };
+					outOutputs = [output];
 					return true;
 				}
 			}
@@ -211,7 +211,7 @@ namespace EpicGames.BuildGraph
 				BgNodeDef? node;
 				if (NameToNode.TryGetValue(name, out node))
 				{
-					outOutputs = node.Outputs.Union(node.Inputs).ToArray();
+					outOutputs = [.. node.Outputs.Union(node.Inputs)];
 					return true;
 				}
 
@@ -219,7 +219,7 @@ namespace EpicGames.BuildGraph
 				BgAggregateDef? aggregate;
 				if (NameToAggregate.TryGetValue(name, out aggregate))
 				{
-					outOutputs = aggregate.RequiredNodes.SelectMany(x => x.Outputs.Union(x.Inputs)).Distinct().ToArray();
+					outOutputs = [.. aggregate.RequiredNodes.SelectMany(x => x.Outputs.Union(x.Inputs)).Distinct()];
 					return true;
 				}
 			}
@@ -247,7 +247,7 @@ namespace EpicGames.BuildGraph
 		public void Select(IEnumerable<BgNodeDef> targetNodes)
 		{
 			// Find this node and all its dependencies
-			HashSet<BgNodeDef> retainNodes = new HashSet<BgNodeDef>();
+			HashSet<BgNodeDef> retainNodes = [];
 			foreach (BgNodeDef targetNode in targetNodes)
 			{
 				AddDependencies(targetNode, retainNodes);
@@ -256,7 +256,7 @@ namespace EpicGames.BuildGraph
 			// Remove all the nodes which are not marked to be kept
 			foreach (BgAgentDef agent in Agents)
 			{
-				agent.Nodes = agent.Nodes.Where(x => retainNodes.Contains(x)).ToList();
+				agent.Nodes = [.. agent.Nodes.Where(x => retainNodes.Contains(x))];
 			}
 
 			// Remove all the empty agents
@@ -323,7 +323,7 @@ namespace EpicGames.BuildGraph
 		public void Export(FileReference file, HashSet<BgNodeDef> completedNodes)
 		{
 			// Find all the nodes which we're actually going to execute. We'll use this to filter the graph.
-			HashSet<BgNodeDef> nodesToExecute = new HashSet<BgNodeDef>();
+			HashSet<BgNodeDef> nodesToExecute = [];
 			foreach (BgNodeDef node in Agents.SelectMany(x => x.Nodes))
 			{
 				if (!completedNodes.Contains(node))
@@ -341,7 +341,7 @@ namespace EpicGames.BuildGraph
 				jsonWriter.WriteArrayStart("Groups");
 				foreach (BgAgentDef agent in Agents)
 				{
-					BgNodeDef[] nodes = agent.Nodes.Where(x => nodesToExecute.Contains(x)).ToArray();
+					BgNodeDef[] nodes = [.. agent.Nodes.Where(x => nodesToExecute.Contains(x))];
 					if (nodes.Length > 0)
 					{
 						jsonWriter.WriteObjectStart();
@@ -377,11 +377,11 @@ namespace EpicGames.BuildGraph
 				jsonWriter.WriteArrayStart("Badges");
 				foreach (BgBadgeDef badge in Badges)
 				{
-					BgNodeDef[] dependencies = badge.Nodes.Where(x => nodesToExecute.Contains(x)).ToArray();
+					BgNodeDef[] dependencies = [.. badge.Nodes.Where(x => nodesToExecute.Contains(x))];
 					if (dependencies.Length > 0)
 					{
 						// Reduce that list to the smallest subset of direct dependencies
-						HashSet<BgNodeDef> directDependencies = new HashSet<BgNodeDef>(dependencies);
+						HashSet<BgNodeDef> directDependencies = [.. dependencies];
 						foreach (BgNodeDef dependency in dependencies)
 						{
 							directDependencies.ExceptWith(dependency.OrderDependencies);
@@ -408,11 +408,11 @@ namespace EpicGames.BuildGraph
 				jsonWriter.WriteArrayStart("Reports");
 				foreach (BgReport report in NameToReport.Values)
 				{
-					BgNodeDef[] dependencies = report.Nodes.Where(x => nodesToExecute.Contains(x)).ToArray();
+					BgNodeDef[] dependencies = [.. report.Nodes.Where(x => nodesToExecute.Contains(x))];
 					if (dependencies.Length > 0)
 					{
 						// Reduce that list to the smallest subset of direct dependencies
-						HashSet<BgNodeDef> directDependencies = new HashSet<BgNodeDef>(dependencies);
+						HashSet<BgNodeDef> directDependencies = [.. dependencies];
 						foreach (BgNodeDef dependency in dependencies)
 						{
 							directDependencies.ExceptWith(dependency.OrderDependencies);
@@ -619,7 +619,7 @@ namespace EpicGames.BuildGraph
 					if (dependencies.Count > 0)
 					{
 						// Reduce that list to the smallest subset of direct dependencies
-						HashSet<BgNodeDef> directDependencies = new HashSet<BgNodeDef>(dependencies);
+						HashSet<BgNodeDef> directDependencies = [.. dependencies];
 						foreach (BgNodeDef dependency in dependencies)
 						{
 							directDependencies.ExceptWith(dependency.OrderDependencies);
@@ -657,7 +657,7 @@ namespace EpicGames.BuildGraph
 			if ((printOptions & GraphPrintOptions.ShowCommandLineOptions) != 0)
 			{
 				// Get the list of messages
-				List<KeyValuePair<string, string>> parameters = new List<KeyValuePair<string, string>>();
+				List<KeyValuePair<string, string>> parameters = [];
 				foreach (BgOptionDef option in Options)
 				{
 					string name = String.Format("-set:{0}=...", option.Name);
@@ -693,7 +693,7 @@ namespace EpicGames.BuildGraph
 					logger.LogInformation("Options:");
 					logger.LogInformation("");
 
-					List<string> lines = new List<string>();
+					List<string> lines = [];
 					HelpUtils.FormatTable(parameters, 4, 24, HelpUtils.WindowWidth - 20, lines);
 
 					foreach (string line in lines)
@@ -714,12 +714,12 @@ namespace EpicGames.BuildGraph
 					logger.LogInformation("            Node: {Name}{Type}{Retry}", node.Name, completedNodes.Contains(node) ? " (completed)" : node.RunEarly ? " (early)" : "", !node.AllowRetry ? " (retry disabled)" : "");
 					if (printOptions.HasFlag(GraphPrintOptions.ShowDependencies))
 					{
-						HashSet<BgNodeDef> inputDependencies = new HashSet<BgNodeDef>(node.GetDirectInputDependencies());
+						HashSet<BgNodeDef> inputDependencies = [.. node.GetDirectInputDependencies()];
 						foreach (BgNodeDef inputDependency in inputDependencies)
 						{
 							logger.LogInformation("                input> {Name}", inputDependency.Name);
 						}
-						HashSet<BgNodeDef> orderDependencies = new HashSet<BgNodeDef>(node.GetDirectOrderDependencies());
+						HashSet<BgNodeDef> orderDependencies = [.. node.GetDirectOrderDependencies()];
 						foreach (BgNodeDef orderDependency in orderDependencies.Except(inputDependencies))
 						{
 							logger.LogInformation("                after> {Name}", orderDependency.Name);
@@ -742,7 +742,7 @@ namespace EpicGames.BuildGraph
 			logger.LogInformation("");
 
 			// Print out all the non-empty aggregates
-			BgAggregateDef[] aggregates = NameToAggregate.Values.OrderBy(x => x.Name).ToArray();
+			BgAggregateDef[] aggregates = [.. NameToAggregate.Values.OrderBy(x => x.Name)];
 			if (aggregates.Length > 0)
 			{
 				logger.LogInformation("Aggregates:");
@@ -754,7 +754,7 @@ namespace EpicGames.BuildGraph
 			}
 
 			// Print all the produced artifacts
-			BgArtifactDef[] artifacts = Artifacts.OrderBy(x => x.Name).ToArray();
+			BgArtifactDef[] artifacts = [.. Artifacts.OrderBy(x => x.Name)];
 			if (artifacts.Length > 0)
 			{
 				logger.LogInformation("Artifacts:");
@@ -775,12 +775,12 @@ namespace EpicGames.BuildGraph
 		/// <summary>
 		/// Nodes for the graph
 		/// </summary>
-		public List<BgNodeExpressionDef> Nodes { get; } = new List<BgNodeExpressionDef>();
+		public List<BgNodeExpressionDef> Nodes { get; } = [];
 
 		/// <summary>
 		/// Aggregates for the graph
 		/// </summary>
-		public List<BgAggregateExpressionDef> Aggregates { get; } = new List<BgAggregateExpressionDef>();
+		public List<BgAggregateExpressionDef> Aggregates { get; } = [];
 
 		/// <summary>
 		/// Creates a graph definition from this template
@@ -788,7 +788,7 @@ namespace EpicGames.BuildGraph
 		/// <returns></returns>
 		public BgGraphDef ToGraphDef()
 		{
-			List<BgNodeExpressionDef> nodes = new List<BgNodeExpressionDef>(Nodes);
+			List<BgNodeExpressionDef> nodes = [.. Nodes];
 
 			BgGraphDef graph = new BgGraphDef();
 			foreach (BgAggregateExpressionDef aggregate in Aggregates)
@@ -797,14 +797,14 @@ namespace EpicGames.BuildGraph
 				nodes.AddRange(aggregate.RequiredNodes.Select(x => (BgNodeExpressionDef)x));
 			}
 
-			HashSet<BgNodeDef> uniqueNodes = new HashSet<BgNodeDef>();
-			HashSet<BgAgentDef> uniqueAgents = new HashSet<BgAgentDef>();
+			HashSet<BgNodeDef> uniqueNodes = [];
+			HashSet<BgAgentDef> uniqueAgents = [];
 			foreach (BgNodeExpressionDef node in nodes)
 			{
 				RegisterNode(graph, node, uniqueNodes, uniqueAgents);
 			}
 
-			HashSet<BgLabelDef> labels = new HashSet<BgLabelDef>();
+			HashSet<BgLabelDef> labels = [];
 			foreach (BgNodeExpressionDef node in nodes)
 			{
 				foreach (BgLabelDef label in node.Labels)
diff --git a/Engine/Source/Programs/Shared/EpicGames.BuildGraph/BgInterpreter.cs b/Engine/Source/Programs/Shared/EpicGames.BuildGraph/BgInterpreter.cs
--- a/Engine/Source/Programs/Shared/EpicGames.BuildGraph/BgInterpreter.cs
+++ b/Engine/Source/Programs/Shared/EpicGames.BuildGraph/BgInterpreter.cs
@@ -108,7 +108,7 @@ namespace EpicGames.BuildGraph
 			{
 				Offset = offset;
 				Arguments = arguments;
-				Objects = new Dictionary<int, object>();
+				Objects = [];
 			}
 		}
 
@@ -118,7 +118,7 @@ namespace EpicGames.BuildGraph
 		readonly IReadOnlyDictionary<string, string> _options;
 		readonly BgBytecodeVersion _version;
 		readonly int[] _fragments;
-		readonly List<BgOptionDef> _optionDefs = new List<BgOptionDef>();
+		readonly List<BgOptionDef> _optionDefs = [];
 
 		/// <summary>
 		/// The option definitions that were parsed during execution
@@ -233,7 +233,7 @@ namespace EpicGames.BuildGraph
 			// Take the given parameters, evaluate the graph that was produced (creating a map of all named entites)
 			// Filter the map to the required target
 			// Evaluate it fully into a BgGraph
-			return Evaluate(new Frame(_fragments[0], Array.Empty<object>()));
+			return Evaluate(new Frame(_fragments[0], []));
 		}
 
 		object Evaluate(Frame frame)
@@ -518,7 +518,7 @@ namespace EpicGames.BuildGraph
 						IEnumerable<object> list = (IEnumerable<object>)arg0;
 						object item = arg1;
 
-						return list.Concat(new[] { item });
+						return list.Concat([item]);
 					}
 				case BgOpcode.ListPushLazy:
 					{
@@ -562,14 +562,14 @@ namespace EpicGames.BuildGraph
 						IEnumerable<object> source = (IEnumerable<object>)arg0;
 						int fragment = (int)arg1;
 
-						return source.Select(x => Call(fragment, new object[] { x }));
+						return source.Select(x => Call(fragment, [x]));
 					}
 				case BgOpcode.ListWhere:
 					{
 						IEnumerable<object> source = (IEnumerable<object>)arg0;
 						int fragment = (int)arg1;
 
-						return source.Where(x => (bool)Call(fragment, new object[] { x }));
+						return source.Where(x => (bool)Call(fragment, [x]));
 					}
 				case BgOpcode.ListDistinct:
 					{
@@ -600,7 +600,7 @@ namespace EpicGames.BuildGraph
 							value = result;
 						}
 
-						return value!.Split(new[] { '+', ';' }, StringSplitOptions.RemoveEmptyEntries).Select(x => (object)x).ToList();
+						return value!.Split(['+', ';'], StringSplitOptions.RemoveEmptyEntries).Select(x => (object)x).ToList();
 					}
 
 				#endregion
@@ -704,7 +704,7 @@ namespace EpicGames.BuildGraph
 
 		static List<T> GetListArg<T>(object value)
 		{
-			return ((IEnumerable<object>)value).Select(x => (T)x).ToList();
+			return [.. ((IEnumerable<object>)value).Select(x => (T)x)];
 		}
 
 		IEnumerable<object> LazyEvaluateItem(Frame frame, int fragment)
@@ -833,7 +833,7 @@ namespace EpicGames.BuildGraph
 			{
 				logger.LogInformation("");
 				logger.LogInformation("Fragment {Idx}:", idx);
-				Disassemble(new Frame(_fragments[idx], Array.Empty<object>()), logger);
+				Disassemble(new Frame(_fragments[idx], []), logger);
 			}
 		}
 
diff --git a/Engine/Source/Programs/Shared/EpicGames.BuildGraph/BgLabelDef.cs b/Engine/Source/Programs/Shared/EpicGames.BuildGraph/BgLabelDef.cs
--- a/Engine/Source/Programs/Shared/EpicGames.BuildGraph/BgLabelDef.cs
+++ b/Engine/Source/Programs/Shared/EpicGames.BuildGraph/BgLabelDef.cs
@@ -55,12 +55,12 @@ namespace EpicGames.BuildGraph
 		/// <summary>
 		/// Set of nodes that must be run for this label to be shown.
 		/// </summary>
-		public HashSet<BgNodeDef> RequiredNodes { get; } = new HashSet<BgNodeDef>();
+		public HashSet<BgNodeDef> RequiredNodes { get; } = [];
 
 		/// <summary>
 		/// Set of nodes that will be included in this label if present.
 		/// </summary>
-		public HashSet<BgNodeDef> IncludedNodes { get; } = new HashSet<BgNodeDef>();
+		public HashSet<BgNodeDef> IncludedNodes { get; } = [];
 
 		/// <summary>
 		/// Constructor
diff --git a/Engine/Source/Programs/Shared/EpicGames.BuildGraph/BgNodeDef.cs b/Engine/Source/Programs/Shared/EpicGames.BuildGraph/BgNodeDef.cs
--- a/Engine/Source/Programs/Shared/EpicGames.BuildGraph/BgNodeDef.cs
+++ b/Engine/Source/Programs/Shared/EpicGames.BuildGraph/BgNodeDef.cs
@@ -79,11 +79,11 @@ namespace EpicGames.BuildGraph
 		{
 			if (OutputIndex == -1)
 			{
-				return ProducingNode.InputDependencies.SelectMany(x => x.Outputs).Concat(ProducingNode.Outputs).ToArray();
+				return [.. ProducingNode.InputDependencies.SelectMany(x => x.Outputs), .. ProducingNode.Outputs];
 			}
 			else
 			{
-				return new[] { ProducingNode.Outputs[OutputIndex] };
+				return [ProducingNode.Outputs[OutputIndex]];
 			}
 		}
 	}
@@ -115,7 +115,7 @@ namespace EpicGames.BuildGraph
 		/// <summary>
 		/// Array of inputs which this node requires to run
 		/// </summary>
-		public List<BgNodeOutput> Inputs { get; } = new List<BgNodeOutput>();
+		public List<BgNodeOutput> Inputs { get; } = [];
 
 		/// <summary>
 		/// Array of outputs produced by this node
@@ -125,17 +125,17 @@ namespace EpicGames.BuildGraph
 		/// <summary>
 		/// Nodes which this node has input dependencies on
 		/// </summary>
-		public List<BgNodeDef> InputDependencies { get; } = new List<BgNodeDef>();
+		public List<BgNodeDef> InputDependencies { get; } = [];
 
 		/// <summary>
 		/// Nodes which this node needs to run after
 		/// </summary>
-		public List<BgNodeDef> OrderDependencies { get; } = new List<BgNodeDef>();
+		public List<BgNodeDef> OrderDependencies { get; } = [];
 
 		/// <summary>
 		/// Tokens which must be acquired for this node to run
 		/// </summary>
-		public List<FileReference> RequiredTokens { get; } = new List<FileReference>();
+		public List<FileReference> RequiredTokens { get; } = [];
 
 		/// <summary>
 		/// List of email addresses to notify if this node fails.
@@ -170,12 +170,12 @@ namespace EpicGames.BuildGraph
 		/// <summary>
 		/// Ignore modified files matching the patterns provided
 		/// </summary>
-		public List<string> IgnoreModified { get; set; } = new List<string>();
+		public List<string> IgnoreModified { get; set; } = [];
 
 		/// <summary>
 		/// Diagnostics to output if executing this node
 		/// </summary>
-		public List<BgDiagnosticDef> Diagnostics { get; } = new List<BgDiagnosticDef>();
+		public List<BgDiagnosticDef> Diagnostics { get; } = [];
 
 		/// <summary>
 		/// Constructor
@@ -185,10 +185,12 @@ namespace EpicGames.BuildGraph
 			Name = name;
 			Thunk = thunk;
 
-			List<BgNodeOutput> allOutputs = new List<BgNodeOutput>();
-			allOutputs.Add(new BgNodeOutput(this, "#" + Name));
-			allOutputs.AddRange(outputNames.Where(x => !String.Equals(x, Name, StringComparison.OrdinalIgnoreCase)).Select(x => new BgNodeOutput(this, x)));
-			Outputs = allOutputs.ToArray();
+			List<BgNodeOutput> allOutputs =
+			[
+				new BgNodeOutput(this, "#" + Name),
+				.. outputNames.Where(x => !String.Equals(x, Name, StringComparison.OrdinalIgnoreCase)).Select(x => new BgNodeOutput(this, x)),
+			];
+			Outputs = [.. allOutputs];
 		}
 
 		/// <summary>
@@ -223,7 +225,7 @@ namespace EpicGames.BuildGraph
 		/// <returns>Sequence of nodes that are direct inputs to this node</returns>
 		public IEnumerable<BgNodeDef> GetDirectInputDependencies()
 		{
-			HashSet<BgNodeDef> directDependencies = new HashSet<BgNodeDef>(InputDependencies);
+			HashSet<BgNodeDef> directDependencies = [.. InputDependencies];
 			foreach (BgNodeDef inputDependency in InputDependencies)
 			{
 				directDependencies.ExceptWith(inputDependency.InputDependencies);
@@ -237,7 +239,7 @@ namespace EpicGames.BuildGraph
 		/// <returns>Sequence of nodes that are direct order dependencies of this node</returns>
 		public IEnumerable<BgNodeDef> GetDirectOrderDependencies()
 		{
-			HashSet<BgNodeDef> directDependencies = new HashSet<BgNodeDef>(OrderDependencies);
+			HashSet<BgNodeDef> directDependencies = [.. OrderDependencies];
 			foreach (BgNodeDef orderDependency in OrderDependencies)
 			{
 				directDependencies.ExceptWith(orderDependency.OrderDependencies);
@@ -269,12 +271,12 @@ namespace EpicGames.BuildGraph
 		/// <summary>
 		/// Labels to add this node to
 		/// </summary>
-		public List<BgLabelDef> Labels { get; } = new List<BgLabelDef>();
+		public List<BgLabelDef> Labels { get; } = [];
 
 		/// <summary>
 		/// Input expressions
 		/// </summary>
-		public List<BgNodeOutputExprDef> InputExprs { get; } = new List<BgNodeOutputExprDef>();
+		public List<BgNodeOutputExprDef> InputExprs { get; } = [];
 
 		/// <summary>
 		/// Number of outputs from this node
@@ -293,7 +295,7 @@ namespace EpicGames.BuildGraph
 
 		static string[] GetOutputNames(string name, int numOutputs)
 		{
-			return Enumerable.Range(0, numOutputs).Select(x => BgNode.GetDefaultTagName(name, x)).ToArray();
+			return [.. Enumerable.Range(0, numOutputs).Select(x => BgNode.GetDefaultTagName(name, x))];
 		}
 	}
 
@@ -307,7 +309,7 @@ namespace EpicGames.BuildGraph
 
 			node.Inputs.AddRange(node.InputExprs.SelectMany(x => x.Flatten()));
 
-			HashSet<BgNodeDef> inputDependencies = new HashSet<BgNodeDef>();
+			HashSet<BgNodeDef> inputDependencies = [];
 			foreach (BgNodeOutput input in node.Inputs)
 			{
 				inputDependencies.Add(input.ProducingNode);
diff --git a/Engine/Source/Programs/Shared/EpicGames.BuildGraph/BgReport.cs b/Engine/Source/Programs/Shared/EpicGames.BuildGraph/BgReport.cs
--- a/Engine/Source/Programs/Shared/EpicGames.BuildGraph/BgReport.cs
+++ b/Engine/Source/Programs/Shared/EpicGames.BuildGraph/BgReport.cs
@@ -18,7 +18,7 @@ namespace EpicGames.BuildGraph
 		/// <summary>
 		/// Set of nodes to include in the report
 		/// </summary>
-		public HashSet<BgNodeDef> Nodes { get; } = new HashSet<BgNodeDef>();
+		public HashSet<BgNodeDef> Nodes { get; } = [];
 
 		/// <summary>
 		/// List of users to notify with this report
diff --git a/Engine/Source/Programs/Shared/EpicGames.BuildGraph/Expressions/BgFileSet.cs b/Engine/Source/Programs/Shared/EpicGames.BuildGraph/Expressions/BgFileSet.cs
--- a/Engine/Source/Programs/Shared/EpicGames.BuildGraph/Expressions/BgFileSet.cs
+++ b/Engine/Source/Programs/Shared/EpicGames.BuildGraph/Expressions/BgFileSet.cs
@@ -2,7 +2,6 @@
 
 using System;
 using System.Collections.Generic;
-using System.Linq;
 using EpicGames.Core;
 
 namespace EpicGames.BuildGraph.Expressions
@@ -49,7 +48,7 @@ namespace EpicGames.BuildGraph.Expressions
 		public override BgFileSet Constant(object value)
 		{
 			BgObjectDef obj = (BgObjectDef)value;
-			BgNodeOutput[] outputs = obj.Deserialize<BgNodeOutputExprDef>().Flatten().ToArray();
+			BgNodeOutput[] outputs = [.. obj.Deserialize<BgNodeOutputExprDef>().Flatten()];
 			return new BgFileSetInputExpr(outputs);
 		}
 
diff --git a/Engine/Source/Programs/Shared/EpicGames.BuildGraph/Expressions/BgList.cs b/Engine/Source/Programs/Shared/EpicGames.BuildGraph/Expressions/BgList.cs
--- a/Engine/Source/Programs/Shared/EpicGames.BuildGraph/Expressions/BgList.cs
+++ b/Engine/Source/Programs/Shared/EpicGames.BuildGraph/Expressions/BgList.cs
@@ -235,7 +235,7 @@ namespace EpicGames.BuildGraph.Expressions
 		{
 			IEnumerable<object> elements = (IEnumerable<object>)value;
 
-			List<T> items = new List<T>();
+			List<T> items = [];
 			foreach (object element in elements)
 			{
 				items.Add(BgType.Constant<T>(element));
@@ -271,7 +271,7 @@ namespace EpicGames.BuildGraph.Expressions
 		public BgListConstantExpr(IEnumerable<T> value)
 			: base(BgExprFlags.NotInterned)
 		{
-			Value = value.ToArray();
+			Value = [.. value];
 		}
 
 		public override void Write(BgBytecodeWriter writer)
diff --git a/Engine/Source/Programs/Shared/EpicGames.BuildGraph/Expressions/BgNode.cs b/Engine/Source/Programs/Shared/EpicGames.BuildGraph/Expressions/BgNode.cs
--- a/Engine/Source/Programs/Shared/EpicGames.BuildGraph/Expressions/BgNode.cs
+++ b/Engine/Source/Programs/Shared/EpicGames.BuildGraph/Expressions/BgNode.cs
@@ -204,7 +204,7 @@ namespace EpicGames.BuildGraph.Expressions
 		static BgString GetNodeNameFromTemplate(string template, ParameterInfo[] parameters, IReadOnlyList<object?> arguments)
 		{
 			// Create a list of lazily computed string fragments which comprise the evaluated name
-			List<BgString> fragments = new List<BgString>();
+			List<BgString> fragments = [];
 
 			int lastIdx = 0;
 			for (int nextIdx = 0; nextIdx < template.Length; nextIdx++)
@@ -375,7 +375,7 @@ namespace EpicGames.BuildGraph.Expressions
 			}
 			else
 			{
-				BgExpr[] outputs = CreateOutputExprs(new[] { type });
+				BgExpr[] outputs = CreateOutputExprs([type]);
 				OutputCount = outputs.Length;
 				return (T)(object)outputs[0];
 			}
diff --git a/Engine/Source/Programs/Shared/EpicGames.Core.Tests/BuzHashTests.cs b/Engine/Source/Programs/Shared/EpicGames.Core.Tests/BuzHashTests.cs
--- a/Engine/Source/Programs/Shared/EpicGames.Core.Tests/BuzHashTests.cs
+++ b/Engine/Source/Programs/Shared/EpicGames.Core.Tests/BuzHashTests.cs
@@ -73,7 +73,7 @@ namespace EpicGames.Core.Tests
 				{
 					Assert.AreEqual(lastOffset + MaxSize, expected[idx].Offset);
 					idx++;
-					lastOffset = lastOffset + MaxSize;
+					lastOffset += MaxSize;
 				}
 				if (offset >= lastOffset + MinSize)
 				{
diff --git a/Engine/Source/Programs/Shared/EpicGames.Core.Tests/JsonObjectTests.cs b/Engine/Source/Programs/Shared/EpicGames.Core.Tests/JsonObjectTests.cs
--- a/Engine/Source/Programs/Shared/EpicGames.Core.Tests/JsonObjectTests.cs
+++ b/Engine/Source/Programs/Shared/EpicGames.Core.Tests/JsonObjectTests.cs
@@ -212,7 +212,7 @@ namespace EpicGames.Core.Tests
 		public void GetObjectArrayValidField()
 		{
 			string fieldName = "Attacks";
-			JsonObject[] correctValue = _validJsonDataSource.Attacks.Select(x => x.ToJsonObject()).ToArray();
+			JsonObject[] correctValue = [.. _validJsonDataSource.Attacks.Select(x => x.ToJsonObject())];
 			JsonObject[] result = _validTestObject.GetObjectArrayField(fieldName);
 			CollectionAssert.AreEqual(correctValue, result);
 			JsonObject[]? outValue;
@@ -440,7 +440,7 @@ namespace EpicGames.Core.Tests
 				new EnemyAttack("Attack2", 100, StatusEffect.Burn)
 			];
 			Assert.IsTrue(_validTestObject.ContainsField(fieldName));
-			JsonObject[] setValue = attacks.Select(x => x.ToJsonObject()).ToArray();
+			JsonObject[] setValue = [.. attacks.Select(x => x.ToJsonObject())];
 			_validTestObject.AddOrSetFieldValue(fieldName, setValue);
 			CollectionAssert.AreEqual(setValue, _validTestObject.GetObjectArrayField(fieldName));
 
@@ -536,7 +536,7 @@ namespace EpicGames.Core.Tests
 		public void ToJsonString()
 		{
 			string stringToMatch = new string(_testJsonText.ToCharArray());
-			string[] lines = stringToMatch.Split(new[] { Environment.NewLine }, StringSplitOptions.None);
+			string[] lines = stringToMatch.Split([Environment.NewLine], StringSplitOptions.None);
 			StringBuilder jsonStringBuilder = new StringBuilder();
 			foreach (string line in lines)
 			{
diff --git a/Engine/Source/Programs/Shared/EpicGames.Core.Tests/JsonTracerTests.cs b/Engine/Source/Programs/Shared/EpicGames.Core.Tests/JsonTracerTests.cs
--- a/Engine/Source/Programs/Shared/EpicGames.Core.Tests/JsonTracerTests.cs
+++ b/Engine/Source/Programs/Shared/EpicGames.Core.Tests/JsonTracerTests.cs
@@ -135,7 +135,7 @@ namespace EpicGames.Core.Tests
 
 		private static async Task ForEachAsyncSpanTestAsync(JsonTracer tracer, IScope parentScope)
 		{
-			int[] values = Enumerable.Range(0, 10000).ToArray();
+			int[] values = [.. Enumerable.Range(0, 10000)];
 			await Parallel.ForEachAsync(values, new ParallelOptions { MaxDegreeOfParallelism = 10000 }, async (i, token) =>
 			{
 				using (IScope childScope = tracer.BuildSpan("Child").StartActive())
diff --git a/Engine/Source/Programs/Shared/EpicGames.Core/BinaryArchiveReader.cs b/Engine/Source/Programs/Shared/EpicGames.Core/BinaryArchiveReader.cs
--- a/Engine/Source/Programs/Shared/EpicGames.Core/BinaryArchiveReader.cs
+++ b/Engine/Source/Programs/Shared/EpicGames.Core/BinaryArchiveReader.cs
@@ -68,11 +68,8 @@ namespace EpicGames.Core
 		/// </summary>
 		public void Dispose()
 		{
-			if (_stream != null)
-			{
-				_stream.Dispose();
-				_stream = null;
-			}
+			_stream?.Dispose();
+			_stream = null;
 
 			_buffer = null;
 		}
diff --git a/Engine/Source/Programs/Shared/EpicGames.Core/BinaryArchiveWriter.cs b/Engine/Source/Programs/Shared/EpicGames.Core/BinaryArchiveWriter.cs
--- a/Engine/Source/Programs/Shared/EpicGames.Core/BinaryArchiveWriter.cs
+++ b/Engine/Source/Programs/Shared/EpicGames.Core/BinaryArchiveWriter.cs
@@ -79,12 +79,8 @@ namespace EpicGames.Core
 		public void Dispose()
 		{
 			Flush();
-
-			if (_stream != null)
-			{
-				_stream.Dispose();
-				_stream = null;
-			}
+			_stream?.Dispose();
+			_stream = null;
 		}
 
 		/// <summary>
diff --git a/Engine/Source/Programs/Shared/EpicGames.Core/BuzHash.cs b/Engine/Source/Programs/Shared/EpicGames.Core/BuzHash.cs
--- a/Engine/Source/Programs/Shared/EpicGames.Core/BuzHash.cs
+++ b/Engine/Source/Programs/Shared/EpicGames.Core/BuzHash.cs
@@ -213,7 +213,7 @@ namespace EpicGames.Core
 		/// <returns>List of chunk offsets and hashes</returns>
 		public static List<(int, uint)> FindSplitPoints(ReadOnlySpan<byte> span, int minSize, int maxSize, int targetSize)
 		{
-			List<(int, uint)> chunkPoints = new List<(int, uint)>();
+			List<(int, uint)> chunkPoints = [];
 
 			int offset = 0;
 			for (; ; )
@@ -242,7 +242,7 @@ namespace EpicGames.Core
 		/// <returns>List of chunk offsets and hashes</returns>
 		public static unsafe List<(int, uint)> FindCandidateSplitPoints(ReadOnlySpan<byte> span, int windowSize, uint threshold)
 		{
-			List<(int, uint)> offsets = new List<(int, uint)>();
+			List<(int, uint)> offsets = [];
 			if (span.Length < windowSize)
 			{
 				return offsets;
diff --git a/Engine/Source/Programs/Shared/EpicGames.Core/CommandHost.cs b/Engine/Source/Programs/Shared/EpicGames.Core/CommandHost.cs
--- a/Engine/Source/Programs/Shared/EpicGames.Core/CommandHost.cs
+++ b/Engine/Source/Programs/Shared/EpicGames.Core/CommandHost.cs
@@ -204,7 +204,7 @@ namespace EpicGames.Core
 		public static async Task<int> RunAsync(CommandLineArguments args, IServiceProvider serviceProvider, Type? defaultCommandType, string? toolDescription = null)
 		{
 			// Find all the command types
-			List<ICommandFactory> commandFactories = serviceProvider.GetServices<ICommandFactory>().ToList();
+			List<ICommandFactory> commandFactories = [.. serviceProvider.GetServices<ICommandFactory>()];
 
 			// Check if there's a matching command
 			ICommand? command = null;
diff --git a/Engine/Source/Programs/Shared/EpicGames.Core/CommandLineArguments.cs b/Engine/Source/Programs/Shared/EpicGames.Core/CommandLineArguments.cs
--- a/Engine/Source/Programs/Shared/EpicGames.Core/CommandLineArguments.cs
+++ b/Engine/Source/Programs/Shared/EpicGames.Core/CommandLineArguments.cs
@@ -874,7 +874,7 @@ namespace EpicGames.Core
 					IEnumerable<CommandLineAttribute> attributes = fieldInfo.GetCustomAttributes<CommandLineAttribute>();
 					if (attributes.Any())
 					{
-						targets.Add(new ArgumentTarget(fieldInfo, fieldInfo.FieldType, fieldInfo.SetValue, fieldInfo.GetValue, attributes.ToArray()));
+						targets.Add(new ArgumentTarget(fieldInfo, fieldInfo.FieldType, fieldInfo.SetValue, fieldInfo.GetValue, [.. attributes]));
 					}
 				}
 				foreach (PropertyInfo propertyInfo in targetType.GetProperties(BindingFlags.Instance | BindingFlags.GetProperty | BindingFlags.Public | BindingFlags.NonPublic | BindingFlags.DeclaredOnly))
@@ -882,7 +882,7 @@ namespace EpicGames.Core
 					IEnumerable<CommandLineAttribute> attributes = propertyInfo.GetCustomAttributes<CommandLineAttribute>();
 					if (attributes.Any())
 					{
-						targets.Add(new ArgumentTarget(propertyInfo, propertyInfo.PropertyType, propertyInfo.SetValue, propertyInfo.GetValue, attributes.ToArray()));
+						targets.Add(new ArgumentTarget(propertyInfo, propertyInfo.PropertyType, propertyInfo.SetValue, propertyInfo.GetValue, [.. attributes]));
 					}
 				}
 
@@ -1063,7 +1063,7 @@ namespace EpicGames.Core
 		/// <returns>New argument list</returns>
 		public CommandLineArguments Append(IEnumerable<string> appendArguments)
 		{
-			CommandLineArguments newArguments = new CommandLineArguments(Enumerable.Concat(_arguments, appendArguments).ToArray());
+			CommandLineArguments newArguments = new CommandLineArguments([.. Enumerable.Concat(_arguments, appendArguments)]);
 			for (int idx = 0; idx < _arguments.Length; idx++)
 			{
 				if (HasBeenUsed(idx))
diff --git a/Engine/Source/Programs/Shared/EpicGames.Core/FileFilter.cs b/Engine/Source/Programs/Shared/EpicGames.Core/FileFilter.cs
--- a/Engine/Source/Programs/Shared/EpicGames.Core/FileFilter.cs
+++ b/Engine/Source/Programs/Shared/EpicGames.Core/FileFilter.cs
@@ -152,7 +152,7 @@ namespace EpicGames.Core
 				}
 
 				// Check there's a matching tag
-				string[] ruleTags = cleanRule.Substring(1, conditionEnd - 1).Split(',').Select(x => x.Trim()).ToArray();
+				string[] ruleTags = [.. cleanRule.Substring(1, conditionEnd - 1).Split(',').Select(x => x.Trim())];
 				if (!ruleTags.Any(x => allowTags.Contains(x)))
 				{
 					return;
@@ -416,7 +416,7 @@ namespace EpicGames.Core
 		{
 			string[] tokens = NormalizePath(folderName).Trim('/').Split('/');
 
-			FileFilterNode matchingNode = FindMatchingNode(_rootNode, tokens.Union([""]).ToArray(), 0, _defaultNode);
+			FileFilterNode matchingNode = FindMatchingNode(_rootNode, [.. tokens.Union([""])], 0, _defaultNode);
 
 			return matchingNode.Type == FileFilterType.Include || HighestPossibleIncludeMatch(_rootNode, tokens, 0, matchingNode.RuleNumber) > matchingNode.RuleNumber;
 		}
diff --git a/Engine/Source/Programs/Shared/EpicGames.Core/FileSet.cs b/Engine/Source/Programs/Shared/EpicGames.Core/FileSet.cs
--- a/Engine/Source/Programs/Shared/EpicGames.Core/FileSet.cs
+++ b/Engine/Source/Programs/Shared/EpicGames.Core/FileSet.cs
@@ -475,7 +475,7 @@ namespace EpicGames.Core
 			{
 				string filterName = _inner.Path + item.Key;
 
-				FileFilter[] possibleFilters = _filters.Where(x => x.PossiblyMatches(filterName)).ToArray();
+				FileFilter[] possibleFilters = [.. _filters.Where(x => x.PossiblyMatches(filterName))];
 				if (possibleFilters.Length > 0)
 				{
 					FileSetFromFilter subTreeFilter = new FileSetFromFilter(item.Value, possibleFilters);
diff --git a/Engine/Source/Programs/Shared/EpicGames.Core/IMemoryAllocator.cs b/Engine/Source/Programs/Shared/EpicGames.Core/IMemoryAllocator.cs
--- a/Engine/Source/Programs/Shared/EpicGames.Core/IMemoryAllocator.cs
+++ b/Engine/Source/Programs/Shared/EpicGames.Core/IMemoryAllocator.cs
@@ -192,16 +192,11 @@ namespace EpicGames.Core
 						_viewAccessor!.SafeMemoryMappedViewHandle.ReleasePointer();
 						_pointer = null;
 					}
-					if (_viewAccessor != null)
-					{
-						_viewAccessor.Dispose();
-						_viewAccessor = null;
-					}
-					if (_file != null)
-					{
-						_file.Dispose();
-						_file = null;
-					}
+
+					_viewAccessor?.Dispose();
+					_viewAccessor = null;
+					_file?.Dispose();
+					_file = null;
 				}
 			}
 		}
diff --git a/Engine/Source/Programs/Shared/EpicGames.Core/IReadOnlyMemoryOwner.cs b/Engine/Source/Programs/Shared/EpicGames.Core/IReadOnlyMemoryOwner.cs
--- a/Engine/Source/Programs/Shared/EpicGames.Core/IReadOnlyMemoryOwner.cs
+++ b/Engine/Source/Programs/Shared/EpicGames.Core/IReadOnlyMemoryOwner.cs
@@ -37,11 +37,8 @@ namespace EpicGames.Core
 			public void Dispose()
 			{
 				_memory = null;
-				if (_owner != null)
-				{
-					_owner.Dispose();
-					_owner = null;
-				}
+				_owner?.Dispose();
+				_owner = null;
 			}
 		}
 
diff --git a/Engine/Source/Programs/Shared/EpicGames.Core/JsonObject.cs b/Engine/Source/Programs/Shared/EpicGames.Core/JsonObject.cs
--- a/Engine/Source/Programs/Shared/EpicGames.Core/JsonObject.cs
+++ b/Engine/Source/Programs/Shared/EpicGames.Core/JsonObject.cs
@@ -441,7 +441,7 @@ namespace EpicGames.Core
 			object? rawValue;
 			if (_rawOrderedObject.TryGetValue(fieldName, out rawValue) && (rawValue is IEnumerable<object> enumValue) && enumValue.All(x => x is string))
 			{
-				result = enumValue.Select(x => (string)x).ToArray();
+				result = [.. enumValue.Select(x => (string)x)];
 				return true;
 			}
 			else
@@ -722,7 +722,7 @@ namespace EpicGames.Core
 			object? rawValue;
 			if (_rawOrderedObject.TryGetValue(fieldName, out rawValue) && (rawValue is IEnumerable<object> enumValue) && enumValue.All(x => x is OrderedDictionary))
 			{
-				result = enumValue.Select(x => new JsonObject((OrderedDictionary)x, false)).ToArray();
+				result = [.. enumValue.Select(x => new JsonObject((OrderedDictionary)x, false))];
 				return true;
 			}
 			else
@@ -878,7 +878,7 @@ namespace EpicGames.Core
 				return;
 			}
 			// Contents of the strings array should never be null and should be "" instead to match behavior EpicGames.Core.JsonWriter
-			string[] stringArray = value.Select(x => x ?? "").ToArray();
+			string[] stringArray = [.. value.Select(x => x ?? "")];
 			_rawOrderedObject[fieldName] = stringArray;
 		}
 
@@ -923,7 +923,7 @@ namespace EpicGames.Core
 				return;
 			}
 
-			string[] stringArray = value.Select(x => x.ToString()!).ToArray();
+			string[] stringArray = [.. value.Select(x => x.ToString()!)];
 			_rawOrderedObject[fieldName] = stringArray;
 		}
 
@@ -950,7 +950,7 @@ namespace EpicGames.Core
 			}
 			// The Utf8JsonWriter doesn't format json the same as we want in UE, we massage the string to meet our standards 
 			string jsonString = Encoding.UTF8.GetString(jsonMemoryStream.ToArray());
-			string[] lines = jsonString.Split(new[] { Environment.NewLine }, StringSplitOptions.None);
+			string[] lines = jsonString.Split([Environment.NewLine], StringSplitOptions.None);
 			StringBuilder jsonStringBuilder = new StringBuilder();
 			foreach (string line in lines)
 			{
diff --git a/Engine/Source/Programs/Shared/EpicGames.Core/JsonTracer.cs b/Engine/Source/Programs/Shared/EpicGames.Core/JsonTracer.cs
--- a/Engine/Source/Programs/Shared/EpicGames.Core/JsonTracer.cs
+++ b/Engine/Source/Programs/Shared/EpicGames.Core/JsonTracer.cs
@@ -467,7 +467,7 @@ namespace EpicGames.Core
 		{
 			lock (_lock)
 			{
-				return new List<JsonTracerSpan>(_finishedSpans);
+				return [.. _finishedSpans];
 			}
 		}
 
diff --git a/Engine/Source/Programs/Shared/EpicGames.Core/ListExtensions.cs b/Engine/Source/Programs/Shared/EpicGames.Core/ListExtensions.cs
--- a/Engine/Source/Programs/Shared/EpicGames.Core/ListExtensions.cs
+++ b/Engine/Source/Programs/Shared/EpicGames.Core/ListExtensions.cs
@@ -21,7 +21,9 @@ namespace EpicGames.Core
 		{
 			readonly IList<T> _inner;
 
-			public ReadOnlyList(IList<T> inner) => _inner = inner;
+			private ReadOnlyList(IList<T> inner) => _inner = inner;
+
+			public static ReadOnlyList<T> Wrap(IList<T> inner) => new ReadOnlyList<T>(inner);
 
 			static NotSupportedException CreateNotSupportedException()
 				=> new NotSupportedException("Collection is read-only.");
@@ -54,7 +56,7 @@ namespace EpicGames.Core
 		/// <typeparam name="T"></typeparam>
 		/// <param name="list"></param>
 		/// <returns></returns>
-		public static IList<T> AsReadOnly<T>(this IList<T> list) => new ReadOnlyList<T>(list);
+		public static IList<T> AsReadOnly<T>(this IList<T> list) => ReadOnlyList<T>.Wrap(list);
 
 		/// <summary>
 		/// Create a read-only wrapper around a list
@@ -62,7 +64,7 @@ namespace EpicGames.Core
 		/// <typeparam name="T"></typeparam>
 		/// <param name="list"></param>
 		/// <returns></returns>
-		public static IReadOnlyList<T> AsReadOnlyList<T>(this IList<T> list) => new ReadOnlyList<T>(list);
+		public static IReadOnlyList<T> AsReadOnlyList<T>(this IList<T> list) => ReadOnlyList<T>.Wrap(list);
 
 		/// <summary>
 		/// Sorts a list by a particular field
diff --git a/Engine/Source/Programs/Shared/EpicGames.Core/Log.cs b/Engine/Source/Programs/Shared/EpicGames.Core/Log.cs
--- a/Engine/Source/Programs/Shared/EpicGames.Core/Log.cs
+++ b/Engine/Source/Programs/Shared/EpicGames.Core/Log.cs
@@ -1129,7 +1129,7 @@ namespace EpicGames.Core
 					Trace.Listeners.Add(logTraceListener);
 					WriteInitialTimestamp();
 
-					List<FileStream> newJsonStreams = new List<FileStream>(_jsonStreams);
+					List<FileStream> newJsonStreams = [.. _jsonStreams];
 					newJsonStreams.Add(FileReference.Open(outputFile.ChangeExtension(".json"), FileMode.Create, FileAccess.Write, FileShare.Read | FileShare.Delete));
 					_jsonStreams = newJsonStreams;
 				}
diff --git a/Engine/Source/Programs/Shared/EpicGames.Core/LogEventParser.cs b/Engine/Source/Programs/Shared/EpicGames.Core/LogEventParser.cs
--- a/Engine/Source/Programs/Shared/EpicGames.Core/LogEventParser.cs
+++ b/Engine/Source/Programs/Shared/EpicGames.Core/LogEventParser.cs
@@ -94,7 +94,7 @@ namespace EpicGames.Core
 		public LogEventMatch(LogEventPriority priority, IEnumerable<LogEvent> events, bool isComplete)
 		{
 			Priority = priority;
-			Events = events.ToList();
+			Events = [.. events];
 			IsComplete = isComplete;
 		}
 	}
diff --git a/Engine/Source/Programs/Shared/EpicGames.Core/ManagedProcess.cs b/Engine/Source/Programs/Shared/EpicGames.Core/ManagedProcess.cs
--- a/Engine/Source/Programs/Shared/EpicGames.Core/ManagedProcess.cs
+++ b/Engine/Source/Programs/Shared/EpicGames.Core/ManagedProcess.cs
@@ -222,11 +222,8 @@ namespace EpicGames.Core
 		/// </summary>
 		public void Dispose()
 		{
-			if (JobHandle != null)
-			{
-				JobHandle.Dispose();
-				JobHandle = null;
-			}
+			JobHandle?.Dispose();
+			JobHandle = null;
 		}
 	}
 
@@ -760,21 +757,12 @@ namespace EpicGames.Core
 						finally
 						{
 							// Close the write ends of the handle. We don't want any other process to be able to inherit these.
-							if (stdInRead != null)
-							{
-								stdInRead.Dispose();
-								stdInRead = null;
-							}
-							if (stdOutWrite != null)
-							{
-								stdOutWrite.Dispose();
-								stdOutWrite = null;
-							}
-							if (stdErrWrite != null)
-							{
-								stdErrWrite.Dispose();
-								stdErrWrite = null;
-							}
+							stdInRead?.Dispose();
+							stdInRead = null;
+							stdOutWrite?.Dispose();
+							stdOutWrite = null;
+							stdErrWrite?.Dispose();
+							stdErrWrite = null;
 						}
 					}
 
@@ -981,16 +969,11 @@ namespace EpicGames.Core
 				_processHandle.Dispose();
 				_processHandle = null;
 			}
-			if (_stdInWrite != null)
-			{
-				_stdInWrite.Dispose();
-				_stdInWrite = null;
-			}
-			if (_stdOutRead != null)
-			{
-				_stdOutRead.Dispose();
-				_stdOutRead = null;
-			}
+
+			_stdInWrite?.Dispose();
+			_stdInWrite = null;
+			_stdOutRead?.Dispose();
+			_stdOutRead = null;
 
 			StdIn?.Dispose();
 			StdOut?.Dispose();
diff --git a/Engine/Source/Programs/Shared/EpicGames.Core/MemoryMappedFileCache.cs b/Engine/Source/Programs/Shared/EpicGames.Core/MemoryMappedFileCache.cs
--- a/Engine/Source/Programs/Shared/EpicGames.Core/MemoryMappedFileCache.cs
+++ b/Engine/Source/Programs/Shared/EpicGames.Core/MemoryMappedFileCache.cs
@@ -65,22 +65,12 @@ namespace EpicGames.Core
 			{
 				Interlocked.Add(ref _cache._mappedSize, -_data.Length);
 				_data = ReadOnlyMemory<byte>.Empty;
-
-				if (_memoryMappedView != null)
-				{
-					_memoryMappedView.Dispose();
-					_memoryMappedView = null;
-				}
-				if (_memoryMappedViewAccessor != null)
-				{
-					_memoryMappedViewAccessor.Dispose();
-					_memoryMappedViewAccessor = null;
-				}
-				if (_memoryMappedFile != null)
-				{
-					_memoryMappedFile.Dispose();
-					_memoryMappedFile = null;
-				}
+				_memoryMappedView?.Dispose();
+				_memoryMappedView = null;
+				_memoryMappedViewAccessor?.Dispose();
+				_memoryMappedViewAccessor = null;
+				_memoryMappedFile?.Dispose();
+				_memoryMappedFile = null;
 
 				if (_deleteOnDispose)
 				{
diff --git a/Engine/Source/Programs/Shared/EpicGames.Core/PrefixLogger.cs b/Engine/Source/Programs/Shared/EpicGames.Core/PrefixLogger.cs
--- a/Engine/Source/Programs/Shared/EpicGames.Core/PrefixLogger.cs
+++ b/Engine/Source/Programs/Shared/EpicGames.Core/PrefixLogger.cs
@@ -36,7 +36,7 @@ namespace EpicGames.Core
 		{
 			if (state is IEnumerable<KeyValuePair<string, object>> enumerable)
 			{
-				List<KeyValuePair<string, object>> copy = new List<KeyValuePair<string, object>>(enumerable);
+				List<KeyValuePair<string, object>> copy = [.. enumerable];
 
 				int idx = copy.FindIndex(x => x.Key.Equals("{OriginalFormat}", StringComparison.OrdinalIgnoreCase));
 				if (idx != -1 && copy[idx].Value is string format)
diff --git a/Engine/Source/Programs/Shared/EpicGames.Core/ProcessUtils.cs b/Engine/Source/Programs/Shared/EpicGames.Core/ProcessUtils.cs
--- a/Engine/Source/Programs/Shared/EpicGames.Core/ProcessUtils.cs
+++ b/Engine/Source/Programs/Shared/EpicGames.Core/ProcessUtils.cs
@@ -182,7 +182,7 @@ namespace EpicGames.Core
 					}
 
 					// Wait for them all to complete
-					if (WaitForMultipleObjects(waitHandles.Count, waitHandles.Select(x => x.DangerousGetHandle()).ToArray(), true, 10 * 1000) == WAIT_FAILED)
+					if (WaitForMultipleObjects(waitHandles.Count, [.. waitHandles.Select(x => x.DangerousGetHandle())], true, 10 * 1000) == WAIT_FAILED)
 					{
 						logger.LogInformation("Failed call to WaitForMultipleObjects ({Code})", Marshal.GetLastWin32Error());
 					}
diff --git a/Engine/Source/Programs/Shared/EpicGames.Core/ProgressTextReader.cs b/Engine/Source/Programs/Shared/EpicGames.Core/ProgressTextReader.cs
--- a/Engine/Source/Programs/Shared/EpicGames.Core/ProgressTextReader.cs
+++ b/Engine/Source/Programs/Shared/EpicGames.Core/ProgressTextReader.cs
@@ -222,7 +222,7 @@ namespace EpicGames.Core
 
 		static List<string> ParseTokens(string line)
 		{
-			List<string> tokens = new List<string>();
+			List<string> tokens = [];
 			for (int idx = 0; ;)
 			{
 				// Skip whitespace
diff --git a/Engine/Source/Programs/Shared/EpicGames.Core/ReadOnlyHashSet.cs b/Engine/Source/Programs/Shared/EpicGames.Core/ReadOnlyHashSet.cs
--- a/Engine/Source/Programs/Shared/EpicGames.Core/ReadOnlyHashSet.cs
+++ b/Engine/Source/Programs/Shared/EpicGames.Core/ReadOnlyHashSet.cs
@@ -33,7 +33,7 @@ namespace EpicGames.Core
 		/// <param name="elements">Elements for the hash set</param>
 		public ReadOnlyHashSet(IEnumerable<T> elements)
 		{
-			_inner = new HashSet<T>(elements);
+			_inner = [.. elements];
 		}
 
 		/// <summary>
diff --git a/Engine/Source/Programs/Shared/EpicGames.Core/StreamWrapper.cs b/Engine/Source/Programs/Shared/EpicGames.Core/StreamWrapper.cs
--- a/Engine/Source/Programs/Shared/EpicGames.Core/StreamWrapper.cs
+++ b/Engine/Source/Programs/Shared/EpicGames.Core/StreamWrapper.cs
@@ -172,11 +172,8 @@ namespace EpicGames.Core
 
 			void DisposeResource()
 			{
-				if (_resource != null)
-				{
-					_resource.Dispose();
-					_resource = null;
-				}
+				_resource?.Dispose();
+				_resource = null;
 			}
 		}
 
diff --git a/Engine/Source/Programs/Shared/EpicGames.Core/StringUtils.cs b/Engine/Source/Programs/Shared/EpicGames.Core/StringUtils.cs
--- a/Engine/Source/Programs/Shared/EpicGames.Core/StringUtils.cs
+++ b/Engine/Source/Programs/Shared/EpicGames.Core/StringUtils.cs
@@ -263,7 +263,7 @@ namespace EpicGames.Core
 		/// <returns>Formatted list of strings</returns>
 		public static string FormatList(IEnumerable<string> arguments, string conjunction = "and")
 		{
-			return FormatList(arguments.ToArray(), conjunction);
+			return FormatList([.. arguments], conjunction);
 		}
 
 		/// <summary>
diff --git a/Engine/Source/Programs/Shared/EpicGames.Horde.Tests/AesTransportTests.cs b/Engine/Source/Programs/Shared/EpicGames.Horde.Tests/AesTransportTests.cs
--- a/Engine/Source/Programs/Shared/EpicGames.Horde.Tests/AesTransportTests.cs
+++ b/Engine/Source/Programs/Shared/EpicGames.Horde.Tests/AesTransportTests.cs
@@ -81,10 +81,9 @@ public sealed class AesTransportTests : IAsyncDisposable
 		foreach (int bufferSize in _bufferSizes)
 		{
 			List<int> evilChunkSizes = [1, 2, 3, bufferSize - 1, bufferSize, bufferSize + 1, bufferSize * 5];
-			evilChunkSizes = evilChunkSizes
+			evilChunkSizes = [.. evilChunkSizes
 				.Where(size => size > 0)
-				.Where(size => bufferSize < SmallBufferThreshold || size >= MinChunkSizeForLargeBuffers)
-				.ToList();
+				.Where(size => bufferSize < SmallBufferThreshold || size >= MinChunkSizeForLargeBuffers)];
 
 			foreach (int sendChunk in evilChunkSizes)
 			{
diff --git a/Engine/Source/Programs/Shared/EpicGames.Horde.Tests/Compute/ComputeSocketTests.cs b/Engine/Source/Programs/Shared/EpicGames.Horde.Tests/Compute/ComputeSocketTests.cs
--- a/Engine/Source/Programs/Shared/EpicGames.Horde.Tests/Compute/ComputeSocketTests.cs
+++ b/Engine/Source/Programs/Shared/EpicGames.Horde.Tests/Compute/ComputeSocketTests.cs
@@ -28,8 +28,8 @@ namespace EpicGames.Horde.Tests.Compute
 		public const int TestTimoutMs = 40000;
 		class TestComputeSocket : ComputeSocket, IDisposable
 		{
-			public Dictionary<int, ComputeBufferWriter> RecvBufferWriters { get; } = new Dictionary<int, ComputeBufferWriter>();
-			public Dictionary<int, ComputeBufferReader> SendBufferReaders { get; } = new Dictionary<int, ComputeBufferReader>();
+			public Dictionary<int, ComputeBufferWriter> RecvBufferWriters { get; } = [];
+			public Dictionary<int, ComputeBufferReader> SendBufferReaders { get; } = [];
 
 			public override ComputeProtocol Protocol => ComputeProtocol.Latest;
 			public override ILogger Logger => NullLogger.Instance;
@@ -167,7 +167,7 @@ namespace EpicGames.Horde.Tests.Compute
 					using (AgentMessage message = await channel.ReceiveAsync(cancellationToken))
 					{
 						Assert.AreEqual(AgentMessageType.Ping, message.Type);
-						Assert.IsTrue(message.Data.Span.SequenceEqual(ReadOnlySpan<byte>.Empty));
+						Assert.IsTrue(message.Data.Span.SequenceEqual([]));
 					}
 
 					await channel.SendXorRequestAsync(new byte[] { 1, 2, 3 }, 44, cancellationToken);
@@ -226,13 +226,13 @@ namespace EpicGames.Horde.Tests.Compute
 						byte[] readData = await FileReference.ReadAllBytesAsync(file, cancellationToken);
 						Assert.IsTrue(readData.SequenceEqual(data));
 
-						await channel.DeleteFilesAsync(new[] { "subdir/hello.txt" }, cancellationToken);
+						await channel.DeleteFilesAsync(["subdir/hello.txt"], cancellationToken);
 
 						await channel.PingAsync(cancellationToken);
 						using (AgentMessage message = await channel.ReceiveAsync(cancellationToken))
 						{
 							Assert.AreEqual(AgentMessageType.Ping, message.Type);
-							Assert.IsTrue(message.Data.Span.SequenceEqual(ReadOnlySpan<byte>.Empty));
+							Assert.IsTrue(message.Data.Span.SequenceEqual([]));
 						}
 
 						Assert.IsFalse(FileReference.Exists(file));
diff --git a/Engine/Source/Programs/Shared/EpicGames.Horde.Tests/Compute/RequestIdAllocatorTests.cs b/Engine/Source/Programs/Shared/EpicGames.Horde.Tests/Compute/RequestIdAllocatorTests.cs
--- a/Engine/Source/Programs/Shared/EpicGames.Horde.Tests/Compute/RequestIdAllocatorTests.cs
+++ b/Engine/Source/Programs/Shared/EpicGames.Horde.Tests/Compute/RequestIdAllocatorTests.cs
@@ -13,9 +13,9 @@ public class RequestIdAllocatorTests
 	public void IdsAreReused()
 	{
 		RequestIdAllocator ria = new();
-		List<string> before = new() { ria.AllocateId(), ria.AllocateId(), ria.AllocateId() };
+		List<string> before = [ria.AllocateId(), ria.AllocateId(), ria.AllocateId()];
 		ria.StartBatch();
-		List<string> after = new() { ria.AllocateId(), ria.AllocateId(), ria.AllocateId() };
+		List<string> after = [ria.AllocateId(), ria.AllocateId(), ria.AllocateId()];
 
 		CollectionAssert.AreEquivalent(before, after);
 	}
@@ -24,11 +24,11 @@ public class RequestIdAllocatorTests
 	public void AcceptedIdsAreNotReused()
 	{
 		RequestIdAllocator ria = new();
-		List<string> before = new() { ria.AllocateId(), ria.AllocateId(), ria.AllocateId() };
+		List<string> before = [ria.AllocateId(), ria.AllocateId(), ria.AllocateId()];
 		ria.MarkAccepted(before[1]);
 
 		ria.StartBatch();
-		List<string> after = new() { ria.AllocateId(), ria.AllocateId(), ria.AllocateId() };
+		List<string> after = [ria.AllocateId(), ria.AllocateId(), ria.AllocateId()];
 		Assert.IsTrue(after.Remove(before[0]));
 		Assert.IsFalse(after.Remove(before[1]));
 		Assert.IsTrue(after.Remove(before[2]));
diff --git a/Engine/Source/Programs/Shared/EpicGames.Horde.Tests/HordeHttpTests.cs b/Engine/Source/Programs/Shared/EpicGames.Horde.Tests/HordeHttpTests.cs
--- a/Engine/Source/Programs/Shared/EpicGames.Horde.Tests/HordeHttpTests.cs
+++ b/Engine/Source/Programs/Shared/EpicGames.Horde.Tests/HordeHttpTests.cs
@@ -30,7 +30,7 @@ public class FakeServerWithAuth : HttpMessageHandler
 	public const string ServerUrl = "http://fake-server-test";
 	public const string SuccessResponse = "fakeServerSuccess";
 
-	public List<HttpRequestMessage> HttpRequests { get; } = new();
+	public List<HttpRequestMessage> HttpRequests { get; } = [];
 	private readonly GetAuthConfigResponse? _authConfigResponse;
 	private readonly FakeOidcTokenManager _oidcTokenManager;
 
@@ -40,7 +40,7 @@ public class FakeServerWithAuth : HttpMessageHandler
 		{
 			Method = AuthMethod.OpenIdConnect,
 			ServerUrl = ServerUrl,
-			LocalRedirectUrls = new[] { ServerUrl }
+			LocalRedirectUrls = [ServerUrl]
 		};
 		_oidcTokenManager = oidcTokenManager;
 	}
diff --git a/Engine/Source/Programs/Shared/EpicGames.Horde.Tests/Logs/LogTests.cs b/Engine/Source/Programs/Shared/EpicGames.Horde.Tests/Logs/LogTests.cs
--- a/Engine/Source/Programs/Shared/EpicGames.Horde.Tests/Logs/LogTests.cs
+++ b/Engine/Source/Programs/Shared/EpicGames.Horde.Tests/Logs/LogTests.cs
@@ -25,7 +25,7 @@ namespace EpicGames.Horde.Tests
 		{
 			NgramSetBuilder builder = new NgramSetBuilder();
 
-			ulong[] values = { 1, 2, 3, 5, 8, 13, 21, 34, 55, 89, 144, 233, 377, 610, 987, 1597, 2584, 4181, 6765, 10946, 0xfedcba9876543210UL };
+			ulong[] values = [1, 2, 3, 5, 8, 13, 21, 34, 55, 89, 144, 233, 377, 610, 987, 1597, 2584, 4181, 6765, 10946, 0xfedcba9876543210UL];
 			foreach (ulong value in values)
 			{
 				builder.Add(value);
@@ -105,12 +105,12 @@ namespace EpicGames.Horde.Tests
 
 			// Generate the test data
 			string[] lines =
-			{
+			[
 				"abcdefghi\n",
 				"jklmno123\n",
 				"pqrst99uv\n",
 				"wx\n"
-			};
+			];
 
 			LogBuilder builder = new LogBuilder(LogFormat.Text, 1, 1, NullLogger.Instance);
 			for (int lineIdx = 0; lineIdx < lines.Length; lineIdx++)
@@ -195,11 +195,11 @@ namespace EpicGames.Horde.Tests
 */
 		static async Task SearchLogDataTestAsync(LogIndexNode index)
 		{
-			await SearchLogDataTestAsync(index, "HISPANIOLA", 0, 4, new[] { 1503, 1520, 1525, 1595 });
-			await SearchLogDataTestAsync(index, "Hispaniola", 0, 4, new[] { 1503, 1520, 1525, 1595 });
-			await SearchLogDataTestAsync(index, "HizpaniolZ", 0, 4, Array.Empty<int>());
-			await SearchLogDataTestAsync(index, "Pieces of eight!", 0, 100, new[] { 2227, 2228, 5840, 5841, 7520 });
-			await SearchLogDataTestAsync(index, "NEWSLETTER", 0, 100, new[] { 7886 });
+			await SearchLogDataTestAsync(index, "HISPANIOLA", 0, 4, [1503, 1520, 1525, 1595]);
+			await SearchLogDataTestAsync(index, "Hispaniola", 0, 4, [1503, 1520, 1525, 1595]);
+			await SearchLogDataTestAsync(index, "HizpaniolZ", 0, 4, []);
+			await SearchLogDataTestAsync(index, "Pieces of eight!", 0, 100, [2227, 2228, 5840, 5841, 7520]);
+			await SearchLogDataTestAsync(index, "NEWSLETTER", 0, 100, [7886]);
 		}
 
 		static async Task SearchLogDataTestAsync(LogIndexNode index, string text, int firstLine, int count, int[] expectedLines)
diff --git a/Engine/Source/Programs/Shared/EpicGames.Horde.Tests/Logs/ServerLogPacketBuilderTests.cs b/Engine/Source/Programs/Shared/EpicGames.Horde.Tests/Logs/ServerLogPacketBuilderTests.cs
--- a/Engine/Source/Programs/Shared/EpicGames.Horde.Tests/Logs/ServerLogPacketBuilderTests.cs
+++ b/Engine/Source/Programs/Shared/EpicGames.Horde.Tests/Logs/ServerLogPacketBuilderTests.cs
@@ -30,7 +30,7 @@ namespace EpicGames.Horde.Tests
 	{
 		class JsonLoggerImpl : ILogger
 		{
-			public List<LogEvent> _lines = new List<LogEvent>();
+			public List<LogEvent> _lines = [];
 
 			public IDisposable? BeginScope<TState>(TState state) where TState : notnull => null!;
 
@@ -191,12 +191,12 @@ namespace EpicGames.Horde.Tests
 
 		class FakeLogRpcClient : LogRpc.LogRpcClient
 		{
-			public Dictionary<LogId, BlobLocator> Logs { get; } = new Dictionary<LogId, BlobLocator>();
+			public Dictionary<LogId, BlobLocator> Logs { get; } = [];
 
 			public override AsyncUnaryCall<RpcUpdateLogResponse> UpdateLogAsync(RpcUpdateLogRequest request, CallOptions options)
 			{
 				Logs[LogId.Parse(request.LogId)] = new BlobLocator(request.TargetLocator);
-				return new AsyncUnaryCall<RpcUpdateLogResponse>(Task.FromResult(new RpcUpdateLogResponse()), Task.FromResult(new Metadata()), () => Status.DefaultSuccess, () => new Metadata(), () => { });
+				return new AsyncUnaryCall<RpcUpdateLogResponse>(Task.FromResult(new RpcUpdateLogResponse()), Task.FromResult(new Metadata()), () => Status.DefaultSuccess, () => [], () => { });
 			}
 
 			public override AsyncDuplexStreamingCall<RpcUpdateLogTailRequest, RpcUpdateLogTailResponse> UpdateLogTail(CallOptions options)
@@ -214,7 +214,7 @@ namespace EpicGames.Horde.Tests
 		class FakeHordeClient : IHordeClient, IAsyncDisposable
 		{
 			public FakeLogRpcClient LogRpc { get; } = new FakeLogRpcClient();
-			public Dictionary<string, BundleStorageNamespace> StorageNamespaces { get; } = new Dictionary<string, BundleStorageNamespace>();
+			public Dictionary<string, BundleStorageNamespace> StorageNamespaces { get; } = [];
 
 			public Uri ServerUrl => throw new NotImplementedException();
 
@@ -297,7 +297,7 @@ namespace EpicGames.Horde.Tests
 			LogNode file = await storageNamespace.CreateBlobRef(hordeClient.LogRpc.Logs[logId]).ReadBlobAsync<LogNode>();
 
 			// Check the index text
-			List<Utf8String> extractedIndexText = new List<Utf8String>();
+			List<Utf8String> extractedIndexText = [];
 
 			LogIndexNode index = await file.IndexRef.ReadBlobAsync();
 			foreach (LogChunkRef block in index.PlainTextChunkRefs)
@@ -313,7 +313,7 @@ namespace EpicGames.Horde.Tests
 			}
 
 			// Check the body text
-			List<string> extractedBodyText = new List<string>();
+			List<string> extractedBodyText = [];
 
 			foreach (LogChunkRef blockRef in file.TextChunkRefs)
 			{
@@ -355,7 +355,7 @@ namespace EpicGames.Horde.Tests
 
 			string[] output = Encoding.UTF8.GetString(writer.CreatePacket().Item1.Span).Split('\n');
 			string[] expected =
-			{
+			[
 				@"{""time"":""2023-01-01T00:00:00"",""level"":""Information"",""message"":""start A0"",""format"":""start {LongLineA$0}"",""properties"":{""LongLineA$0"":""A0""},""line"":0,""lineCount"":14}",
 				@"{""time"":""2023-01-01T00:00:00"",""level"":""Information"",""message"":""A1"",""format"":""{LongLineA$1}"",""properties"":{""LongLineA$1"":""A1""},""line"":1,""lineCount"":14}",
 				@"{""time"":""2023-01-01T00:00:00"",""level"":""Information"",""message"":""A2"",""format"":""{LongLineA$2}"",""properties"":{""LongLineA$2"":""A2""},""line"":2,""lineCount"":14}",
@@ -371,7 +371,7 @@ namespace EpicGames.Horde.Tests
 				@"{""time"":""2023-01-01T00:00:00"",""level"":""Information"",""message"":""C3"",""format"":""{LongLineC$3}"",""properties"":{""LongLineC$3"":""C3""},""line"":12,""lineCount"":14}",
 				@"{""time"":""2023-01-01T00:00:00"",""level"":""Information"",""message"":""C4 end"",""format"":""{LongLineC$4} end"",""properties"":{""LongLineC$4"":""C4""},""line"":13,""lineCount"":14}",
 				""
-			};
+			];
 
 			for (int idx = 0; idx < output.Length; idx++)
 			{
@@ -410,10 +410,10 @@ namespace EpicGames.Horde.Tests
 
 			string[] output = Encoding.UTF8.GetString(writer.CreatePacket().Item1.Span).Split('\n');
 			string[] expected =
-			{
+			[
 				@"{""time"":""2023-01-01T00:00:00"",""level"":""Information"",""message"":""abcdefghijklmnopqrstuvwxyzABCDEFGHI [...]""}",
 				""
-			};
+			];
 
 			for (int idx = 0; idx < output.Length; idx++)
 			{
diff --git a/Engine/Source/Programs/Shared/EpicGames.Horde.Tests/Storage/BackendTests.cs b/Engine/Source/Programs/Shared/EpicGames.Horde.Tests/Storage/BackendTests.cs
--- a/Engine/Source/Programs/Shared/EpicGames.Horde.Tests/Storage/BackendTests.cs
+++ b/Engine/Source/Programs/Shared/EpicGames.Horde.Tests/Storage/BackendTests.cs
@@ -61,7 +61,7 @@ namespace EpicGames.Horde.Tests.Storage
 			{
 				using MemoryMappedFileCache memoryMappedFileCache = new MemoryMappedFileCache();
 				FileStorageBackend backend = new FileStorageBackend(tempDir.Location, memoryMappedFileCache, NullLogger.Instance);
-				BlobLocator path = await backend.WriteBytesAsync(Encoding.UTF8.GetBytes("hello world"), Array.Empty<BlobLocator>());
+				BlobLocator path = await backend.WriteBytesAsync(Encoding.UTF8.GetBytes("hello world"), []);
 
 				using (IReadOnlyMemoryOwner<byte> handle = await backend.ReadBlobAsync(path, 0, null))
 				{
@@ -95,28 +95,28 @@ namespace EpicGames.Horde.Tests.Storage
 				Assert.IsTrue(value.SequenceEqual(Encoding.UTF8.GetBytes("item 2")));
 
 				byte[] data3 = Encoding.UTF8.GetBytes("3");
-				BlobLocator path3 = await cacheBackend.WriteBytesAsync(data3, Array.Empty<BlobLocator>());
+				BlobLocator path3 = await cacheBackend.WriteBytesAsync(data3, []);
 				await cacheBackend.ReadBytesAsync(path3);
 
 				byte[] data4 = Encoding.UTF8.GetBytes("4");
-				BlobLocator path4 = await cacheBackend.WriteBytesAsync(data4, Array.Empty<BlobLocator>());
+				BlobLocator path4 = await cacheBackend.WriteBytesAsync(data4, []);
 				await cacheBackend.ReadBytesAsync(path4);
 
 				byte[] data5 = Encoding.UTF8.GetBytes("5");
-				BlobLocator path5 = await cacheBackend.WriteBytesAsync(data5, Array.Empty<BlobLocator>());
+				BlobLocator path5 = await cacheBackend.WriteBytesAsync(data5, []);
 				await cacheBackend.ReadBytesAsync(path5);
 
-				HashSet<BlobLocator> paths = new HashSet<BlobLocator>(cache.GetLocators());
+				HashSet<BlobLocator> paths = [.. cache.GetLocators()];
 				Assert.AreEqual(4, paths.Count);
 				Assert.IsTrue(paths.Contains(path3));
 				Assert.IsTrue(paths.Contains(path4));
 				Assert.IsTrue(paths.Contains(path5));
 
 				byte[] data6 = Encoding.UTF8.GetBytes("12345678901");
-				BlobLocator path6 = await cacheBackend.WriteBytesAsync(data6, Array.Empty<BlobLocator>());
+				BlobLocator path6 = await cacheBackend.WriteBytesAsync(data6, []);
 				await cacheBackend.ReadBytesAsync(path6);
 
-				paths = new HashSet<BlobLocator>(cache.GetLocators());
+				paths = [.. cache.GetLocators()];
 				Assert.AreEqual(2, paths.Count);
 				Assert.IsTrue(paths.Contains(path5));
 				Assert.IsTrue(paths.Contains(path6));
@@ -124,7 +124,7 @@ namespace EpicGames.Horde.Tests.Storage
 				await cacheBackend.ReadBytesAsync(path3);
 				await cacheBackend.ReadBytesAsync(path4);
 
-				paths = new HashSet<BlobLocator>(cache.GetLocators());
+				paths = [.. cache.GetLocators()];
 				Assert.AreEqual(2, paths.Count);
 				Assert.IsTrue(paths.Contains(path3));
 				Assert.IsTrue(paths.Contains(path4));
@@ -136,10 +136,10 @@ namespace EpicGames.Horde.Tests.Storage
 			byte[] data1 = Encoding.UTF8.GetBytes("hello world");
 			byte[] data2 = Encoding.UTF8.GetBytes("item 2");
 
-			BlobLocator path1 = await backend.WriteBytesAsync(data1, Array.Empty<BlobLocator>());
+			BlobLocator path1 = await backend.WriteBytesAsync(data1, []);
 			byte[] outputData1 = await backend.ReadBytesAsync(path1);
 
-			BlobLocator path2 = await backend.WriteBytesAsync(data2, Array.Empty<BlobLocator>());
+			BlobLocator path2 = await backend.WriteBytesAsync(data2, []);
 			byte[] outputData2 = await backend.ReadBytesAsync(path2);
 
 			Assert.IsTrue(data1.SequenceEqual(outputData1));
diff --git a/Engine/Source/Programs/Shared/EpicGames.Horde.Tests/Storage/BlobDataTests.cs b/Engine/Source/Programs/Shared/EpicGames.Horde.Tests/Storage/BlobDataTests.cs
--- a/Engine/Source/Programs/Shared/EpicGames.Horde.Tests/Storage/BlobDataTests.cs
+++ b/Engine/Source/Programs/Shared/EpicGames.Horde.Tests/Storage/BlobDataTests.cs
@@ -36,9 +36,7 @@ namespace EpicGames.Horde.Tests.Storage
 		[TestMethod]
 		public void HeaderSerialization()
 		{
-			List<IBlobRef> refs = new List<IBlobRef>();
-			refs.Add(new DummyHandle("hello"));
-			refs.Add(new DummyHandle("world"));
+			List<IBlobRef> refs = [new DummyHandle("hello"), new DummyHandle("world")];
 
 			using BlobData blobData = new BlobData(new BlobType(Guid.NewGuid(), 4), new byte[] { 1, 2, 3 }, refs);
 
diff --git a/Engine/Source/Programs/Shared/EpicGames.Horde.Tests/Storage/BlobPipelineTests.cs b/Engine/Source/Programs/Shared/EpicGames.Horde.Tests/Storage/BlobPipelineTests.cs
--- a/Engine/Source/Programs/Shared/EpicGames.Horde.Tests/Storage/BlobPipelineTests.cs
+++ b/Engine/Source/Programs/Shared/EpicGames.Horde.Tests/Storage/BlobPipelineTests.cs
@@ -22,7 +22,7 @@ namespace EpicGames.Horde.Tests.Storage
 		{
 			BundleStorageNamespace client = BundleStorageNamespace.CreateInMemory(NullLogger.Instance);
 
-			List<IHashedBlobRef> handles = new List<IHashedBlobRef>();
+			List<IHashedBlobRef> handles = [];
 			await using (IBlobWriter writer = client.CreateBlobWriter())
 			{
 				for (int idx = 0; idx < 4000; idx++)
diff --git a/Engine/Source/Programs/Shared/EpicGames.Horde.Tests/Storage/BundleStoreTests.cs b/Engine/Source/Programs/Shared/EpicGames.Horde.Tests/Storage/BundleStoreTests.cs
--- a/Engine/Source/Programs/Shared/EpicGames.Horde.Tests/Storage/BundleStoreTests.cs
+++ b/Engine/Source/Programs/Shared/EpicGames.Horde.Tests/Storage/BundleStoreTests.cs
@@ -78,16 +78,13 @@ namespace EpicGames.Horde.Tests.Storage
 			payloadWriter.WriteString("Hello world");
 			byte[] payload = payloadWriter.WrittenMemory.ToArray();
 
-			List<BlobType> types = new List<BlobType>();
-			types.Add(new BlobType("{F63606D4-4061-5DBB-446F-55A69E22654F}", 1));
+			List<BlobType> types = [new BlobType("{F63606D4-4061-5DBB-446F-55A69E22654F}", 1)];
 
-			List<BundleExport> exports = new List<BundleExport>();
-			exports.Add(new BundleExport(0, 0, 0, payload.Length, Array.Empty<BundleExportRef>()));
+			List<BundleExport> exports = [new BundleExport(0, 0, 0, payload.Length, [])];
 
-			List<BundlePacket> packets = new List<BundlePacket>();
-			packets.Add(new BundlePacket(BundleCompressionFormat.None, 0, payload.Length, payload.Length));
+			List<BundlePacket> packets = [new BundlePacket(BundleCompressionFormat.None, 0, payload.Length, payload.Length)];
 
-			BundleHeader header = new BundleHeader(types.ToArray(), Array.Empty<BlobLocator>(), exports.ToArray(), packets.ToArray());
+			BundleHeader header = new BundleHeader([.. types], [], [.. exports], [.. packets]);
 
 			ReadOnlySequenceBuilder<byte> builder = new ReadOnlySequenceBuilder<byte>();
 			header.AppendTo(builder);
@@ -161,12 +158,12 @@ namespace EpicGames.Horde.Tests.Storage
 			{
 				await using IBlobWriter writer = store.CreateBlobWriter("test");
 
-				SimpleNode node1 = new SimpleNode(new ReadOnlySequence<byte>(new byte[] { 1 }), Array.Empty<IHashedBlobRef<SimpleNode>>());
-				SimpleNode node2 = new SimpleNode(new ReadOnlySequence<byte>(new byte[] { 2 }), new[] { await writer.WriteBlobAsync(node1) });
-				SimpleNode node3 = new SimpleNode(new ReadOnlySequence<byte>(new byte[] { 3 }), new[] { await writer.WriteBlobAsync(node2) });
-				SimpleNode node4 = new SimpleNode(new ReadOnlySequence<byte>(new byte[] { 4 }), Array.Empty<IHashedBlobRef<SimpleNode>>());
+				SimpleNode node1 = new SimpleNode(new ReadOnlySequence<byte>([1]), []);
+				SimpleNode node2 = new SimpleNode(new ReadOnlySequence<byte>([2]), [await writer.WriteBlobAsync(node1)]);
+				SimpleNode node3 = new SimpleNode(new ReadOnlySequence<byte>([3]), [await writer.WriteBlobAsync(node2)]);
+				SimpleNode node4 = new SimpleNode(new ReadOnlySequence<byte>([4]), []);
 
-				SimpleNode root = new SimpleNode(new ReadOnlySequence<byte>(new byte[] { 5 }), new[] { await writer.WriteBlobAsync(node4), await writer.WriteBlobAsync(node3) });
+				SimpleNode root = new SimpleNode(new ReadOnlySequence<byte>([5]), [await writer.WriteBlobAsync(node4), await writer.WriteBlobAsync(node3)]);
 				IHashedBlobRef<SimpleNode> rootRef = await writer.WriteBlobAsync(root);
 				await writer.FlushAsync();
 
@@ -226,7 +223,7 @@ namespace EpicGames.Horde.Tests.Storage
 			IHashedBlobRef<SimpleNode> inputRef;
 			await using (IBlobWriter writer = store.CreateBlobWriter(refName))
 			{
-				inputRef = await writer.WriteBlobAsync(new SimpleNode(new ReadOnlySequence<byte>(new byte[] { (byte)123 }), Array.Empty<IHashedBlobRef<SimpleNode>>()));
+				inputRef = await writer.WriteBlobAsync(new SimpleNode(new ReadOnlySequence<byte>([(byte)123]), []));
 			}
 			await store.AddRefAsync(refName, inputRef);
 
@@ -295,8 +292,7 @@ namespace EpicGames.Horde.Tests.Storage
 				using ChunkedDataWriter fileWriter = new ChunkedDataWriter(writer, new ChunkingOptions());
 				ChunkedData fileHandle = await fileWriter.CreateAsync(Encoding.UTF8.GetBytes("world"), CancellationToken.None);
 
-				List<FileUpdate> fileUpdates = new List<FileUpdate>();
-				fileUpdates.Add(new FileUpdate("hello/world", FileEntryFlags.None, fileWriter.Length, fileHandle));
+				List<FileUpdate> fileUpdates = [new FileUpdate("hello/world", FileEntryFlags.None, fileWriter.Length, fileHandle)];
 
 				DirectoryNode root = new DirectoryNode();
 				await root.UpdateAsync(fileUpdates, writer);
@@ -417,7 +413,7 @@ namespace EpicGames.Horde.Tests.Storage
 
 				FileEntry file = root.GetFileEntry("test");
 
-				Dictionary<BlobLocator, long> locatorToSize = new Dictionary<BlobLocator, long>();
+				Dictionary<BlobLocator, long> locatorToSize = [];
 				await GetUniqueBlobsAsync(file.Target, locatorToSize);
 
 				long uniqueSize = locatorToSize.Sum(x => x.Value);
diff --git a/Engine/Source/Programs/Shared/EpicGames.Horde.Tests/Storage/BundleTests.cs b/Engine/Source/Programs/Shared/EpicGames.Horde.Tests/Storage/BundleTests.cs
--- a/Engine/Source/Programs/Shared/EpicGames.Horde.Tests/Storage/BundleTests.cs
+++ b/Engine/Source/Programs/Shared/EpicGames.Horde.Tests/Storage/BundleTests.cs
@@ -120,22 +120,23 @@ namespace EpicGames.Horde.Tests.Storage
 		{
 			BundleHeader oldHeader;
 			{
-				List<BlobType> types = new List<BlobType>();
-				types.Add(new BlobType(Guid.NewGuid(), 0));
+				List<BlobType> types = [new BlobType(Guid.NewGuid(), 0)];
 
-				List<BlobLocator> imports = new List<BlobLocator>();
-				imports.Add(new BlobLocator("import1"));
-				imports.Add(new BlobLocator("import2"));
+				List<BlobLocator> imports = [new BlobLocator("import1"), new BlobLocator("import2")];
 
-				List<BundleExport> exports = new List<BundleExport>();
-				exports.Add(new BundleExport(0, 0, 0, 2, new BundleExportRef[] { new BundleExportRef(0, 5), new BundleExportRef(0, 6) }));
-				exports.Add(new BundleExport(0, 1, 0, 3, new BundleExportRef[] { new BundleExportRef(-1, 0) }));
+				List<BundleExport> exports =
+				[
+					new BundleExport(0, 0, 0, 2, [new BundleExportRef(0, 5), new BundleExportRef(0, 6)]),
+					new BundleExport(0, 1, 0, 3, [new BundleExportRef(-1, 0)]),
+				];
 
-				List<BundlePacket> packets = new List<BundlePacket>();
-				packets.Add(new BundlePacket(BundleCompressionFormat.LZ4, 0, 20, 40));
-				packets.Add(new BundlePacket(BundleCompressionFormat.LZ4, 20, 10, 20));
+				List<BundlePacket> packets =
+				[
+					new BundlePacket(BundleCompressionFormat.LZ4, 0, 20, 40),
+					new BundlePacket(BundleCompressionFormat.LZ4, 20, 10, 20),
+				];
 
-				oldHeader = new BundleHeader(types.ToArray(), imports.ToArray(), exports.ToArray(), packets.ToArray());
+				oldHeader = new BundleHeader([.. types], [.. imports], [.. exports], [.. packets]);
 			}
 
 			byte[] serializedData = oldHeader.ToByteArray();
diff --git a/Engine/Source/Programs/Shared/EpicGames.Horde.Tests/Storage/ChunkingTests.cs b/Engine/Source/Programs/Shared/EpicGames.Horde.Tests/Storage/ChunkingTests.cs
--- a/Engine/Source/Programs/Shared/EpicGames.Horde.Tests/Storage/ChunkingTests.cs
+++ b/Engine/Source/Programs/Shared/EpicGames.Horde.Tests/Storage/ChunkingTests.cs
@@ -241,10 +241,10 @@ namespace EpicGames.Horde.Tests.Storage
 			int leafRefIndex = MemoryBlobWriter.GetIndex((IHashedBlobRef)leafRef);
 			ChunkedDataNodeRef leafChunkedRef = new ChunkedDataNodeRef(3, 0, leafRef);
 
-			List<ChunkedDataNodeRef> leafNodeRefs = Enumerable.Repeat(leafChunkedRef, 10000).ToList();
+			List<ChunkedDataNodeRef> leafNodeRefs = [.. Enumerable.Repeat(leafChunkedRef, 10000)];
 			ChunkedDataNodeRef root = await InteriorChunkedDataNode.CreateTreeAsync(leafNodeRefs, interiorOptions, blobWriter, CancellationToken.None);
 
-			List<IBlobRef> list = new List<IBlobRef>();
+			List<IBlobRef> list = [];
 			await GetReadOrderAsync(root, list);
 
 			int prevIndex = Int32.MaxValue;
@@ -283,8 +283,8 @@ namespace EpicGames.Horde.Tests.Storage
 		{
 			ChunkOffsetsAndSizes chunks;
 
-			chunks._offsets = new List<int>();
-			chunks._sizes = new List<int>();
+			chunks._offsets = [];
+			chunks._sizes = [];
 
 			int currentOffset = 0;
 			while (currentOffset != data.Length)
@@ -320,12 +320,12 @@ namespace EpicGames.Horde.Tests.Storage
 			int dataSize = 1 << 20; // 1MB test buffer
 			byte[] data = GenerateChunkingTestBuffer(dataSize, 1234);
 
-			List<int> expectedChunkOffsets = new List<int>
-			{
+			List<int> expectedChunkOffsets =
+			[
 				0, 34577, 128471, 195115, 238047, 297334, 358754, 396031,
 				462359, 508658, 601550, 702021, 754650, 790285, 854987, 887998,
 				956848, 1042406
-			};
+			];
 
 			int targetSize = 65536;
 			ChunkOffsetsAndSizes chunks = FindChunks(data, targetSize);
@@ -344,8 +344,8 @@ namespace EpicGames.Horde.Tests.Storage
 			byte[] data = GenerateChunkingTestBuffer(dataSize, 1234);
 
 			// Chunk size in KB and expected number of generated chunks
-			List<(int, int)> configList = new List<(int, int)>
-			{
+			List<(int, int)> configList =
+			[
 				(8, 16442),
 				(16, 8146),
 				(32, 4089),
@@ -355,7 +355,7 @@ namespace EpicGames.Horde.Tests.Storage
 				(160, 811),
 				(192, 681),
 				(256, 503),
-			};
+			];
 
 			Parallel.ForEach(configList, config =>
 			{
diff --git a/Engine/Source/Programs/Shared/EpicGames.Horde.Tests/Storage/ContentChunkerTests.cs b/Engine/Source/Programs/Shared/EpicGames.Horde.Tests/Storage/ContentChunkerTests.cs
--- a/Engine/Source/Programs/Shared/EpicGames.Horde.Tests/Storage/ContentChunkerTests.cs
+++ b/Engine/Source/Programs/Shared/EpicGames.Horde.Tests/Storage/ContentChunkerTests.cs
@@ -80,7 +80,7 @@ namespace EpicGames.Horde.Tests.Storage
 
 			pipeline.SourceWriter.Complete();
 
-			List<SplitPoint> splitPoints = new List<SplitPoint>();
+			List<SplitPoint> splitPoints = [];
 			for (int idx = 0; idx < lengths.Length; idx++)
 			{
 				int outputLength = 0;
