Shared: Enable IDE0028, IDE0031, IDE0034, IDE0054, IDE0062, IDE0071, IDE0074, IDE0075, IDE0280, IDE0300, IDE0301 IDE0305

diff --git a/Engine/Source/Programs/Shared/EpicGames.Horde.Tests/Storage/FileMetadataDbTests.cs b/Engine/Source/Programs/Shared/EpicGames.Horde.Tests/Storage/FileMetadataDbTests.cs
--- a/Engine/Source/Programs/Shared/EpicGames.Horde.Tests/Storage/FileMetadataDbTests.cs
+++ b/Engine/Source/Programs/Shared/EpicGames.Horde.Tests/Storage/FileMetadataDbTests.cs
@@ -24,7 +24,7 @@ namespace EpicGames.Horde.Tests.Storage
 			int d = await db.AddDirectoryAsync(b, "D");
 
 			{
-				List<DirectoryRow> rows = (await db.GetDirectoriesAsync(FileMetadataDb.RootDirectoryId)).OrderBy(x => x.Id).ToList();
+				List<DirectoryRow> rows = [.. (await db.GetDirectoriesAsync(FileMetadataDb.RootDirectoryId)).OrderBy(x => x.Id)];
 				Assert.AreEqual(2, rows.Count);
 				Assert.AreEqual(a, rows[0].Id);
 				Assert.AreEqual("A", rows[0].Name);
@@ -33,12 +33,12 @@ namespace EpicGames.Horde.Tests.Storage
 			}
 
 			{
-				List<DirectoryRow> rows = (await db.GetDirectoriesAsync(a)).OrderBy(x => x.Id).ToList();
+				List<DirectoryRow> rows = [.. (await db.GetDirectoriesAsync(a)).OrderBy(x => x.Id)];
 				Assert.AreEqual(0, rows.Count);
 			}
 
 			{
-				List<DirectoryRow> rows = (await db.GetDirectoriesAsync(b)).OrderBy(x => x.Id).ToList();
+				List<DirectoryRow> rows = [.. (await db.GetDirectoriesAsync(b)).OrderBy(x => x.Id)];
 				Assert.AreEqual(2, rows.Count);
 				Assert.AreEqual(c, rows[0].Id);
 				Assert.AreEqual("C", rows[0].Name);
@@ -67,7 +67,7 @@ namespace EpicGames.Horde.Tests.Storage
 			int f = await db.AddFileAsync(a, "F.txt", time, 4567);
 
 			{
-				List<FileRow> files = (await db.FindFilesInDirectoryAsync(FileMetadataDb.RootDirectoryId)).OrderBy(x => x.Id).ToList();
+				List<FileRow> files = [.. (await db.FindFilesInDirectoryAsync(FileMetadataDb.RootDirectoryId)).OrderBy(x => x.Id)];
 				Assert.AreEqual(2, files.Count);
 
 				Assert.AreEqual(c, files[0].Id);
@@ -80,7 +80,7 @@ namespace EpicGames.Horde.Tests.Storage
 			}
 
 			{
-				List<FileRow> files = (await db.FindFilesInDirectoryAsync(a)).OrderBy(x => x.Id).ToList();
+				List<FileRow> files = [.. (await db.FindFilesInDirectoryAsync(a)).OrderBy(x => x.Id)];
 				Assert.AreEqual(2, files.Count);
 
 				Assert.AreEqual(e, files[0].Id);
@@ -104,10 +104,12 @@ namespace EpicGames.Horde.Tests.Storage
 			using FileMetadataDb db = await FileMetadataDb.CreateInMemoryAsync();
 			int root = await db.AddDirectoryAsync(FileMetadataDb.RootDirectoryId, "");
 
-			List<FileRow> files = new List<FileRow>();
-			files.Add(new FileRow(root, "A", DateTime.MinValue, 1));
-			files.Add(new FileRow(root, "B", DateTime.MinValue, 1));
-			files.Add(new FileRow(root, "C", DateTime.MinValue, 1));
+			List<FileRow> files =
+			[
+				new FileRow(root, "A", DateTime.MinValue, 1),
+				new FileRow(root, "B", DateTime.MinValue, 1),
+				new FileRow(root, "C", DateTime.MinValue, 1),
+			];
 
 			await db.AddFilesAsync(files);
 			Assert.AreEqual(3, (await db.FindFilesInDirectoryAsync(root)).Count());
@@ -117,7 +119,7 @@ namespace EpicGames.Horde.Tests.Storage
 		public async Task BatchDirectoryTestsAsync()
 		{
 			using FileMetadataDb db = await FileMetadataDb.CreateInMemoryAsync();
-			await db.AddDirectoriesAsync(new[] { new DirectoryRow(FileMetadataDb.RootDirectoryId, "A"), new DirectoryRow(FileMetadataDb.RootDirectoryId, "B"), new DirectoryRow(FileMetadataDb.RootDirectoryId, "C") });
+			await db.AddDirectoriesAsync([new DirectoryRow(FileMetadataDb.RootDirectoryId, "A"), new DirectoryRow(FileMetadataDb.RootDirectoryId, "B"), new DirectoryRow(FileMetadataDb.RootDirectoryId, "C")]);
 			Assert.AreEqual(3, (await db.GetDirectoriesAsync(FileMetadataDb.RootDirectoryId)).Count());
 		}
 
@@ -144,7 +146,7 @@ namespace EpicGames.Horde.Tests.Storage
 			int c4 = await db.AddChunkAsync(f4, 0, 10, h2);
 
 			{
-				List<ChunkRow> chunks = (await db.FindChunksAsync(h1, 10)).OrderBy(x => x.Id).ToList();
+				List<ChunkRow> chunks = [.. (await db.FindChunksAsync(h1, 10)).OrderBy(x => x.Id)];
 				Assert.AreEqual(3, chunks.Count);
 				Assert.AreEqual(new ChunkRow(c1, f1, 0, 10, h1), chunks[0]);
 				Assert.AreEqual(new ChunkRow(c2, f2, 10, 10, h1), chunks[1]);
@@ -152,7 +154,7 @@ namespace EpicGames.Horde.Tests.Storage
 			}
 
 			{
-				List<ChunkRow> chunks = (await db.FindChunksAsync(h2, 10)).OrderBy(x => x.Id).ToList();
+				List<ChunkRow> chunks = [.. (await db.FindChunksAsync(h2, 10)).OrderBy(x => x.Id)];
 				Assert.AreEqual(1, chunks.Count);
 				Assert.AreEqual(new ChunkRow(c4, f4, 0, 10, h2), chunks[0]);
 			}
@@ -185,7 +187,7 @@ namespace EpicGames.Horde.Tests.Storage
 			int c5 = await db.AddChunkAsync(f5, 10, 10, h1);
 
 			{
-				List<ChunkRow> chunks = (await db.FindChunksAsync(h1, 10)).OrderBy(x => x.Id).ToList();
+				List<ChunkRow> chunks = [.. (await db.FindChunksAsync(h1, 10)).OrderBy(x => x.Id)];
 				Assert.AreEqual(5, chunks.Count);
 				Assert.AreEqual(new ChunkRow(c1, f1, 0, 10, h1), chunks[0]);
 				Assert.AreEqual(new ChunkRow(c2, f2, 10, 10, h1), chunks[1]);
@@ -197,7 +199,7 @@ namespace EpicGames.Horde.Tests.Storage
 			await db.RemoveDirectoryAsync(d4);
 
 			{
-				List<ChunkRow> chunks = (await db.FindChunksAsync(h1, 10)).OrderBy(x => x.Id).ToList();
+				List<ChunkRow> chunks = [.. (await db.FindChunksAsync(h1, 10)).OrderBy(x => x.Id)];
 				Assert.AreEqual(3, chunks.Count);
 
 				Assert.AreEqual(new ChunkRow(c1, f1, 0, 10, h1), chunks[0]);
diff --git a/Engine/Source/Programs/Shared/EpicGames.Horde.Tests/Storage/StorageClientTests.cs b/Engine/Source/Programs/Shared/EpicGames.Horde.Tests/Storage/StorageClientTests.cs
--- a/Engine/Source/Programs/Shared/EpicGames.Horde.Tests/Storage/StorageClientTests.cs
+++ b/Engine/Source/Programs/Shared/EpicGames.Horde.Tests/Storage/StorageClientTests.cs
@@ -1,6 +1,5 @@
 // Copyright Epic Games, Inc. All Rights Reserved.
 
-using System;
 using System.Threading.Tasks;
 using EpicGames.Core;
 using EpicGames.Horde.Storage;
@@ -21,7 +20,7 @@ namespace EpicGames.Horde.Tests.Storage
 			public static BlobType BlobType = new BlobType("{99601905-4F6E-A089-03D6-F187711BAFEE}", 1);
 
 			public int Value { get; }
-			public byte[] Padding { get; set; } = Array.Empty<byte>();
+			public byte[] Padding { get; set; } = [];
 			public IHashedBlobRef<TestNode>[] Refs { get; }
 
 			public TestNode(int value, params IHashedBlobRef<TestNode>[] refs)
diff --git a/Engine/Source/Programs/Shared/EpicGames.Horde.Tests/StorageCompressorTest.cs b/Engine/Source/Programs/Shared/EpicGames.Horde.Tests/StorageCompressorTest.cs
--- a/Engine/Source/Programs/Shared/EpicGames.Horde.Tests/StorageCompressorTest.cs
+++ b/Engine/Source/Programs/Shared/EpicGames.Horde.Tests/StorageCompressorTest.cs
@@ -44,7 +44,7 @@ namespace EpicGames.Horde.Tests
 		public async Task CompressAndDecompressAsync()
 		{
 			FakeStorageCompressor compressor = new();
-			byte[] uncompressedData = { 0x41, 0x42, 0x43 };
+			byte[] uncompressedData = [0x41, 0x42, 0x43];
 			using MemoryStream input = new(uncompressedData);
 			using MemoryStream output = new();
 			await compressor.CompressAsync(input, output);
diff --git a/Engine/Source/Programs/Shared/EpicGames.Horde/Agents/Utilization/UtilizationMessages.cs b/Engine/Source/Programs/Shared/EpicGames.Horde/Agents/Utilization/UtilizationMessages.cs
--- a/Engine/Source/Programs/Shared/EpicGames.Horde/Agents/Utilization/UtilizationMessages.cs
+++ b/Engine/Source/Programs/Shared/EpicGames.Horde/Agents/Utilization/UtilizationMessages.cs
@@ -27,7 +27,7 @@ namespace EpicGames.Horde.Agents.Telemetry
 		/// <summary>
 		/// List of pools
 		/// </summary>
-		public List<GetUtilizationPoolDataResponse> Pools { get; set; } = new List<GetUtilizationPoolDataResponse>();
+		public List<GetUtilizationPoolDataResponse> Pools { get; set; } = [];
 
 		/// <summary>
 		/// Total admin time
@@ -78,7 +78,7 @@ namespace EpicGames.Horde.Agents.Telemetry
 		/// <summary>
 		/// List of streams
 		/// </summary>
-		public List<GetUtilizationStreamDataResponse> Streams { get; set; } = new List<GetUtilizationStreamDataResponse>();
+		public List<GetUtilizationStreamDataResponse> Streams { get; set; } = [];
 	}
 
 	/// <summary>
diff --git a/Engine/Source/Programs/Shared/EpicGames.Horde/Artifacts/ArtifactDescriptor.cs b/Engine/Source/Programs/Shared/EpicGames.Horde/Artifacts/ArtifactDescriptor.cs
--- a/Engine/Source/Programs/Shared/EpicGames.Horde/Artifacts/ArtifactDescriptor.cs
+++ b/Engine/Source/Programs/Shared/EpicGames.Horde/Artifacts/ArtifactDescriptor.cs
@@ -110,7 +110,7 @@ namespace EpicGames.Horde.Storage
 			RefName = refName;
 			if (filter != null && filter.Count > 0)
 			{
-				Filter = new List<string>(filter);
+				Filter = [.. filter];
 			}
 		}
 
diff --git a/Engine/Source/Programs/Shared/EpicGames.Horde/Artifacts/ArtifactMessages.cs b/Engine/Source/Programs/Shared/EpicGames.Horde/Artifacts/ArtifactMessages.cs
--- a/Engine/Source/Programs/Shared/EpicGames.Horde/Artifacts/ArtifactMessages.cs
+++ b/Engine/Source/Programs/Shared/EpicGames.Horde/Artifacts/ArtifactMessages.cs
@@ -132,8 +132,8 @@ namespace EpicGames.Horde.Artifacts
 		[JsonConstructor]
 		public GetArtifactResponse()
 		{
-			Keys = Array.Empty<string>();
-			Metadata = Array.Empty<string>();
+			Keys = [];
+			Metadata = [];
 		}
 
 		/// <summary>
@@ -164,7 +164,7 @@ namespace EpicGames.Horde.Artifacts
 		/// <summary>
 		/// List of artifacts matching the search criteria
 		/// </summary>
-		public List<GetArtifactResponse> Artifacts { get; set; } = new List<GetArtifactResponse>();
+		public List<GetArtifactResponse> Artifacts { get; set; } = [];
 	}
 
 	/// <summary>
@@ -253,6 +253,6 @@ namespace EpicGames.Horde.Artifacts
 		/// <summary>
 		/// Filter lines for the zip. Uses standard <see cref="FileFilter"/> syntax.
 		/// </summary>
-		public List<string> Filter { get; set; } = new List<string>();
+		public List<string> Filter { get; set; } = [];
 	}
 }
diff --git a/Engine/Source/Programs/Shared/EpicGames.Horde/Artifacts/HttpArtifactCollection.cs b/Engine/Source/Programs/Shared/EpicGames.Horde/Artifacts/HttpArtifactCollection.cs
--- a/Engine/Source/Programs/Shared/EpicGames.Horde/Artifacts/HttpArtifactCollection.cs
+++ b/Engine/Source/Programs/Shared/EpicGames.Horde/Artifacts/HttpArtifactCollection.cs
@@ -3,7 +3,6 @@
 using System;
 using System.Collections.Generic;
 using System.Diagnostics;
-using System.Linq;
 using System.Runtime.CompilerServices;
 using System.Threading;
 using System.Threading.Tasks;
@@ -68,7 +67,7 @@ namespace EpicGames.Horde.Artifacts
 			readonly Artifact _artifact;
 			readonly CreateArtifactResponse _response;
 			readonly IStorageNamespace _namespace;
-			readonly List<AliasInfo> _aliases = new List<AliasInfo>();
+			readonly List<AliasInfo> _aliases = [];
 
 			public ArtifactId Id => _artifact.Id;
 			public ArtifactName Name => _artifact.Name;
@@ -130,7 +129,7 @@ namespace EpicGames.Horde.Artifacts
 		{
 			HordeHttpClient hordeHttpClient = _hordeClient.CreateHttpClient();
 			CreateArtifactResponse response = await hordeHttpClient.CreateArtifactAsync(name, type, description, streamId, commitId, keys, metadata, cancellationToken);
-			Artifact artifact = new Artifact(this, response.ArtifactId, name, type, description, streamId, CommitIdWithOrder.FromPerforceChange(commitId.GetPerforceChange()), keys.ToList(), metadata.ToList(), response.NamespaceId, response.RefName, DateTime.UtcNow);
+			Artifact artifact = new Artifact(this, response.ArtifactId, name, type, description, streamId, CommitIdWithOrder.FromPerforceChange(commitId.GetPerforceChange()), [.. keys], [.. metadata], response.NamespaceId, response.RefName, DateTime.UtcNow);
 			return new ArtifactBuilder(this, artifact, response);
 		}
 
diff --git a/Engine/Source/Programs/Shared/EpicGames.Horde/Artifacts/UnsyncMessages.cs b/Engine/Source/Programs/Shared/EpicGames.Horde/Artifacts/UnsyncMessages.cs
--- a/Engine/Source/Programs/Shared/EpicGames.Horde/Artifacts/UnsyncMessages.cs
+++ b/Engine/Source/Programs/Shared/EpicGames.Horde/Artifacts/UnsyncMessages.cs
@@ -23,13 +23,13 @@ namespace EpicGames.Horde.Artifacts
 		/// Files to retrieve
 		/// </summary>
 		[JsonPropertyName("blocks")]
-		public List<string> Blocks { get; set; } = new List<string>();
+		public List<string> Blocks { get; set; } = [];
 
 		/// <summary>
 		/// Files to retrieve
 		/// </summary>
 		[JsonPropertyName("files")]
-		public List<GetUnsyncFileRequest> Files { get; set; } = new List<GetUnsyncFileRequest>();
+		public List<GetUnsyncFileRequest> Files { get; set; } = [];
 	}
 
 	/// <summary>
@@ -45,7 +45,7 @@ namespace EpicGames.Horde.Artifacts
 		/// <summary>
 		/// Blobs to return
 		/// </summary>
-		public List<GetUnsyncBlockRequest> Blocks { get; set; } = new List<GetUnsyncBlockRequest>();
+		public List<GetUnsyncBlockRequest> Blocks { get; set; } = [];
 	}
 
 	/// <summary>
diff --git a/Engine/Source/Programs/Shared/EpicGames.Horde/Auth/OAuthHandler.cs b/Engine/Source/Programs/Shared/EpicGames.Horde/Auth/OAuthHandler.cs
--- a/Engine/Source/Programs/Shared/EpicGames.Horde/Auth/OAuthHandler.cs
+++ b/Engine/Source/Programs/Shared/EpicGames.Horde/Auth/OAuthHandler.cs
@@ -105,13 +105,13 @@ namespace EpicGames.Horde.Auth
 		/// <returns></returns>
 		async Task UpdateAccessTokenAsync(CancellationToken cancellationToken)
 		{
-			KeyValuePair<string, string>[] content = new KeyValuePair<string, string>[]
-			{
+			KeyValuePair<string, string>[] content =
+			[
 				new KeyValuePair<string, string>("grant_type", _options.GrantType),
 				new KeyValuePair<string, string>("client_id", _options.ClientId),
 				new KeyValuePair<string, string>("client_secret", _options.ClientSecret),
 				new KeyValuePair<string, string>("scope", _options.Scope)
-			};
+			];
 
 			try
 			{
diff --git a/Engine/Source/Programs/Shared/EpicGames.Horde/Commits/CommitTag.cs b/Engine/Source/Programs/Shared/EpicGames.Horde/Commits/CommitTag.cs
--- a/Engine/Source/Programs/Shared/EpicGames.Horde/Commits/CommitTag.cs
+++ b/Engine/Source/Programs/Shared/EpicGames.Horde/Commits/CommitTag.cs
@@ -17,7 +17,7 @@ namespace EpicGames.Horde.Commits
 
 		public static bool PresenceOf(this Group group)
 		{
-			return group.Success ? true : false;
+			return group.Success;
 		}
 	}
 
diff --git a/Engine/Source/Programs/Shared/EpicGames.Horde/Commits/ICommitCollection.cs b/Engine/Source/Programs/Shared/EpicGames.Horde/Commits/ICommitCollection.cs
--- a/Engine/Source/Programs/Shared/EpicGames.Horde/Commits/ICommitCollection.cs
+++ b/Engine/Source/Programs/Shared/EpicGames.Horde/Commits/ICommitCollection.cs
@@ -127,7 +127,7 @@ namespace EpicGames.Horde.Commits
 		/// <returns>The last code change</returns>
 		public static async ValueTask<ICommit?> GetLastCodeChangeAsync(this ICommitCollection commitCollection, CommitId? maxCommitId, CancellationToken cancellationToken = default)
 		{
-			return await commitCollection.FindAsync(minCommitId: null, maxCommitId: maxCommitId, maxResults: 1, tags: new[] { CommitTag.Code }, cancellationToken: cancellationToken).FirstOrDefaultAsync(cancellationToken);
+			return await commitCollection.FindAsync(minCommitId: null, maxCommitId: maxCommitId, maxResults: 1, tags: [CommitTag.Code], cancellationToken: cancellationToken).FirstOrDefaultAsync(cancellationToken);
 		}
 
 		/// <summary>
diff --git a/Engine/Source/Programs/Shared/EpicGames.Horde/Common/Condition.cs b/Engine/Source/Programs/Shared/EpicGames.Horde/Common/Condition.cs
--- a/Engine/Source/Programs/Shared/EpicGames.Horde/Common/Condition.cs
+++ b/Engine/Source/Programs/Shared/EpicGames.Horde/Common/Condition.cs
@@ -333,11 +333,11 @@ namespace EpicGames.Horde.Common
 		/// </summary>
 		public string? Error { get; private set; }
 
-		readonly List<Token> _tokens = new List<Token>();
-		readonly List<string> _strings = new List<string>();
+		readonly List<Token> _tokens = [];
+		readonly List<string> _strings = [];
 
-		static readonly IEnumerable<string> s_trueScalar = new[] { "true" };
-		static readonly IEnumerable<string> s_falseScalar = new[] { "false" };
+		static readonly IEnumerable<string> s_trueScalar = ["true"];
+		static readonly IEnumerable<string> s_falseScalar = ["false"];
 
 		private Condition(string text)
 		{
@@ -599,7 +599,7 @@ namespace EpicGames.Horde.Common
 				TokenType.True => s_trueScalar,
 				TokenType.False => s_falseScalar,
 				TokenType.Identifier => getPropertyValues(_strings[token.Index]),
-				TokenType.Scalar => new string[] { _strings[token.Index] },
+				TokenType.Scalar => [_strings[token.Index]],
 				_ => throw new InvalidOperationException("Invalid token type")
 			};
 		}
@@ -619,7 +619,7 @@ namespace EpicGames.Horde.Common
 		/// Implicit conversion from string to conditions
 		/// </summary>
 		/// <param name="text"></param>
-		[return: NotNullIfNotNull("Text")]
+		[return: NotNullIfNotNull(nameof(text))]
 		public static implicit operator Condition?(string? text)
 		{
 			if (text == null)
diff --git a/Engine/Source/Programs/Shared/EpicGames.Horde/Compute/AgentMessageHandler.cs b/Engine/Source/Programs/Shared/EpicGames.Horde/Compute/AgentMessageHandler.cs
--- a/Engine/Source/Programs/Shared/EpicGames.Horde/Compute/AgentMessageHandler.cs
+++ b/Engine/Source/Programs/Shared/EpicGames.Horde/Compute/AgentMessageHandler.cs
@@ -43,7 +43,7 @@ namespace EpicGames.Horde.Compute
 		public AgentMessageHandler(DirectoryReference sandboxDir, Dictionary<string, string?>? envVars, bool executeInProcess, string? wineExecutablePath, string? containerEngineExecutable, ILogger logger)
 		{
 			_sandboxDir = sandboxDir;
-			_envVars = envVars ?? new Dictionary<string, string?>();
+			_envVars = envVars ?? [];
 			_executeInProcess = executeInProcess;
 			_wineExecutablePath = wineExecutablePath;
 			_containerEngineExecutable = containerEngineExecutable;
@@ -83,7 +83,7 @@ namespace EpicGames.Horde.Compute
 
 		async Task RunAsync(ComputeSocket socket, int channelId, int bufferSize, Func<Exception, Task> postException, CancellationToken cancellationToken)
 		{
-			List<Task> childTasks = new List<Task>();
+			List<Task> childTasks = [];
 			using AgentMessageChannel channel = socket.CreateAgentMessageChannel(channelId, bufferSize);
 			try
 			{
@@ -321,10 +321,10 @@ namespace EpicGames.Horde.Compute
 			using CancellationTokenSource cancellationTokenSource = CancellationTokenSource.CreateLinkedTokenSource(cancellationTokens);
 			CancellationToken cancellationToken = cancellationTokenSource.Token;
 
-			List<SharedMemoryBuffer> buffers = new();
+			List<SharedMemoryBuffer> buffers = [];
 			try
 			{
-				List<(int, ComputeBufferWriter)> writers = new List<(int, ComputeBufferWriter)>();
+				List<(int, ComputeBufferWriter)> writers = [];
 				while (await ipcReader.WaitToReadAsync(1, cancellationToken))
 				{
 					ReadOnlyMemory<byte> memory = ipcReader.GetReadBuffer();
@@ -432,7 +432,7 @@ namespace EpicGames.Horde.Compute
 
 		async Task ExecuteProcessAssemblyAsync(AgentMessageChannel channel, IReadOnlyList<string> arguments, string? workingDir, IReadOnlyDictionary<string, string?>? envVars, CancellationToken cancellationToken)
 		{
-			List<(string, string?)> prevEnvVars = new List<(string, string?)>();
+			List<(string, string?)> prevEnvVars = [];
 			if (envVars != null)
 			{
 				foreach ((string key, string? value) in envVars)
@@ -448,7 +448,7 @@ namespace EpicGames.Horde.Compute
 			try
 			{
 				string assemblyPath = FileReference.Combine(_sandboxDir, arguments[0]).FullName;
-				string[] mainArgs = arguments.Skip(1).ToArray();
+				string[] mainArgs = [.. arguments.Skip(1)];
 
 				_logger.LogWarning("Note: Loading and running {Assembly} in process", assemblyPath);
 
@@ -497,15 +497,15 @@ namespace EpicGames.Horde.Compute
 			}
 			await File.WriteAllTextAsync(envFilePath, sb.ToString(), cancellationToken);
 
-			List<string> resolvedArguments = new()
-			{
+			List<string> resolvedArguments =
+			[
 				"run",
 				"--tty", // Allocate a pseudo-TTY
 				"--rm", // Ensure container is removed after run
 				$"--user={linuxUid}:{linuxGid}", // Run container as current user (important for mounted dirs)
 				$"--volume={_sandboxDir}:{_sandboxDir}:rw",
 				"--env-file=" + envFilePath,
-			};
+			];
 
 			if (flags.HasFlag(ExecuteProcessFlags.ReplaceContainerEntrypoint))
 			{
@@ -541,7 +541,7 @@ namespace EpicGames.Horde.Compute
 				if (flags.HasFlag(ExecuteProcessFlags.UseWine) && _wineExecutablePath != null)
 				{
 					// Path to the original Windows executable is prepended to the argument list so Wine can run it
-					resolvedCommandLine = CommandLineArguments.Join(new[] { resolvedExecutable }.Concat(arguments).ToList());
+					resolvedCommandLine = CommandLineArguments.Join([resolvedExecutable, .. arguments]);
 					resolvedExecutable = _wineExecutablePath;
 				}
 
diff --git a/Engine/Source/Programs/Shared/EpicGames.Horde/Compute/Buffers/PooledBuffer.cs b/Engine/Source/Programs/Shared/EpicGames.Horde/Compute/Buffers/PooledBuffer.cs
--- a/Engine/Source/Programs/Shared/EpicGames.Horde/Compute/Buffers/PooledBuffer.cs
+++ b/Engine/Source/Programs/Shared/EpicGames.Horde/Compute/Buffers/PooledBuffer.cs
@@ -94,7 +94,7 @@ namespace EpicGames.Horde.Compute.Buffers
 				{
 					_chunkOwners[idx].Dispose();
 				}
-				_chunkOwners = Array.Empty<IMemoryOwner<byte>>();
+				_chunkOwners = [];
 
 				_headerHandle.Free();
 				_headerHandle = default;
diff --git a/Engine/Source/Programs/Shared/EpicGames.Horde/Compute/Clients/AgentComputeClient.cs b/Engine/Source/Programs/Shared/EpicGames.Horde/Compute/Clients/AgentComputeClient.cs
--- a/Engine/Source/Programs/Shared/EpicGames.Horde/Compute/Clients/AgentComputeClient.cs
+++ b/Engine/Source/Programs/Shared/EpicGames.Horde/Compute/Clients/AgentComputeClient.cs
@@ -29,7 +29,7 @@ namespace EpicGames.Horde.Compute.Clients
 			public ClusterId Cluster { get; } = s_cluster;
 
 			/// <inheritdoc/>
-			public IReadOnlyList<string> Properties { get; } = new List<string>();
+			public IReadOnlyList<string> Properties { get; } = [];
 
 			/// <inheritdoc/>
 			public IReadOnlyDictionary<string, int> AssignedResources => new Dictionary<string, int>();
@@ -135,10 +135,7 @@ namespace EpicGames.Horde.Compute.Clients
 			using ManagedProcessGroup group = new ManagedProcessGroup();
 			if (RuntimeInformation.IsOSPlatform(OSPlatform.Windows))
 			{
-				List<string> arguments = new List<string>();
-				arguments.Add(hordeAgentAssembly);
-				arguments.Add("computeworker");
-				arguments.Add($"-port={port}");
+				List<string> arguments = [hordeAgentAssembly, "computeworker", $"-port={port}"];
 
 				using ManagedProcess process = new ManagedProcess(group, "dotnet", CommandLineArguments.Join(arguments), null, null, ProcessPriorityClass.Normal);
 
diff --git a/Engine/Source/Programs/Shared/EpicGames.Horde/Compute/Clients/LocalComputeClient.cs b/Engine/Source/Programs/Shared/EpicGames.Horde/Compute/Clients/LocalComputeClient.cs
--- a/Engine/Source/Programs/Shared/EpicGames.Horde/Compute/Clients/LocalComputeClient.cs
+++ b/Engine/Source/Programs/Shared/EpicGames.Horde/Compute/Clients/LocalComputeClient.cs
@@ -22,7 +22,7 @@ namespace EpicGames.Horde.Compute.Clients
 		class LeaseImpl : IComputeLease
 		{
 			public ClusterId Cluster { get; } = s_cluster;
-			public IReadOnlyList<string> Properties { get; } = new List<string>();
+			public IReadOnlyList<string> Properties { get; } = [];
 			public IReadOnlyDictionary<string, int> AssignedResources => new Dictionary<string, int>();
 			public RemoteComputeSocket Socket => _socket;
 			public string Ip => "127.0.0.1";
diff --git a/Engine/Source/Programs/Shared/EpicGames.Horde/Compute/ComputeBuffer.cs b/Engine/Source/Programs/Shared/EpicGames.Horde/Compute/ComputeBuffer.cs
--- a/Engine/Source/Programs/Shared/EpicGames.Horde/Compute/ComputeBuffer.cs
+++ b/Engine/Source/Programs/Shared/EpicGames.Horde/Compute/ComputeBuffer.cs
@@ -512,7 +512,7 @@ namespace EpicGames.Horde.Compute
 		protected virtual void Dispose(bool disposing)
 		{
 			_headerPtr = default;
-			_chunks = Array.Empty<Memory<byte>>();
+			_chunks = [];
 		}
 
 		/// <summary>
diff --git a/Engine/Source/Programs/Shared/EpicGames.Horde/Compute/ComputeMessages.cs b/Engine/Source/Programs/Shared/EpicGames.Horde/Compute/ComputeMessages.cs
--- a/Engine/Source/Programs/Shared/EpicGames.Horde/Compute/ComputeMessages.cs
+++ b/Engine/Source/Programs/Shared/EpicGames.Horde/Compute/ComputeMessages.cs
@@ -123,7 +123,7 @@ namespace EpicGames.Horde.Compute
 		/// Value = actual port number
 		/// Relay connection mode uses this information to set up port forwarding.
 		/// </summary>
-		public Dictionary<string, int> Ports { get; set; } = new();
+		public Dictionary<string, int> Ports { get; set; } = [];
 
 		/// <summary>
 		/// Type of connection mode that is preferred by the client. Server can still override.
@@ -241,7 +241,7 @@ namespace EpicGames.Horde.Compute
 		/// <summary>
 		/// Resources assigned to this machine
 		/// </summary>
-		public Dictionary<string, int> AssignedResources { get; set; } = new Dictionary<string, int>();
+		public Dictionary<string, int> AssignedResources { get; set; } = [];
 
 		/// <summary>
 		/// Version number for the compute protocol
@@ -251,7 +251,7 @@ namespace EpicGames.Horde.Compute
 		/// <summary>
 		/// Properties of the agent assigned to do the work
 		/// </summary>
-		public IReadOnlyList<string> Properties { get; set; } = new List<string>();
+		public IReadOnlyList<string> Properties { get; set; } = [];
 	}
 
 	/// <summary>
@@ -324,7 +324,7 @@ namespace EpicGames.Horde.Compute
 		/// <summary>
 		/// Key/value of resources needed by session (such as CPU or memory, see KnownPropertyNames in Horde.Server)
 		/// </summary>
-		public Dictionary<string, int> ResourceNeeds { get; set; } = new();
+		public Dictionary<string, int> ResourceNeeds { get; set; } = [];
 	}
 
 	/// <summary>
@@ -335,6 +335,6 @@ namespace EpicGames.Horde.Compute
 		/// <summary>
 		/// List of resource needs
 		/// </summary>
-		public List<ResourceNeedsMessage> ResourceNeeds { get; set; } = new();
+		public List<ResourceNeedsMessage> ResourceNeeds { get; set; } = [];
 	}
 }
diff --git a/Engine/Source/Programs/Shared/EpicGames.Horde/Compute/ComputeSocket.cs b/Engine/Source/Programs/Shared/EpicGames.Horde/Compute/ComputeSocket.cs
--- a/Engine/Source/Programs/Shared/EpicGames.Horde/Compute/ComputeSocket.cs
+++ b/Engine/Source/Programs/Shared/EpicGames.Horde/Compute/ComputeSocket.cs
@@ -5,7 +5,6 @@ using System.Buffers;
 using System.Buffers.Binary;
 using System.Collections.Generic;
 using System.Diagnostics.CodeAnalysis;
-using System.Linq;
 using System.Net.Sockets;
 using System.Threading;
 using System.Threading.Tasks;
@@ -91,7 +90,7 @@ namespace EpicGames.Horde.Compute
 		public const string IpcEnvVar = "UE_HORDE_COMPUTE_IPC";
 
 		readonly ComputeBufferWriter _commandBufferWriter;
-		readonly List<ComputeBuffer> _buffers = new List<ComputeBuffer>();
+		readonly List<ComputeBuffer> _buffers = [];
 		readonly ILogger _logger;
 
 		/// <inheritdoc/>
@@ -244,10 +243,10 @@ namespace EpicGames.Horde.Compute
 
 		static async Task ProcessIpcMessagesAsync(ComputeSocket socket, ComputeBufferReader ipcReader, ILogger logger, CancellationToken cancellationToken)
 		{
-			List<SharedMemoryBuffer> buffers = new();
+			List<SharedMemoryBuffer> buffers = [];
 			try
 			{
-				List<(int, ComputeBufferWriter)> writers = new List<(int, ComputeBufferWriter)>();
+				List<(int, ComputeBufferWriter)> writers = [];
 				while (await ipcReader.WaitToReadAsync(1, cancellationToken))
 				{
 					ReadOnlyMemory<byte> memory = ipcReader.GetReadBuffer();
@@ -385,10 +384,10 @@ namespace EpicGames.Horde.Compute
 		readonly ILogger _logger;
 
 		BackgroundTask? _recvTask;
-		readonly Dictionary<int, RecvBuffer> _recvBuffers = new Dictionary<int, RecvBuffer>();
+		readonly Dictionary<int, RecvBuffer> _recvBuffers = [];
 
 		readonly SemaphoreSlim _sendSemaphore = new SemaphoreSlim(1, 1);
-		readonly Dictionary<int, SendBuffer> _sendBuffers = new Dictionary<int, SendBuffer>();
+		readonly Dictionary<int, SendBuffer> _sendBuffers = [];
 
 		/// <inheritdoc/>
 		public override ComputeProtocol Protocol => _protocol;
@@ -471,7 +470,7 @@ namespace EpicGames.Horde.Compute
 			{
 				_logger.LogDebug("Started socket reader");
 
-				List<Task> detachTasks = new List<Task>();
+				List<Task> detachTasks = [];
 
 				byte[] header = new byte[8];
 				try
@@ -741,11 +740,11 @@ namespace EpicGames.Horde.Compute
 			lock (_lockObject)
 			{
 				_complete = true;
-				recvChannelIds = _recvBuffers.Keys.ToArray();
-				sendChannelIds = _sendBuffers.Keys.ToArray();
+				recvChannelIds = [.. _recvBuffers.Keys];
+				sendChannelIds = [.. _sendBuffers.Keys];
 			}
 
-			List<Task> tasks = new List<Task>();
+			List<Task> tasks = [];
 			if (recvBuffers)
 			{
 				foreach (int recvChannelId in recvChannelIds)
diff --git a/Engine/Source/Programs/Shared/EpicGames.Horde/Compute/ExternalIpResolver.cs b/Engine/Source/Programs/Shared/EpicGames.Horde/Compute/ExternalIpResolver.cs
--- a/Engine/Source/Programs/Shared/EpicGames.Horde/Compute/ExternalIpResolver.cs
+++ b/Engine/Source/Programs/Shared/EpicGames.Horde/Compute/ExternalIpResolver.cs
@@ -36,10 +36,10 @@ public class ExternalIpResolver
 	private IPAddress? _resolvedIp;
 
 	private readonly HttpClient _httpClient;
-	private readonly List<string> _ipLookupUrls = new()
-	{
+	private readonly List<string> _ipLookupUrls =
+	[
 		"http://whatismyip.akamai.com", "http://checkip.amazonaws.com", "http://ifconfig.me/ip"
-	};
+	];
 
 	/// <summary>
 	/// Constructor
diff --git a/Engine/Source/Programs/Shared/EpicGames.Horde/Compute/RequestIdAllocator.cs b/Engine/Source/Programs/Shared/EpicGames.Horde/Compute/RequestIdAllocator.cs
--- a/Engine/Source/Programs/Shared/EpicGames.Horde/Compute/RequestIdAllocator.cs
+++ b/Engine/Source/Programs/Shared/EpicGames.Horde/Compute/RequestIdAllocator.cs
@@ -12,8 +12,8 @@ namespace EpicGames.Horde.Compute;
 public class RequestIdAllocator
 {
 	private readonly object _lock = new();
-	private readonly List<string> _pendingIds = new();
-	private readonly List<string> _batchIds = new();
+	private readonly List<string> _pendingIds = [];
+	private readonly List<string> _batchIds = [];
 	private readonly string _allocatorId;
 	private int _requestCounter;
 
diff --git a/Engine/Source/Programs/Shared/EpicGames.Horde/Compute/Requirements.cs b/Engine/Source/Programs/Shared/EpicGames.Horde/Compute/Requirements.cs
--- a/Engine/Source/Programs/Shared/EpicGames.Horde/Compute/Requirements.cs
+++ b/Engine/Source/Programs/Shared/EpicGames.Horde/Compute/Requirements.cs
@@ -26,12 +26,12 @@ namespace EpicGames.Horde.Compute
 		/// <summary>
 		/// Properties required from the remote machine
 		/// </summary>
-		public HashSet<string> Properties { get; set; } = new HashSet<string>();
+		public HashSet<string> Properties { get; set; } = [];
 
 		/// <summary>
 		/// Resources used by the process
 		/// </summary>
-		public Dictionary<string, ResourceRequirements> Resources { get; } = new Dictionary<string, ResourceRequirements>();
+		public Dictionary<string, ResourceRequirements> Resources { get; } = [];
 
 		/// <summary>
 		/// Whether we require exclusive access to the device
@@ -57,7 +57,7 @@ namespace EpicGames.Horde.Compute
 		/// <inheritdoc/>
 		public override string ToString()
 		{
-			List<string> list = new List<string>();
+			List<string> list = [];
 			if (Pool != null)
 			{
 				list.Add($"Pool:{Pool}");
diff --git a/Engine/Source/Programs/Shared/EpicGames.Horde/Compute/Transports/TcpSslTransport.cs b/Engine/Source/Programs/Shared/EpicGames.Horde/Compute/Transports/TcpSslTransport.cs
--- a/Engine/Source/Programs/Shared/EpicGames.Horde/Compute/Transports/TcpSslTransport.cs
+++ b/Engine/Source/Programs/Shared/EpicGames.Horde/Compute/Transports/TcpSslTransport.cs
@@ -75,7 +75,7 @@ public sealed class TcpSslTransport : ComputeTransport
 		SslClientAuthenticationOptions options = new()
 		{
 			TargetHost = "horde", // Since cert is self-signed and not CA-based, this is ignored
-			ClientCertificates = new X509CertificateCollection { _cert },
+			ClientCertificates = [_cert],
 			CertificateRevocationCheckMode = X509RevocationMode.Offline,
 			EncryptionPolicy = EncryptionPolicy.RequireEncryption,
 			RemoteCertificateValidationCallback = ValidateCert
diff --git a/Engine/Source/Programs/Shared/EpicGames.Horde/Dashboard/DashboardMessages.cs b/Engine/Source/Programs/Shared/EpicGames.Horde/Dashboard/DashboardMessages.cs
--- a/Engine/Source/Programs/Shared/EpicGames.Horde/Dashboard/DashboardMessages.cs
+++ b/Engine/Source/Programs/Shared/EpicGames.Horde/Dashboard/DashboardMessages.cs
@@ -63,22 +63,22 @@ namespace EpicGames.Horde.Dashboard
 		/// <summary>
 		/// Categories to display on the agents page
 		/// </summary>
-		public List<GetDashboardAgentCategoryResponse> AgentCategories { get; set; } = new List<GetDashboardAgentCategoryResponse>();
+		public List<GetDashboardAgentCategoryResponse> AgentCategories { get; set; } = [];
 
 		/// <summary>
 		/// Categories to display on the pools page
 		/// </summary>
-		public List<GetDashboardPoolCategoryResponse> PoolCategories { get; set; } = new List<GetDashboardPoolCategoryResponse>();
+		public List<GetDashboardPoolCategoryResponse> PoolCategories { get; set; } = [];
 
 		/// <summary>
 		/// Configured artifact types
 		/// </summary>
-		public List<string> ArtifactTypes { get; set; } = new List<string>();
+		public List<string> ArtifactTypes { get; set; } = [];
 
 		/// <summary>
 		/// Configured artifact type info
 		/// </summary>
-		public List<GetDashboardArtifactTypeResponse> ArtifactTypeInfo { get; set; } = new List<GetDashboardArtifactTypeResponse>();
+		public List<GetDashboardArtifactTypeResponse> ArtifactTypeInfo { get; set; } = [];
 
 	}
 
diff --git a/Engine/Source/Programs/Shared/EpicGames.Horde/Devices/DeviceMessages.cs b/Engine/Source/Programs/Shared/EpicGames.Horde/Devices/DeviceMessages.cs
--- a/Engine/Source/Programs/Shared/EpicGames.Horde/Devices/DeviceMessages.cs
+++ b/Engine/Source/Programs/Shared/EpicGames.Horde/Devices/DeviceMessages.cs
@@ -533,9 +533,9 @@ namespace EpicGames.Horde.Devices
 		/// </summary>
 		public HashSet<string> GetDuplicateEntries()
 		{
-			List<List<string>?> tagSets = new List<List<string>?>() { RequiredTags, PreferredTags, UndesiredTags, BlockedTags };
-			HashSet<string> seenTags = new HashSet<string>();
-			HashSet<string> duplicateTags = new HashSet<string>();
+			List<List<string>?> tagSets = [RequiredTags, PreferredTags, UndesiredTags, BlockedTags];
+			HashSet<string> seenTags = [];
+			HashSet<string> duplicateTags = [];
 
 			foreach (List<string>? set in tagSets)
 			{
diff --git a/Engine/Source/Programs/Shared/EpicGames.Horde/HordeHttpClient.cs b/Engine/Source/Programs/Shared/EpicGames.Horde/HordeHttpClient.cs
--- a/Engine/Source/Programs/Shared/EpicGames.Horde/HordeHttpClient.cs
+++ b/Engine/Source/Programs/Shared/EpicGames.Horde/HordeHttpClient.cs
@@ -139,7 +139,7 @@ namespace EpicGames.Horde
 		/// <param name="cancellationToken">Cancellation token for the operation</param>
 		public Task<CreateArtifactResponse> CreateArtifactAsync(ArtifactName name, ArtifactType type, string? description, StreamId streamId, CommitId commitId, IEnumerable<string>? keys = null, IEnumerable<string>? metadata = null, CancellationToken cancellationToken = default)
 		{
-			return PostAsync<CreateArtifactResponse, CreateArtifactRequest>(_httpClient, $"api/v2/artifacts", new CreateArtifactRequest(name, type, description, streamId, keys?.ToList() ?? new List<string>(), metadata?.ToList() ?? new List<string>()) { CommitId = commitId }, cancellationToken);
+			return PostAsync<CreateArtifactResponse, CreateArtifactRequest>(_httpClient, $"api/v2/artifacts", new CreateArtifactRequest(name, type, description, streamId, keys?.ToList() ?? [], metadata?.ToList() ?? []) { CommitId = commitId }, cancellationToken);
 		}
 
 		/// <summary>
@@ -647,7 +647,7 @@ namespace EpicGames.Horde
 				queryBuilder.Append($"&maxCreateTime={Uri.EscapeDataString(maxCreateTime.Value.ToString("o"))}");
 			}
 
-			filters = filters ?? ["id", "name", "change", "preflightChange", "createTime"];
+			filters ??= ["id", "name", "change", "preflightChange", "createTime"];
 			if (filters.Count > 0)
 			{
 				queryBuilder.Append($"&filter={Uri.EscapeDataString(String.Join(",", filters))}");
@@ -722,7 +722,7 @@ namespace EpicGames.Horde
 			}
 			else
 			{
-				throw new ArgumentException($"StreamId '{streamId.ToString()} is invalid'");
+				throw new ArgumentException($"StreamId '{streamId} is invalid'");
 			}
 		}
 		#endregion
diff --git a/Engine/Source/Programs/Shared/EpicGames.Horde/HordeHttpMessageHandler.cs b/Engine/Source/Programs/Shared/EpicGames.Horde/HordeHttpMessageHandler.cs
--- a/Engine/Source/Programs/Shared/EpicGames.Horde/HordeHttpMessageHandler.cs
+++ b/Engine/Source/Programs/Shared/EpicGames.Horde/HordeHttpMessageHandler.cs
@@ -59,7 +59,7 @@ namespace EpicGames.Horde
 				logger.LogWarning(KnownLogEvents.Horde_BuildHealth_Ignore, ex, "{Method} {Url} retrying after {Time}s.", request.Method, request.RedactedRequestUri(), timespan.TotalSeconds);
 			}
 
-			TimeSpan[] retryTimes = new[] { TimeSpan.FromSeconds(5.0), TimeSpan.FromSeconds(10.0), TimeSpan.FromSeconds(15.0) };
+			TimeSpan[] retryTimes = [TimeSpan.FromSeconds(5.0), TimeSpan.FromSeconds(10.0), TimeSpan.FromSeconds(15.0)];
 			AsyncRetryPolicy retryPolicy = Policy.Handle<TimeoutRejectedException>().WaitAndRetryAsync(retryTimes, OnRetry);
 			return retryPolicy.WrapAsync(timeoutPolicy);
 		}
@@ -75,7 +75,7 @@ namespace EpicGames.Horde
 				return Task.CompletedTask;
 			}
 
-			TimeSpan[] retryTimes = new[] { TimeSpan.FromSeconds(1.0), TimeSpan.FromSeconds(5.0), TimeSpan.FromSeconds(10.0), TimeSpan.FromSeconds(30.0), TimeSpan.FromSeconds(30.0) };
+			TimeSpan[] retryTimes = [TimeSpan.FromSeconds(1.0), TimeSpan.FromSeconds(5.0), TimeSpan.FromSeconds(10.0), TimeSpan.FromSeconds(30.0), TimeSpan.FromSeconds(30.0)];
 
 			// Policy for transient errors is the same as HttpPolicyExtensions.HandleTransientHttpError(), but excludes HttpStatusCode.ServiceUnavailable (which is used as a response
 			// when allocating compute resources when none are available). This pathway is handled explicitly on the application side.
diff --git a/Engine/Source/Programs/Shared/EpicGames.Horde/HordeOptions.cs b/Engine/Source/Programs/Shared/EpicGames.Horde/HordeOptions.cs
--- a/Engine/Source/Programs/Shared/EpicGames.Horde/HordeOptions.cs
+++ b/Engine/Source/Programs/Shared/EpicGames.Horde/HordeOptions.cs
@@ -100,7 +100,7 @@ namespace EpicGames.Horde
 					byte[] data = FileReference.ReadAllBytes(configFile);
 
 					JsonObject? root = JsonNode.Parse(data, new JsonNodeOptions { PropertyNameCaseInsensitive = true }, new JsonDocumentOptions { AllowTrailingCommas = true }) as JsonObject;
-					root ??= new JsonObject();
+					root ??= [];
 
 					JsonNode? value;
 					if (root.TryGetPropertyValue("server", out value))
@@ -160,7 +160,7 @@ namespace EpicGames.Horde
 						root = JsonNode.Parse(data, new JsonNodeOptions { PropertyNameCaseInsensitive = true }, new JsonDocumentOptions { AllowTrailingCommas = true }) as JsonObject;
 					}
 
-					root ??= new JsonObject();
+					root ??= [];
 					root["server"] = serverUrl.ToString();
 
 					using (FileStream stream = FileReference.Open(configFile, FileMode.Create, FileAccess.ReadWrite, FileShare.Read))
diff --git a/Engine/Source/Programs/Shared/EpicGames.Horde/Issues/Handlers/BuildGraphIssueHandler.cs b/Engine/Source/Programs/Shared/EpicGames.Horde/Issues/Handlers/BuildGraphIssueHandler.cs
--- a/Engine/Source/Programs/Shared/EpicGames.Horde/Issues/Handlers/BuildGraphIssueHandler.cs
+++ b/Engine/Source/Programs/Shared/EpicGames.Horde/Issues/Handlers/BuildGraphIssueHandler.cs
@@ -12,7 +12,7 @@ namespace EpicGames.Horde.Issues.Handlers
 	[IssueHandler]
 	public class BuildGraphIssueHandler : IssueHandler
 	{
-		readonly List<IssueEventGroup> _issues = new List<IssueEventGroup>();
+		readonly List<IssueEventGroup> _issues = [];
 
 		/// <inheritdoc/>
 		public override int Priority => 10;
diff --git a/Engine/Source/Programs/Shared/EpicGames.Horde/Issues/Handlers/CompileIssueHandler.cs b/Engine/Source/Programs/Shared/EpicGames.Horde/Issues/Handlers/CompileIssueHandler.cs
--- a/Engine/Source/Programs/Shared/EpicGames.Horde/Issues/Handlers/CompileIssueHandler.cs
+++ b/Engine/Source/Programs/Shared/EpicGames.Horde/Issues/Handlers/CompileIssueHandler.cs
@@ -23,7 +23,7 @@ namespace EpicGames.Horde.Issues.Handlers
 		const string CompileGroupAnnotation = "CompileGroup";
 
 		readonly IssueHandlerContext _context;
-		readonly List<IssueEventGroup> _issues = new List<IssueEventGroup>();
+		readonly List<IssueEventGroup> _issues = [];
 
 		/// <inheritdoc/>
 		public override int Priority => 10;
diff --git a/Engine/Source/Programs/Shared/EpicGames.Horde/Issues/Handlers/GauntletIssueHandler.cs b/Engine/Source/Programs/Shared/EpicGames.Horde/Issues/Handlers/GauntletIssueHandler.cs
--- a/Engine/Source/Programs/Shared/EpicGames.Horde/Issues/Handlers/GauntletIssueHandler.cs
+++ b/Engine/Source/Programs/Shared/EpicGames.Horde/Issues/Handlers/GauntletIssueHandler.cs
@@ -78,7 +78,7 @@ namespace EpicGames.Horde.Issues.Handlers
 		bool _wasErrorEventReported = false;
 
 		readonly IssueHandlerContext _context;
-		readonly List<IssueEventGroup> _issues = new List<IssueEventGroup>();
+		readonly List<IssueEventGroup> _issues = [];
 
 		/// <summary>
 		///  Known Gauntlet events
@@ -96,28 +96,28 @@ namespace EpicGames.Horde.Issues.Handlers
 		/// <summary>
 		///  Known Gauntlet events associated with the highest severity
 		/// </summary>
-		static readonly HashSet<EventId> s_knownSevereGauntletEvents = new HashSet<EventId>
-		{
+		static readonly HashSet<EventId> s_knownSevereGauntletEvents =
+		[
 			KnownLogEvents.Gauntlet_BuildDropEvent,
 			KnownLogEvents.Gauntlet_FatalEvent
-		};
+		];
 
 		/// <summary>
 		///  Known Gauntlet events associated with test context
 		/// </summary>
-		static readonly HashSet<EventId> s_knownTestGauntletEvents = new HashSet<EventId>
-		{
+		static readonly HashSet<EventId> s_knownTestGauntletEvents =
+		[
 			KnownLogEvents.Gauntlet_TestEvent,
 			KnownLogEvents.Gauntlet_UnrealEngineTestEvent
-		};
+		];
 
 		/// <summary>
 		///   Gauntlet events to skip issue generation
 		/// </summary>
-		static readonly HashSet<EventId> s_skipGauntletEvents = new HashSet<EventId>
-		{
+		static readonly HashSet<EventId> s_skipGauntletEvents =
+		[
 			KnownLogEvents.Gauntlet_DeviceEvent
-		};
+		];
 
 		/// <summary>
 		///  Known summaries that are not specific enough
@@ -226,10 +226,7 @@ namespace EpicGames.Horde.Issues.Handlers
 				JsonProperty? property = FindNestedPropertyOfType(logEvent, SummaryLogType);
 				if (property != null)
 				{
-					if (summary == null)
-					{
-						summary = new StringBuilder();
-					}
+					summary ??= new();
 					JsonElement value = property.Value.Value;
 					if (value.ValueKind == JsonValueKind.String
 						// handle LogValue type
diff --git a/Engine/Source/Programs/Shared/EpicGames.Horde/Issues/Handlers/HashedIssueHandler.cs b/Engine/Source/Programs/Shared/EpicGames.Horde/Issues/Handlers/HashedIssueHandler.cs
--- a/Engine/Source/Programs/Shared/EpicGames.Horde/Issues/Handlers/HashedIssueHandler.cs
+++ b/Engine/Source/Programs/Shared/EpicGames.Horde/Issues/Handlers/HashedIssueHandler.cs
@@ -18,12 +18,12 @@ namespace EpicGames.Horde.Issues.Handlers
 	public class HashedIssueHandler : IssueHandler
 	{
 		readonly IssueHandlerContext _context;
-		readonly List<IssueEvent> _issueEvents = new List<IssueEvent>();
+		readonly List<IssueEvent> _issueEvents = [];
 
 		/// <summary>
 		/// Known general events
 		/// </summary>
-		static readonly HashSet<EventId> s_knownGeneralEvents = new HashSet<EventId> { KnownLogEvents.Generic, KnownLogEvents.ExitCode, KnownLogEvents.Horde, KnownLogEvents.Horde_InvalidPreflight };
+		static readonly HashSet<EventId> s_knownGeneralEvents = [KnownLogEvents.Generic, KnownLogEvents.ExitCode, KnownLogEvents.Horde, KnownLogEvents.Horde_InvalidPreflight];
 
 		/// <summary>
 		/// Constructor
@@ -95,13 +95,13 @@ namespace EpicGames.Horde.Issues.Handlers
 		/// <inheritdoc/>
 		public override IEnumerable<IssueEventGroup> GetIssues()
 		{
-			List<IssueEventGroup> issues = new List<IssueEventGroup>();
+			List<IssueEventGroup> issues = [];
 
 			// keep hash consistent when only have general, non-unique events
 			bool allGeneral = _issueEvents.FirstOrDefault(stepEvent => stepEvent.EventId == null || !IsGeneralEventId(stepEvent.EventId.Value)) == null;
 
-			Dictionary<string, List<IssueEvent>> fingerEvents = new Dictionary<string, List<IssueEvent>>();
-			Dictionary<IssueEvent, string?> eventScopes = new Dictionary<IssueEvent, string?>();
+			Dictionary<string, List<IssueEvent>> fingerEvents = [];
+			Dictionary<IssueEvent, string?> eventScopes = [];
 
 			foreach (IssueEvent stepEvent in _issueEvents)
 			{
@@ -126,7 +126,7 @@ namespace EpicGames.Horde.Issues.Handlers
 				List<IssueEvent>? events = null;
 				if (!fingerEvents.TryGetValue(fingerprint, out events))
 				{
-					events = new List<IssueEvent>();
+					events = [];
 					fingerEvents.Add(fingerprint, events);
 				}
 
@@ -139,7 +139,7 @@ namespace EpicGames.Horde.Issues.Handlers
 
 				IssueEventGroup? genericFingerprint = null;
 				IssueEventGroup? genericErrorsFingerprint = null;
-				HashSet<Md5Hash> hashes = new HashSet<Md5Hash>();
+				HashSet<Md5Hash> hashes = [];
 
 				foreach (IssueEvent stepEvent in item.Value)
 				{
@@ -235,7 +235,7 @@ namespace EpicGames.Horde.Issues.Handlers
 							if (property.NameEquals("file") && property.Value.ValueKind == JsonValueKind.String)
 							{
 								string file = property.Value.GetString()!;
-								int endIdx = file.LastIndexOfAny(new char[] { '/', '\\' }) + 1;
+								int endIdx = file.LastIndexOfAny(['/', '\\']) + 1;
 								string fileName = file.Substring(endIdx);
 
 								issue.Metadata.Add("File", fileName);
diff --git a/Engine/Source/Programs/Shared/EpicGames.Horde/Issues/Handlers/IgnoredIssueHandler.cs b/Engine/Source/Programs/Shared/EpicGames.Horde/Issues/Handlers/IgnoredIssueHandler.cs
--- a/Engine/Source/Programs/Shared/EpicGames.Horde/Issues/Handlers/IgnoredIssueHandler.cs
+++ b/Engine/Source/Programs/Shared/EpicGames.Horde/Issues/Handlers/IgnoredIssueHandler.cs
@@ -15,7 +15,7 @@ namespace EpicGames.Horde.Issues.Handlers
 		/// <summary>
 		///  Known systemic errors
 		/// </summary>
-		static readonly HashSet<EventId> s_handledEvents = new HashSet<EventId> { KnownLogEvents.Horde_BuildHealth_Ignore };
+		static readonly HashSet<EventId> s_handledEvents = [KnownLogEvents.Horde_BuildHealth_Ignore];
 
 		/// <inheritdoc/>
 		public override int Priority => 100;
@@ -30,7 +30,7 @@ namespace EpicGames.Horde.Issues.Handlers
 			return s_handledEvents.Contains(eventId);
 		}
 
-		readonly List<IssueEventGroup> _issues = new List<IssueEventGroup>();
+		readonly List<IssueEventGroup> _issues = [];
 
 		/// <summary>
 		/// Constructor
diff --git a/Engine/Source/Programs/Shared/EpicGames.Horde/Issues/Handlers/SystemicIssueHandler.cs b/Engine/Source/Programs/Shared/EpicGames.Horde/Issues/Handlers/SystemicIssueHandler.cs
--- a/Engine/Source/Programs/Shared/EpicGames.Horde/Issues/Handlers/SystemicIssueHandler.cs
+++ b/Engine/Source/Programs/Shared/EpicGames.Horde/Issues/Handlers/SystemicIssueHandler.cs
@@ -15,7 +15,7 @@ namespace EpicGames.Horde.Issues.Handlers
 		/// <summary>
 		///  Known systemic errors
 		/// </summary>
-		static readonly HashSet<EventId> s_knownSystemic = new HashSet<EventId> { KnownLogEvents.Horde, KnownLogEvents.Horde_InvalidPreflight };
+		static readonly HashSet<EventId> s_knownSystemic = [KnownLogEvents.Horde, KnownLogEvents.Horde_InvalidPreflight];
 
 		/// <inheritdoc/>
 		public override int Priority => 10;
@@ -32,7 +32,7 @@ namespace EpicGames.Horde.Issues.Handlers
 
 		bool _nonSystemicError = false;
 		readonly IssueHandlerContext _context;
-		readonly List<IssueEventGroup> _issues = new List<IssueEventGroup>();
+		readonly List<IssueEventGroup> _issues = [];
 
 		/// <summary>
 		/// Constructor
diff --git a/Engine/Source/Programs/Shared/EpicGames.Horde/Issues/IssueChangeFilter.cs b/Engine/Source/Programs/Shared/EpicGames.Horde/Issues/IssueChangeFilter.cs
--- a/Engine/Source/Programs/Shared/EpicGames.Horde/Issues/IssueChangeFilter.cs
+++ b/Engine/Source/Programs/Shared/EpicGames.Horde/Issues/IssueChangeFilter.cs
@@ -1,6 +1,5 @@
 // Copyright Epic Games, Inc. All Rights Reserved.
 
-using System;
 using System.Collections.Generic;
 
 namespace EpicGames.Horde.Issues
@@ -13,18 +12,18 @@ namespace EpicGames.Horde.Issues
 		/// <summary>
 		/// Filter exclude all changes
 		/// </summary>
-		public static IReadOnlyList<string> None { get; } = Array.Empty<string>();
+		public static IReadOnlyList<string> None { get; } = [];
 
 		/// <summary>
 		/// Filter including all changes
 		/// </summary>
-		public static IReadOnlyList<string> All { get; } = new[] { "..." };
+		public static IReadOnlyList<string> All { get; } = ["..."];
 
 		/// <summary>
 		/// Set of extensions to treat as code
 		/// </summary>
-		public static IReadOnlyList<string> Code { get; } = new[]
-		{
+		public static IReadOnlyList<string> Code { get; } =
+		[
 			"*.c",
 			"*.cc",
 			"*.cxx",
@@ -44,16 +43,16 @@ namespace EpicGames.Horde.Issues
 			"*.uplugin",
 			"*.sln",
 			"*.verse"
-		};
+		];
 
 		/// <summary>
 		/// Set of file extensions to treat as content
 		/// </summary>
-		public static IReadOnlyList<string> Content { get; } = new[]
-		{
+		public static IReadOnlyList<string> Content { get; } =
+		[
 			"*.uasset",
 			"*.umap",
 			"*.ini"
-		};
+		];
 	}
 }
diff --git a/Engine/Source/Programs/Shared/EpicGames.Horde/Issues/IssueEvent.cs b/Engine/Source/Programs/Shared/EpicGames.Horde/Issues/IssueEvent.cs
--- a/Engine/Source/Programs/Shared/EpicGames.Horde/Issues/IssueEvent.cs
+++ b/Engine/Source/Programs/Shared/EpicGames.Horde/Issues/IssueEvent.cs
@@ -77,12 +77,12 @@ namespace EpicGames.Horde.Issues
 		/// <summary>
 		/// List of keys which identify this issue.
 		/// </summary>
-		public HashSet<IssueKey> Keys { get; } = new HashSet<IssueKey>();
+		public HashSet<IssueKey> Keys { get; } = [];
 
 		/// <summary>
 		/// Collection of additional metadata added by the handler
 		/// </summary>
-		public HashSet<IssueMetadata> Metadata { get; } = new HashSet<IssueMetadata>();
+		public HashSet<IssueMetadata> Metadata { get; } = [];
 
 		/// <summary>
 		/// Filter for changes that should be included in this issue
@@ -92,7 +92,7 @@ namespace EpicGames.Horde.Issues
 		/// <summary>
 		/// Individual log events
 		/// </summary>
-		public List<IssueEvent> Events { get; } = new List<IssueEvent>();
+		public List<IssueEvent> Events { get; } = [];
 
 		/// <summary>
 		/// Constructor
@@ -137,7 +137,7 @@ namespace EpicGames.Horde.Issues
 		/// </summary>
 		public IssueAuditLogger()
 		{
-			_entries = new List<IssueAuditLogEntry>();
+			_entries = [];
 		}
 
 		/// <summary>
diff --git a/Engine/Source/Programs/Shared/EpicGames.Horde/Issues/IssueFingerprint.cs b/Engine/Source/Programs/Shared/EpicGames.Horde/Issues/IssueFingerprint.cs
--- a/Engine/Source/Programs/Shared/EpicGames.Horde/Issues/IssueFingerprint.cs
+++ b/Engine/Source/Programs/Shared/EpicGames.Horde/Issues/IssueFingerprint.cs
@@ -29,7 +29,7 @@ namespace EpicGames.Horde.Issues
 		/// List of keys which identify this issue.
 		/// </summary>
 		[JsonIgnore]
-		public HashSet<IssueKey> Keys { get; set; } = new HashSet<IssueKey>();
+		public HashSet<IssueKey> Keys { get; set; } = [];
 		IReadOnlySet<IssueKey> IIssueFingerprint.Keys => Keys;
 
 		/// <summary>
@@ -43,7 +43,7 @@ namespace EpicGames.Horde.Issues
 		/// Collection of additional metadata added by the handler
 		/// </summary>
 		[JsonIgnore]
-		public HashSet<IssueMetadata> Metadata { get; set; } = new HashSet<IssueMetadata>();
+		public HashSet<IssueMetadata> Metadata { get; set; } = [];
 		IReadOnlySet<IssueMetadata> IIssueFingerprint.Metadata => Metadata;
 
 		/// <summary>
@@ -60,7 +60,7 @@ namespace EpicGames.Horde.Issues
 		public HashSet<IssueKey>? SerializedKeys
 		{
 			get => (Keys.Count == 0) ? null : Keys;
-			set => Keys = value ?? new HashSet<IssueKey>();
+			set => Keys = value ?? [];
 		}
 
 		/// <summary>
@@ -71,7 +71,7 @@ namespace EpicGames.Horde.Issues
 		public HashSet<IssueMetadata>? SerializedMetadata
 		{
 			get => (Metadata.Count == 0) ? null : Metadata;
-			set => Metadata = value ?? new HashSet<IssueMetadata>();
+			set => Metadata = value ?? [];
 		}
 
 		/// <summary>
diff --git a/Engine/Source/Programs/Shared/EpicGames.Horde/Issues/IssueKey.cs b/Engine/Source/Programs/Shared/EpicGames.Horde/Issues/IssueKey.cs
--- a/Engine/Source/Programs/Shared/EpicGames.Horde/Issues/IssueKey.cs
+++ b/Engine/Source/Programs/Shared/EpicGames.Horde/Issues/IssueKey.cs
@@ -161,7 +161,7 @@ namespace EpicGames.Horde.Issues
 				string? relativePath;
 				if (document.RootElement.TryGetNestedProperty("properties.asset.relativePath", out relativePath) || document.RootElement.TryGetNestedProperty("properties.asset.$text", out relativePath))
 				{
-					int endIdx = relativePath.LastIndexOfAny(new char[] { '/', '\\' }) + 1;
+					int endIdx = relativePath.LastIndexOfAny(['/', '\\']) + 1;
 					string fileName = relativePath.Substring(endIdx);
 					IssueKey issueKey = new IssueKey(fileName, IssueKeyType.File);
 					keys.Add(issueKey);
@@ -244,7 +244,7 @@ namespace EpicGames.Horde.Issues
 		/// <param name="scope">The scope of the key to add</param>
 		public static void AddSourceFile(this HashSet<IssueKey> keys, string relativePath, IssueKeyType type, string? scope = null)
 		{
-			int endIdx = relativePath.LastIndexOfAny(new char[] { '/', '\\' }) + 1;
+			int endIdx = relativePath.LastIndexOfAny(['/', '\\']) + 1;
 
 			string fileName = relativePath.Substring(endIdx);
 			IssueKey key = new IssueKey(fileName, type, scope);
diff --git a/Engine/Source/Programs/Shared/EpicGames.Horde/Issues/IssueMessages.cs b/Engine/Source/Programs/Shared/EpicGames.Horde/Issues/IssueMessages.cs
--- a/Engine/Source/Programs/Shared/EpicGames.Horde/Issues/IssueMessages.cs
+++ b/Engine/Source/Programs/Shared/EpicGames.Horde/Issues/IssueMessages.cs
@@ -129,7 +129,7 @@ namespace EpicGames.Horde.Issues
 		/// <summary>
 		/// The failing builds for a particular event
 		/// </summary>
-		public List<GetIssueStepResponse> Steps { get; set; } = new List<GetIssueStepResponse>();
+		public List<GetIssueStepResponse> Steps { get; set; } = [];
 
 		/// <summary>
 		/// The following successful build
@@ -192,7 +192,7 @@ namespace EpicGames.Horde.Issues
 		/// <summary>
 		/// Map of steps to (event signature id -> trace id)
 		/// </summary>
-		public List<GetIssueSpanResponse> Nodes { get; set; } = new List<GetIssueSpanResponse>();
+		public List<GetIssueSpanResponse> Nodes { get; set; } = [];
 	}
 
 	/// <summary>
@@ -270,22 +270,22 @@ namespace EpicGames.Horde.Issues
 		/// <summary>
 		/// The affected templates
 		/// </summary>
-		public List<GetIssueAffectedTemplateResponse> AffectedTemplates { get; set; } = new List<GetIssueAffectedTemplateResponse>();
+		public List<GetIssueAffectedTemplateResponse> AffectedTemplates { get; set; } = [];
 
 		/// <summary>
 		/// List of affected template ids
 		/// </summary>
-		public List<string> TemplateIds { get; set; } = new List<string>();
+		public List<string> TemplateIds { get; set; } = [];
 
 		/// <summary>
 		/// List of resolved template ids
 		/// </summary>
-		public List<string> ResolvedTemplateIds { get; set; } = new List<string>();
+		public List<string> ResolvedTemplateIds { get; set; } = [];
 
 		/// <summary>
 		/// List of unresolved template ids
 		/// </summary>
-		public List<string> UnresolvedTemplateIds { get; set; } = new List<string>();
+		public List<string> UnresolvedTemplateIds { get; set; } = [];
 	}
 
 	/// <summary>
@@ -461,22 +461,22 @@ namespace EpicGames.Horde.Issues
 		/// <summary>
 		/// List of stream paths affected by this issue
 		/// </summary>
-		public List<string> Streams { get; set; } = new List<string>();
+		public List<string> Streams { get; set; } = [];
 
 		/// <summary>
 		/// List of affected stream ids
 		/// </summary>
-		public List<string> ResolvedStreams { get; set; } = new List<string>();
+		public List<string> ResolvedStreams { get; set; } = [];
 
 		/// <summary>
 		/// List of unresolved streams
 		/// </summary>
-		public List<string> UnresolvedStreams { get; set; } = new List<string>();
+		public List<string> UnresolvedStreams { get; set; } = [];
 
 		/// <summary>
 		/// List of affected streams
 		/// </summary>
-		public List<GetIssueAffectedStreamResponse> AffectedStreams { get; set; } = new List<GetIssueAffectedStreamResponse>();
+		public List<GetIssueAffectedStreamResponse> AffectedStreams { get; set; } = [];
 
 		/// <summary>
 		/// Most likely suspects for causing this issue [DEPRECATED]
@@ -491,7 +491,7 @@ namespace EpicGames.Horde.Issues
 		/// <summary>
 		/// Most likely suspects for causing this issue
 		/// </summary>
-		public List<GetThinUserInfoResponse> PrimarySuspectsInfo { get; set; } = new List<GetThinUserInfoResponse>();
+		public List<GetThinUserInfoResponse> PrimarySuspectsInfo { get; set; } = [];
 
 		/// <summary>
 		/// Whether to show alerts for this issue
@@ -891,12 +891,12 @@ namespace EpicGames.Horde.Issues
 		/// <summary>
 		/// component id => name
 		/// </summary>
-		public Dictionary<string, string> Components { get; set; } = new Dictionary<string, string>();
+		public Dictionary<string, string> Components { get; set; } = [];
 
 		/// <summary>
 		/// IssueType id => name
 		/// </summary>
-		public Dictionary<string, string> IssueTypes { get; set; } = new Dictionary<string, string>();
+		public Dictionary<string, string> IssueTypes { get; set; } = [];
 	}
 
 	/// <summary>
diff --git a/Engine/Source/Programs/Shared/EpicGames.Horde/Issues/WorkflowMessages.cs b/Engine/Source/Programs/Shared/EpicGames.Horde/Issues/WorkflowMessages.cs
--- a/Engine/Source/Programs/Shared/EpicGames.Horde/Issues/WorkflowMessages.cs
+++ b/Engine/Source/Programs/Shared/EpicGames.Horde/Issues/WorkflowMessages.cs
@@ -89,7 +89,7 @@ namespace EpicGames.Horde.Issues
 		public bool ShowMergeWarnings { get; set; }
 
 		/// <inheritdoc cref="IWorkflow.Annotations"/>
-		public NodeAnnotations Annotations { get; set; } = new NodeAnnotations();
+		public NodeAnnotations Annotations { get; set; } = [];
 		IReadOnlyNodeAnnotations IWorkflow.Annotations => Annotations;
 
 		/// <inheritdoc cref="IWorkflow.ExternalIssues"/>
@@ -105,8 +105,8 @@ namespace EpicGames.Horde.Issues
 		/// </summary>
 		public GetWorkflowResponse()
 		{
-			ReportTimes = new List<TimeSpan>();
-			EscalateTimes = new List<int>();
+			ReportTimes = [];
+			EscalateTimes = [];
 			EscalateRootCauseAnalysisTimes = [];
 		}
 
@@ -116,7 +116,7 @@ namespace EpicGames.Horde.Issues
 		public GetWorkflowResponse(IWorkflow workflow)
 		{
 			Id = workflow.Id;
-			ReportTimes = workflow.ReportTimes.ToList();
+			ReportTimes = [.. workflow.ReportTimes];
 			SummaryTab = workflow.SummaryTab;
 			ReportChannel = workflow.ReportChannel;
 			ReportWarnings = workflow.ReportWarnings;
@@ -130,8 +130,8 @@ namespace EpicGames.Horde.Issues
 			TriageWarnings = workflow.TriageWarnings;
 			TriageErrors = workflow.TriageErrors;
 			EscalateAlias = workflow.EscalateAlias;
-			EscalateTimes = workflow.EscalateTimes.ToList();
-			EscalateRootCauseAnalysisTimes = workflow.EscalateRootCauseAnalysisTimes.ToList();
+			EscalateTimes = [.. workflow.EscalateTimes];
+			EscalateRootCauseAnalysisTimes = [.. workflow.EscalateRootCauseAnalysisTimes];
 			AutoRcaEnabled = workflow.AutoRcaEnabled;
 			MaxMentions = workflow.MaxMentions;
 			AllowMentions = workflow.AllowMentions;
diff --git a/Engine/Source/Programs/Shared/EpicGames.Horde/Jobs/Graphs/GraphMessages.cs b/Engine/Source/Programs/Shared/EpicGames.Horde/Jobs/Graphs/GraphMessages.cs
--- a/Engine/Source/Programs/Shared/EpicGames.Horde/Jobs/Graphs/GraphMessages.cs
+++ b/Engine/Source/Programs/Shared/EpicGames.Horde/Jobs/Graphs/GraphMessages.cs
@@ -19,22 +19,22 @@ namespace EpicGames.Horde.Jobs.Graphs
 		/// <summary>
 		/// Inputs for this node
 		/// </summary>
-		public List<string> Inputs { get; set; } = new List<string>();
+		public List<string> Inputs { get; set; } = [];
 
 		/// <summary>
 		/// Output from this node
 		/// </summary>
-		public List<string> Outputs { get; set; } = new List<string>();
+		public List<string> Outputs { get; set; } = [];
 
 		/// <summary>
 		/// Indices of nodes which must have succeeded for this node to run
 		/// </summary>
-		public List<string> InputDependencies { get; set; } = new List<string>();
+		public List<string> InputDependencies { get; set; } = [];
 
 		/// <summary>
 		/// Indices of nodes which must have completed for this node to run
 		/// </summary>
-		public List<string> OrderDependencies { get; set; } = new List<string>();
+		public List<string> OrderDependencies { get; set; } = [];
 
 		/// <summary>
 		/// The priority of this node
@@ -99,7 +99,7 @@ namespace EpicGames.Horde.Jobs.Graphs
 		/// <summary>
 		/// Nodes in the group
 		/// </summary>
-		public List<GetNodeResponse> Nodes { get; set; } = new List<GetNodeResponse>();
+		public List<GetNodeResponse> Nodes { get; set; } = [];
 
 		/// <summary>
 		/// Constructor
@@ -123,7 +123,7 @@ namespace EpicGames.Horde.Jobs.Graphs
 		/// <summary>
 		/// Nodes which must be part of the job for the aggregate to be shown
 		/// </summary>
-		public List<string> Nodes { get; set; } = new List<string>();
+		public List<string> Nodes { get; set; } = [];
 
 		/// <summary>
 		/// Constructor
@@ -172,12 +172,12 @@ namespace EpicGames.Horde.Jobs.Graphs
 		/// <summary>
 		/// Nodes which must be part of the job for the aggregate to be shown
 		/// </summary>
-		public List<string> RequiredNodes { get; set; } = new List<string>();
+		public List<string> RequiredNodes { get; set; } = [];
 
 		/// <summary>
 		/// Nodes to include in the status of this aggregate, if present in the job
 		/// </summary>
-		public List<string> IncludedNodes { get; set; } = new List<string>();
+		public List<string> IncludedNodes { get; set; } = [];
 	}
 
 	/// <summary>
diff --git a/Engine/Source/Programs/Shared/EpicGames.Horde/Jobs/Graphs/IGraph.cs b/Engine/Source/Programs/Shared/EpicGames.Horde/Jobs/Graphs/IGraph.cs
--- a/Engine/Source/Programs/Shared/EpicGames.Horde/Jobs/Graphs/IGraph.cs
+++ b/Engine/Source/Programs/Shared/EpicGames.Horde/Jobs/Graphs/IGraph.cs
@@ -582,7 +582,7 @@ namespace EpicGames.Horde.Jobs.Graphs
 		/// <param name="graph">Existing graph containing a node</param>
 		/// <param name="node">Node to copy</param>
 		public NewNode(IGraph graph, INode node)
-			: this(node.Name, node.Inputs.Select(x => graph.GetNode(x.NodeRef).Name).ToList(), node.OutputNames.ToList(), node.InputDependencies.Select(x => graph.GetNode(x).Name).ToList(), node.OrderDependencies.Select(x => graph.GetNode(x).Name).ToList(), node.Priority, node.AllowRetry, node.RunEarly, node.Warnings, node.Credentials?.ToDictionary(x => x.Key, x => x.Value), node.Properties?.ToDictionary(x => x.Key, x => x.Value), node.Annotations)
+			: this(node.Name, [.. node.Inputs.Select(x => graph.GetNode(x.NodeRef).Name)], [.. node.OutputNames], [.. node.InputDependencies.Select(x => graph.GetNode(x).Name)], [.. node.OrderDependencies.Select(x => graph.GetNode(x).Name)], node.Priority, node.AllowRetry, node.RunEarly, node.Warnings, node.Credentials?.ToDictionary(x => x.Key, x => x.Value), node.Properties?.ToDictionary(x => x.Key, x => x.Value), node.Annotations)
 		{
 		}
 	}
@@ -619,7 +619,7 @@ namespace EpicGames.Horde.Jobs.Graphs
 		/// <param name="graph">Graph containing the node group</param>
 		/// <param name="group">Node group to copy</param>
 		public NewGroup(IGraph graph, INodeGroup group)
-			: this(group.AgentType, group.Nodes.Select(x => new NewNode(graph, x)).ToList())
+			: this(group.AgentType, [.. group.Nodes.Select(x => new NewNode(graph, x))])
 		{
 		}
 	}
@@ -669,12 +669,12 @@ namespace EpicGames.Horde.Jobs.Graphs
 		/// <summary>
 		/// Nodes which must be part of the job for the aggregate to be valid
 		/// </summary>
-		public List<string> RequiredNodes { get; set; } = new List<string>();
+		public List<string> RequiredNodes { get; set; } = [];
 
 		/// <summary>
 		/// Nodes which must be part of the job for the aggregate to be valid
 		/// </summary>
-		public List<string> IncludedNodes { get; set; } = new List<string>();
+		public List<string> IncludedNodes { get; set; } = [];
 	}
 
 	/// <summary>
diff --git a/Engine/Source/Programs/Shared/EpicGames.Horde/Jobs/IJob.cs b/Engine/Source/Programs/Shared/EpicGames.Horde/Jobs/IJob.cs
--- a/Engine/Source/Programs/Shared/EpicGames.Horde/Jobs/IJob.cs
+++ b/Engine/Source/Programs/Shared/EpicGames.Horde/Jobs/IJob.cs
@@ -488,7 +488,7 @@ namespace EpicGames.Horde.Jobs
 			{
 				IReadOnlyDictionary<NodeRef, IJobStep> stepForNode = GetStepForNodeMap(job);
 
-				List<IJobStep> steps = new List<IJobStep>();
+				List<IJobStep> steps = [];
 				foreach (NodeRef aggregateNodeRef in aggregate.Nodes)
 				{
 					IJobStep? step;
@@ -611,7 +611,7 @@ namespace EpicGames.Horde.Jobs
 		/// <returns>Map of <see cref="NodeRef"/> to <see cref="IJobStep"/></returns>
 		public static IReadOnlyDictionary<NodeRef, IJobStep> GetStepForNodeMap(this IJob job)
 		{
-			Dictionary<NodeRef, IJobStep> stepForNode = new Dictionary<NodeRef, IJobStep>();
+			Dictionary<NodeRef, IJobStep> stepForNode = [];
 			foreach (IJobStepBatch batch in job.Batches)
 			{
 				foreach (IJobStep batchStep in batch.Steps)
@@ -898,12 +898,12 @@ namespace EpicGames.Horde.Jobs
 			if (labelIdx != -1 && labelIdx != graph.Labels.Count)
 			{
 				// Return all the nodes included by the label
-				return new HashSet<NodeRef>(graph.Labels[labelIdx].IncludedNodes);
+				return [.. graph.Labels[labelIdx].IncludedNodes];
 			}
 			else
 			{
 				// Set of nodes which are not covered by an existing label, initially containing everything
-				HashSet<NodeRef> unlabeledNodes = new HashSet<NodeRef>();
+				HashSet<NodeRef> unlabeledNodes = [];
 				for (int groupIdx = 0; groupIdx < graph.Groups.Count; groupIdx++)
 				{
 					INodeGroup group = graph.Groups[groupIdx];
@@ -939,7 +939,7 @@ namespace EpicGames.Horde.Jobs
 			IReadOnlyDictionary<NodeRef, IJobStep> stepForNode = job.GetStepForNodeMap();
 
 			// Set of nodes which are not covered by an existing label, initially containing everything
-			HashSet<NodeRef> unlabeledNodes = new HashSet<NodeRef>();
+			HashSet<NodeRef> unlabeledNodes = [];
 			for (int groupIdx = 0; groupIdx < graph.Groups.Count; groupIdx++)
 			{
 				INodeGroup group = graph.Groups[groupIdx];
@@ -1034,7 +1034,7 @@ namespace EpicGames.Horde.Jobs
 		{
 			IReadOnlyDictionary<NodeRef, IJobStep> stepForNodeRef = job.GetStepForNodeMap();
 
-			List<(LabelState, LabelOutcome)> states = new List<(LabelState, LabelOutcome)>();
+			List<(LabelState, LabelOutcome)> states = [];
 			for (int idx = 0; idx < graph.Labels.Count; idx++)
 			{
 				ILabel label = graph.Labels[idx];
@@ -1096,7 +1096,7 @@ namespace EpicGames.Horde.Jobs
 		public static Dictionary<int, UgsBadgeState> GetUgsBadgeStates(this IJob job, IGraph graph, IReadOnlyList<(LabelState, LabelOutcome)> labelStates)
 #pragma warning restore IDE0060 // Remove unused parameter
 		{
-			Dictionary<int, UgsBadgeState> ugsBadgeStates = new Dictionary<int, UgsBadgeState>();
+			Dictionary<int, UgsBadgeState> ugsBadgeStates = [];
 			for (int labelIdx = 0; labelIdx < labelStates.Count; ++labelIdx)
 			{
 				if (graph.Labels[labelIdx].UgsName == null)
@@ -1163,7 +1163,7 @@ namespace EpicGames.Horde.Jobs
 		/// <param name="newOutcome">Receives the outcome of the label</param>
 		public static void GetJobState(this IJob job, IReadOnlyDictionary<NodeRef, IJobStep> stepForNode, out LabelState newState, out LabelOutcome newOutcome)
 		{
-			List<NodeRef> nodes = new List<NodeRef>();
+			List<NodeRef> nodes = [];
 			foreach (IJobStepBatch batch in job.Batches)
 			{
 				foreach (IJobStep step in batch.Steps)
@@ -1445,14 +1445,14 @@ namespace EpicGames.Horde.Jobs
 		public static HashSet<INode> GetStartDependencies(IEnumerable<INode> batchNodes, IReadOnlyList<INodeGroup> groups)
 		{
 			// Find all the nodes that this group will start with.
-			List<INode> nodes = new List<INode>(batchNodes);
+			List<INode> nodes = [.. batchNodes];
 			if (nodes.Any(x => x.RunEarly))
 			{
 				nodes.RemoveAll(x => !x.RunEarly);
 			}
 
 			// Find all their dependencies
-			HashSet<INode> dependencies = new HashSet<INode>();
+			HashSet<INode> dependencies = [];
 			foreach (INode node in nodes)
 			{
 				dependencies.UnionWith(node.InputDependencies.Select(x => groups[x.GroupIdx].Nodes[x.NodeIdx]));
diff --git a/Engine/Source/Programs/Shared/EpicGames.Horde/Jobs/JobMessages.cs b/Engine/Source/Programs/Shared/EpicGames.Horde/Jobs/JobMessages.cs
--- a/Engine/Source/Programs/Shared/EpicGames.Horde/Jobs/JobMessages.cs
+++ b/Engine/Source/Programs/Shared/EpicGames.Horde/Jobs/JobMessages.cs
@@ -280,7 +280,7 @@ namespace EpicGames.Horde.Jobs
 		public ChangeQueryMessage? ChangeQuery
 		{
 			get => (ChangeQueries != null && ChangeQueries.Count > 0) ? ChangeQueries[0] : null;
-			set => ChangeQueries = (value == null) ? null : new List<ChangeQueryMessage> { value };
+			set => ChangeQueries = (value == null) ? null : [value];
 		}
 
 		/// <summary>
@@ -703,17 +703,17 @@ namespace EpicGames.Horde.Jobs
 		/// <summary>
 		/// Parameters for the job
 		/// </summary>
-		public Dictionary<ParameterId, string> Parameters { get; set; } = new Dictionary<ParameterId, string>();
+		public Dictionary<ParameterId, string> Parameters { get; set; } = [];
 
 		/// <summary>
 		/// Command line arguments for the job
 		/// </summary>
-		public List<string> Arguments { get; set; } = new List<string>();
+		public List<string> Arguments { get; set; } = [];
 
 		/// <summary>
 		/// Additional command line arguments for the job for when using the parameters block
 		/// </summary>
-		public List<string> AdditionalArguments { get; set; } = new List<string>();
+		public List<string> AdditionalArguments { get; set; } = [];
 
 		/// <summary>
 		/// Custom list of targets for the job
@@ -1292,7 +1292,7 @@ namespace EpicGames.Horde.Jobs
 		/// <summary>
 		/// Steps within this run
 		/// </summary>
-		public List<GetJobStepResponse> Steps { get; set; } = new List<GetJobStepResponse>();
+		public List<GetJobStepResponse> Steps { get; set; } = [];
 	}
 
 	/// <summary>
@@ -1380,7 +1380,7 @@ namespace EpicGames.Horde.Jobs
 		/// <summary>
 		/// Steps to include in the status of this label
 		/// </summary>
-		public List<JobStepId> Steps { get; set; } = new List<JobStepId>();
+		public List<JobStepId> Steps { get; set; } = [];
 	}
 
 	/// <summary>
@@ -1391,7 +1391,7 @@ namespace EpicGames.Horde.Jobs
 		/// <summary>
 		/// List of nodes covered by default label
 		/// </summary>
-		public List<string> Nodes { get; set; } = new List<string>();
+		public List<string> Nodes { get; set; } = [];
 	}
 
 	/// <summary>
diff --git a/Engine/Source/Programs/Shared/EpicGames.Horde/Jobs/Schedules/ScheduleInterval.cs b/Engine/Source/Programs/Shared/EpicGames.Horde/Jobs/Schedules/ScheduleInterval.cs
--- a/Engine/Source/Programs/Shared/EpicGames.Horde/Jobs/Schedules/ScheduleInterval.cs
+++ b/Engine/Source/Programs/Shared/EpicGames.Horde/Jobs/Schedules/ScheduleInterval.cs
@@ -19,7 +19,7 @@ namespace EpicGames.Horde.Jobs.Schedules
 		/// <summary>
 		/// Parse a string as a time of day
 		/// </summary>
-		[return: NotNullIfNotNull("text")]
+		[return: NotNullIfNotNull(nameof(text))]
 		public static ScheduleInterval? Parse(string? text)
 			=> (text != null) ? new ScheduleInterval((int)IntervalJsonConverter.Parse(text).TotalMinutes) : null;
 	}
diff --git a/Engine/Source/Programs/Shared/EpicGames.Horde/Jobs/Schedules/ScheduleMessages.cs b/Engine/Source/Programs/Shared/EpicGames.Horde/Jobs/Schedules/ScheduleMessages.cs
--- a/Engine/Source/Programs/Shared/EpicGames.Horde/Jobs/Schedules/ScheduleMessages.cs
+++ b/Engine/Source/Programs/Shared/EpicGames.Horde/Jobs/Schedules/ScheduleMessages.cs
@@ -97,10 +97,10 @@ namespace EpicGames.Horde.Jobs.Schedules
 		/// </summary>
 		public GetScheduleResponse()
 		{
-			Patterns = new List<GetSchedulePatternResponse>();
-			TemplateParameters = new Dictionary<string, string>();
-			NextTriggerTimesUTC = new List<DateTime>();
-			ActiveJobs = new List<JobId>();
+			Patterns = [];
+			TemplateParameters = [];
+			NextTriggerTimesUTC = [];
+			ActiveJobs = [];
 		}
 
 		/// <summary>
@@ -120,10 +120,10 @@ namespace EpicGames.Horde.Jobs.Schedules
 			Patterns = schedule.Patterns.ConvertAll(x => new GetSchedulePatternResponse(x));
 			LastTriggerCommitId = schedule.LastTriggerCommitId;
 			LastTriggerTime = schedule.LastTriggerTimeUtc;
-			ActiveJobs = new List<JobId>(schedule.ActiveJobs);
+			ActiveJobs = [.. schedule.ActiveJobs];
 
 			DateTime curTime = schedule.LastTriggerTimeUtc;
-			NextTriggerTimesUTC = new List<DateTime>();
+			NextTriggerTimesUTC = [];
 			for (int i = 0; i < 16; i++)
 			{
 				DateTime? nextTime = schedule.GetNextTriggerTimeUtc(curTime, schedulerTimeZone);
diff --git a/Engine/Source/Programs/Shared/EpicGames.Horde/Jobs/Schedules/ScheduleTimeOfDay.cs b/Engine/Source/Programs/Shared/EpicGames.Horde/Jobs/Schedules/ScheduleTimeOfDay.cs
--- a/Engine/Source/Programs/Shared/EpicGames.Horde/Jobs/Schedules/ScheduleTimeOfDay.cs
+++ b/Engine/Source/Programs/Shared/EpicGames.Horde/Jobs/Schedules/ScheduleTimeOfDay.cs
@@ -19,7 +19,7 @@ namespace EpicGames.Horde.Jobs.Schedules
 		/// <summary>
 		/// Parse a string as a time of day
 		/// </summary>
-		[return: NotNullIfNotNull("text")]
+		[return: NotNullIfNotNull(nameof(text))]
 		public static ScheduleTimeOfDay? Parse(string? text)
 			=> (text != null) ? new ScheduleTimeOfDay((int)TimeOfDayJsonConverter.Parse(text).TotalMinutes) : null;
 	}
diff --git a/Engine/Source/Programs/Shared/EpicGames.Horde/Jobs/Templates/TemplateMessages.cs b/Engine/Source/Programs/Shared/EpicGames.Horde/Jobs/Templates/TemplateMessages.cs
--- a/Engine/Source/Programs/Shared/EpicGames.Horde/Jobs/Templates/TemplateMessages.cs
+++ b/Engine/Source/Programs/Shared/EpicGames.Horde/Jobs/Templates/TemplateMessages.cs
@@ -66,8 +66,8 @@ namespace EpicGames.Horde.Jobs.Templates
 		{
 			Name = null!;
 			AllowPreflights = true;
-			Arguments = new List<string>();
-			Parameters = new List<GetTemplateParameterResponse>();
+			Arguments = [];
+			Parameters = [];
 		}
 
 		/// <summary>
@@ -83,7 +83,7 @@ namespace EpicGames.Horde.Jobs.Templates
 			UpdateIssues = template.UpdateIssues;
 			InitialAgentType = template.InitialAgentType;
 			SubmitNewChange = template.SubmitNewChange;
-			Arguments = new List<string>(template.Arguments);
+			Arguments = [.. template.Arguments];
 			Parameters = template.Parameters.ConvertAll(x => CreateParameterResponse(x));
 		}
 
@@ -295,7 +295,7 @@ namespace EpicGames.Horde.Jobs.Templates
 		public GetTemplateListParameterResponse()
 		{
 			Label = String.Empty;
-			Items = new List<GetTemplateListParameterItemResponse>();
+			Items = [];
 		}
 
 		/// <summary>
diff --git a/Engine/Source/Programs/Shared/EpicGames.Horde/Jobs/TestData/TestDataMessages.cs b/Engine/Source/Programs/Shared/EpicGames.Horde/Jobs/TestData/TestDataMessages.cs
--- a/Engine/Source/Programs/Shared/EpicGames.Horde/Jobs/TestData/TestDataMessages.cs
+++ b/Engine/Source/Programs/Shared/EpicGames.Horde/Jobs/TestData/TestDataMessages.cs
@@ -89,7 +89,7 @@ namespace EpicGames.Horde.Jobs.TestData
 		/// The data stored for this test
 		/// </summary>
 		[Required]
-		public Dictionary<string, object> Data { get; set; } = new Dictionary<string, object>();
+		public Dictionary<string, object> Data { get; set; } = [];
 	}
 
 	/// <summary>
@@ -202,7 +202,7 @@ namespace EpicGames.Horde.Jobs.TestData
 		/// <summary>
 		/// The data stored for this test
 		/// </summary>
-		public Dictionary<string, object> Data { get; set; } = new Dictionary<string, object>();
+		public Dictionary<string, object> Data { get; set; } = [];
 	}
 
 	/// <summary>
@@ -218,17 +218,17 @@ namespace EpicGames.Horde.Jobs.TestData
 		/// <summary>
 		/// The platforms in the environment
 		/// </summary>
-		public List<string> Platforms { get; set; } = new List<string>();
+		public List<string> Platforms { get; set; } = [];
 
 		/// <summary>
 		/// The build configurations being tested
 		/// </summary>
-		public List<string> Configurations { get; set; } = new List<string>();
+		public List<string> Configurations { get; set; } = [];
 
 		/// <summary>
 		/// The build targets being tested
 		/// </summary>
-		public List<string> BuildTargets { get; set; } = new List<string>();
+		public List<string> BuildTargets { get; set; } = [];
 
 		/// <summary>
 		/// The test project name
@@ -274,7 +274,7 @@ namespace EpicGames.Horde.Jobs.TestData
 		/// <summary>
 		/// The meta data the test runs on
 		/// </summary>
-		public List<string> Metadata { get; set; } = new List<string>();
+		public List<string> Metadata { get; set; } = [];
 	}
 
 	/// <summary>
@@ -285,7 +285,7 @@ namespace EpicGames.Horde.Jobs.TestData
 		/// <summary>
 		/// Test ids to get
 		/// </summary>
-		public List<string> TestIds { get; set; } = new List<string>();
+		public List<string> TestIds { get; set; } = [];
 	}
 
 	/// <summary>
@@ -296,7 +296,7 @@ namespace EpicGames.Horde.Jobs.TestData
 		/// <summary>
 		/// Stream ids to get
 		/// </summary>
-		public List<string> StreamIds { get; set; } = new List<string>();
+		public List<string> StreamIds { get; set; } = [];
 	}
 
 	/// <summary>
@@ -376,7 +376,7 @@ namespace EpicGames.Horde.Jobs.TestData
 		/// <summary>
 		/// Stream ids to get
 		/// </summary>
-		public List<string> StreamIds { get; set; } = new List<string>();
+		public List<string> StreamIds { get; set; } = [];
 
 		/// <summary>
 		/// Test ids to get
@@ -397,7 +397,7 @@ namespace EpicGames.Horde.Jobs.TestData
 		/// <summary>
 		/// Stream ids to get
 		/// </summary>
-		public List<string> StreamIds { get; set; } = new List<string>();
+		public List<string> StreamIds { get; set; } = [];
 
 		/// <summary>
 		/// Test ids to get
@@ -413,12 +413,12 @@ namespace EpicGames.Horde.Jobs.TestData
 		/// <summary>
 		/// Stream ids to get
 		/// </summary>
-		public List<string> StreamIds { get; set; } = new List<string>();
+		public List<string> StreamIds { get; set; } = [];
 
 		/// <summary>
 		/// Phase ids to get
 		/// </summary>
-		public List<string> PhaseIds { get; set; } = new List<string>();
+		public List<string> PhaseIds { get; set; } = [];
 	}
 
 	/// <summary>
@@ -495,7 +495,7 @@ namespace EpicGames.Horde.Jobs.TestData
 		/// <summary>
 		/// The test ids
 		/// </summary>
-		public List<string> TestIds { get; set; } = new List<string>();
+		public List<string> TestIds { get; set; } = [];
 	}
 
 	/// <summary>
@@ -653,7 +653,7 @@ namespace EpicGames.Horde.Jobs.TestData
 		/// <summary>
 		/// The meta data the test suite runs on
 		/// </summary>
-		public List<string> Metadata { get; set; } = new List<string>();
+		public List<string> Metadata { get; set; } = [];
 	}
 
 	/// <summary>
@@ -669,17 +669,17 @@ namespace EpicGames.Horde.Jobs.TestData
 		/// <summary>
 		/// Individual tests which run in the stream
 		/// </summary>
-		public List<GetTestResponse> Tests { get; set; } = new List<GetTestResponse>();
+		public List<GetTestResponse> Tests { get; set; } = [];
 
 		/// <summary>
 		/// Test suites that run in the stream
 		/// </summary>
-		public List<GetTestSuiteResponse> TestSuites { get; set; } = new List<GetTestSuiteResponse>();
+		public List<GetTestSuiteResponse> TestSuites { get; set; } = [];
 
 		/// <summary>
 		/// Test suites that run in the stream
 		/// </summary>
-		public List<GetTestMetaResponse> TestMetadata { get; set; } = new List<GetTestMetaResponse>();
+		public List<GetTestMetaResponse> TestMetadata { get; set; } = [];
 	}
 
 	/// <summary>
@@ -731,7 +731,7 @@ namespace EpicGames.Horde.Jobs.TestData
 		/// <summary>
 		/// The test documents for this ref
 		/// </summary>
-		public List<string> TestDataIds { get; set; } = new List<string>();
+		public List<string> TestDataIds { get; set; } = [];
 
 		/// <summary>
 		/// Suite test data
@@ -888,7 +888,7 @@ namespace EpicGames.Horde.Jobs.TestData
 		/// <summary>
 		/// The set of key/value pairs this environment represent
 		/// </summary>
-		public Dictionary<string, string> Entries { get; set; } = new Dictionary<string, string>();
+		public Dictionary<string, string> Entries { get; set; } = [];
 	}
 
 	/// <summary>
@@ -1007,7 +1007,7 @@ namespace EpicGames.Horde.Jobs.TestData
 		/// <summary>
 		/// The phases associated with the test
 		/// </summary>
-		public List<GetPhaseResponse> Phases { get; set; } = new List<GetPhaseResponse>();
+		public List<GetPhaseResponse> Phases { get; set; } = [];
 	}
 
 	/// <summary>
@@ -1023,17 +1023,17 @@ namespace EpicGames.Horde.Jobs.TestData
 		/// <summary>
 		/// Individual tests which run in the stream
 		/// </summary>
-		public List<GetTestNameResponse> Tests { get; set; } = new List<GetTestNameResponse>();
+		public List<GetTestNameResponse> Tests { get; set; } = [];
 
 		/// <summary>
 		/// Known test metadata
 		/// </summary>
-		public List<GetTestMetadataResponse> TestMetadata { get; set; } = new List<GetTestMetadataResponse>();
+		public List<GetTestMetadataResponse> TestMetadata { get; set; } = [];
 
 		/// <summary>
 		/// Known test tags
 		/// </summary>
-		public List<GetTestTagResponse> TestTags { get; set; } = new List<GetTestTagResponse>();
+		public List<GetTestTagResponse> TestTags { get; set; } = [];
 	}
 
 	/// <summary>
@@ -1150,7 +1150,7 @@ namespace EpicGames.Horde.Jobs.TestData
 		/// <summary>
 		/// The phase keys associated with the error fingerprint
 		/// </summary>
-		public List<string> Phases { get; set; } = new List<string>();
+		public List<string> Phases { get; set; } = [];
 	}
 
 	/// <summary>
diff --git a/Engine/Source/Programs/Shared/EpicGames.Horde/Logs/LogChunkNode.cs b/Engine/Source/Programs/Shared/EpicGames.Horde/Logs/LogChunkNode.cs
--- a/Engine/Source/Programs/Shared/EpicGames.Horde/Logs/LogChunkNode.cs
+++ b/Engine/Source/Programs/Shared/EpicGames.Horde/Logs/LogChunkNode.cs
@@ -148,8 +148,7 @@ namespace EpicGames.Horde.Logs
 		/// <returns></returns>
 		public static List<int> FindLineOffsets(ReadOnlySpan<byte> data)
 		{
-			List<int> lineOffsets = new List<int>();
-			lineOffsets.Add(0);
+			List<int> lineOffsets = [0];
 			UpdateLineOffsets(data, 0, lineOffsets);
 			return lineOffsets;
 		}
@@ -281,7 +280,7 @@ namespace EpicGames.Horde.Logs
 		/// <summary>
 		/// Offsets of the start of each line within the data
 		/// </summary>
-		readonly List<int> _lineOffsets = new List<int> { 0 };
+		readonly List<int> _lineOffsets = [0];
 
 		/// <summary>
 		/// Current length of the buffer
@@ -499,7 +498,7 @@ namespace EpicGames.Horde.Logs
 				// Check what the escape code is
 				if (source[1] == 'u')
 				{
-					char[] chars = { (char)((StringUtils.ParseHexByte(source, 2) << 8) | StringUtils.ParseHexByte(source, 4)) };
+					char[] chars = [(char)((StringUtils.ParseHexByte(source, 2) << 8) | StringUtils.ParseHexByte(source, 4))];
 					int encodedLength = Encoding.UTF8.GetBytes(chars.AsSpan(), target);
 					source = source.Slice(6);
 					target = target.Slice(encodedLength);
@@ -555,7 +554,7 @@ namespace EpicGames.Horde.Logs
 		/// Create a <see cref="LogChunkNode"/> object from the current state
 		/// </summary>
 		/// <returns></returns>
-		public LogChunkNode ToLogChunk() => new LogChunkNode(_data.AsMemory(0, _length).ToArray(), _lineOffsets.ToArray());
+		public LogChunkNode ToLogChunk() => new LogChunkNode(_data.AsMemory(0, _length).ToArray(), [.. _lineOffsets]);
 	}
 
 	/// <summary>
@@ -563,7 +562,7 @@ namespace EpicGames.Horde.Logs
 	/// </summary>
 	class LogChunkSequenceBuilder
 	{
-		readonly List<LogChunkNode> _chunks = new List<LogChunkNode>();
+		readonly List<LogChunkNode> _chunks = [];
 		readonly LogChunkBuilder _nextChunkBuilder;
 		int _flushedLength;
 		int _flushedLineCount;
diff --git a/Engine/Source/Programs/Shared/EpicGames.Horde/Logs/LogEventMessages.cs b/Engine/Source/Programs/Shared/EpicGames.Horde/Logs/LogEventMessages.cs
--- a/Engine/Source/Programs/Shared/EpicGames.Horde/Logs/LogEventMessages.cs
+++ b/Engine/Source/Programs/Shared/EpicGames.Horde/Logs/LogEventMessages.cs
@@ -66,6 +66,6 @@ namespace EpicGames.Horde.Logs
 		/// <summary>
 		/// The structured message data for this event
 		/// </summary>
-		public List<JsonElement> Lines { get; set; } = new List<JsonElement>();
+		public List<JsonElement> Lines { get; set; } = [];
 	}
 }
diff --git a/Engine/Source/Programs/Shared/EpicGames.Horde/Logs/LogIndexNode.cs b/Engine/Source/Programs/Shared/EpicGames.Horde/Logs/LogIndexNode.cs
--- a/Engine/Source/Programs/Shared/EpicGames.Horde/Logs/LogIndexNode.cs
+++ b/Engine/Source/Programs/Shared/EpicGames.Horde/Logs/LogIndexNode.cs
@@ -49,7 +49,7 @@ namespace EpicGames.Horde.Logs
 		/// <summary>
 		/// Empty index data
 		/// </summary>
-		public static LogIndexNode Empty { get; } = new LogIndexNode(NgramSet.Empty, 0, Array.Empty<LogChunkRef>());
+		public static LogIndexNode Empty { get; } = new LogIndexNode(NgramSet.Empty, 0, []);
 
 		/// <summary>
 		/// Public accessor for the plain text chunks
@@ -227,18 +227,17 @@ namespace EpicGames.Horde.Logs
 				// search text, along with a set of partial matches (for substrings that are not guaranteed to start on token boundaries). The former
 				// is relatively cheap to compute, and done by stepping through all enumerators for each token in a wave, only returning chunks
 				// which all contain the token. At the very least, this must contain chunks in the index from chunkIdx onwards.
-				List<IEnumerator<int>> enumerators = new List<IEnumerator<int>>();
-				enumerators.Add(Enumerable.Range(chunkIdx, _plainTextChunkRefs.Length - chunkIdx).GetEnumerator());
+				List<IEnumerator<int>> enumerators = [Enumerable.Range(chunkIdx, _plainTextChunkRefs.Length - chunkIdx).GetEnumerator()];
 
 				// The second aspect of the search requires a more expensive search through the trie. This is done by running a set of arbitrary 
 				// delegates to filter the matches returned from the enumerators.
-				List<Predicate<int>> predicates = new List<Predicate<int>>();
+				List<Predicate<int>> predicates = [];
 
 				// If the index has a trie, tokenize the input and generate a list of enumerators and predicates.
 				if (NgramSet != null)
 				{
 					// Find a list of filters for matching chunks
-					HashSet<ulong> tokens = new HashSet<ulong>();
+					HashSet<ulong> tokens = [];
 					for (int tokenPos = 0; tokenPos < text.Length;)
 					{
 						ReadOnlySpan<byte> token = Ngram.GetText(text.Span, tokenPos);
@@ -338,10 +337,10 @@ namespace EpicGames.Horde.Logs
 
 			Lazy<HashSet<int>> lazyChunks = new Lazy<HashSet<int>>(() =>
 			{
-				HashSet<int> union = new HashSet<int>();
+				HashSet<int> union = [];
 				for (int shift = 0; shift < Ngram.MaxBytes; shift++)
 				{
-					HashSet<int> chunks = new HashSet<int>(ChunksContainingNgram(textCopy.AsSpan(), -shift, allowPartialMatch));
+					HashSet<int> chunks = [.. ChunksContainingNgram(textCopy.AsSpan(), -shift, allowPartialMatch)];
 					for (int offset = -shift + Ngram.MaxBytes; offset < textCopy.Length && chunks.Count > 0; offset += Ngram.MaxBytes)
 					{
 						chunks.IntersectWith(ChunksContainingNgram(textCopy.AsSpan(), offset, allowPartialMatch));
diff --git a/Engine/Source/Programs/Shared/EpicGames.Horde/Logs/LogMessages.cs b/Engine/Source/Programs/Shared/EpicGames.Horde/Logs/LogMessages.cs
--- a/Engine/Source/Programs/Shared/EpicGames.Horde/Logs/LogMessages.cs
+++ b/Engine/Source/Programs/Shared/EpicGames.Horde/Logs/LogMessages.cs
@@ -92,7 +92,7 @@ namespace EpicGames.Horde.Logs
 		/// <summary>
 		/// List of line numbers containing the search text
 		/// </summary>
-		public List<int> Lines { get; set; } = new List<int>();
+		public List<int> Lines { get; set; } = [];
 
 		/// <summary>
 		/// Stats for the search
@@ -128,7 +128,7 @@ namespace EpicGames.Horde.Logs
 		/// <summary>
 		/// List of lines received
 		/// </summary>
-		public List<LogLineResponse> Lines { get; set; } = new List<LogLineResponse>();
+		public List<LogLineResponse> Lines { get; set; } = [];
 	}
 
 	/// <summary>
diff --git a/Engine/Source/Programs/Shared/EpicGames.Horde/Logs/LogNode.cs b/Engine/Source/Programs/Shared/EpicGames.Horde/Logs/LogNode.cs
--- a/Engine/Source/Programs/Shared/EpicGames.Horde/Logs/LogNode.cs
+++ b/Engine/Source/Programs/Shared/EpicGames.Horde/Logs/LogNode.cs
@@ -73,7 +73,7 @@ namespace EpicGames.Horde.Logs
 			Format = format;
 			LineCount = lineCount;
 			Length = length;
-			TextChunkRefs = textChunkRefs.ToArray();
+			TextChunkRefs = [.. textChunkRefs];
 			IndexRef = indexRef;
 			Complete = complete;
 		}
@@ -273,17 +273,17 @@ namespace EpicGames.Horde.Logs
 			lock (_lockObject)
 			{
 				_textBuilder.Flush();
-				writeTextChunks = _textBuilder.Chunks.ToArray();
+				writeTextChunks = [.. _textBuilder.Chunks];
 
 				_indexTextBuilder.Flush();
-				writeIndexTextChunks = _indexTextBuilder.Chunks.ToArray();
+				writeIndexTextChunks = [.. _indexTextBuilder.Chunks];
 			}
 
 			// Flush any complete chunks to storage
 			LogIndexNode newIndex = await _index.AppendAsync(writer, writeIndexTextChunks, cancellationToken);
 			IHashedBlobRef<LogIndexNode> newIndexRef = await writer.WriteBlobAsync(newIndex, cancellationToken);
 
-			List<LogChunkRef> newJsonChunkRefs = new List<LogChunkRef>(_root?.TextChunkRefs ?? Array.Empty<LogChunkRef>());
+			List<LogChunkRef> newJsonChunkRefs = [.. _root?.TextChunkRefs ?? []];
 			int lineCount = _root?.LineCount ?? 0;
 			long length = _root?.Length ?? 0;
 			foreach (LogChunkNode writeTextChunk in writeTextChunks)
diff --git a/Engine/Source/Programs/Shared/EpicGames.Horde/Logs/NgramSet.cs b/Engine/Source/Programs/Shared/EpicGames.Horde/Logs/NgramSet.cs
--- a/Engine/Source/Programs/Shared/EpicGames.Horde/Logs/NgramSet.cs
+++ b/Engine/Source/Programs/Shared/EpicGames.Horde/Logs/NgramSet.cs
@@ -211,7 +211,7 @@ namespace EpicGames.Horde.Logs
 		/// <returns>Array of offsets</returns>
 		static int[] CreateChildLookup(ushort[] nodeData)
 		{
-			List<int> childOffsets = new List<int>();
+			List<int> childOffsets = [];
 			if (nodeData.Length > 0)
 			{
 				int nodeCount = 1;
@@ -235,7 +235,7 @@ namespace EpicGames.Horde.Logs
 					nodeCount = nextNodeCount;
 				}
 			}
-			return childOffsets.ToArray();
+			return [.. childOffsets];
 		}
 
 		/// <summary>
diff --git a/Engine/Source/Programs/Shared/EpicGames.Horde/Logs/NgramSetBuilder.cs b/Engine/Source/Programs/Shared/EpicGames.Horde/Logs/NgramSetBuilder.cs
--- a/Engine/Source/Programs/Shared/EpicGames.Horde/Logs/NgramSetBuilder.cs
+++ b/Engine/Source/Programs/Shared/EpicGames.Horde/Logs/NgramSetBuilder.cs
@@ -77,14 +77,13 @@ namespace EpicGames.Horde.Logs
 		/// <returns></returns>
 		public NgramSet ToNgramSet()
 		{
-			List<ushort> values = new List<ushort>();
+			List<ushort> values = [];
 
-			List<Node> nodes = new List<Node>();
-			nodes.Add(_root);
+			List<Node> nodes = [_root];
 
 			for (int bits = 0; bits < (sizeof(ulong) * 8); bits += 4)
 			{
-				List<Node> nextNodes = new List<Node>();
+				List<Node> nextNodes = [];
 				foreach (Node node in nodes)
 				{
 					ushort value = 0;
@@ -104,7 +103,7 @@ namespace EpicGames.Horde.Logs
 				nodes = nextNodes;
 			}
 
-			return new NgramSet(values.ToArray());
+			return new NgramSet([.. values]);
 		}
 	}
 }
diff --git a/Engine/Source/Programs/Shared/EpicGames.Horde/Logs/ServerLogPacketBuilder.cs b/Engine/Source/Programs/Shared/EpicGames.Horde/Logs/ServerLogPacketBuilder.cs
--- a/Engine/Source/Programs/Shared/EpicGames.Horde/Logs/ServerLogPacketBuilder.cs
+++ b/Engine/Source/Programs/Shared/EpicGames.Horde/Logs/ServerLogPacketBuilder.cs
@@ -40,7 +40,7 @@ namespace EpicGames.Horde.Logs
 
 		static readonly Utf8String s_escapedNewline = new Utf8String("\\n");
 
-		readonly List<JsonLogEvent> _logEvents = new List<JsonLogEvent>();
+		readonly List<JsonLogEvent> _logEvents = [];
 		readonly ArrayBufferWriter<byte> _lineWriter;
 		readonly ArrayBufferWriter<byte> _packetWriter;
 		int _packetLength;
@@ -127,7 +127,7 @@ namespace EpicGames.Horde.Logs
 				}
 
 				string escapedLine = escapedLineBuilder.ToString();
-				KeyValuePair<string, object?>[] properties = new[] { new KeyValuePair<string, object?>("Text", escapedLine) };
+				KeyValuePair<string, object?>[] properties = [new KeyValuePair<string, object?>("Text", escapedLine)];
 				LogEvent newLogEvent = new LogEvent(DateTime.UtcNow, LogLevel.Error, default, $"Invalid json log event: {escapedLineBuilder}", "Invalid json log event: {Text}", properties, LogException.FromException(ex));
 				JsonLogEvent newJsonLogEvent = new JsonLogEvent(newLogEvent);
 
@@ -165,11 +165,11 @@ namespace EpicGames.Horde.Logs
 		int WriteEventWithFormat(JsonLogEvent jsonLogEvent, JsonObject obj, string format)
 		{
 			// There is an object containing "common" properties
-			IEnumerable<KeyValuePair<string, object?>> propertyValueList = Enumerable.Empty<KeyValuePair<string, object?>>();
+			IEnumerable<KeyValuePair<string, object?>> propertyValueList = [];
 
 			// Split the format string into lines
 			string[] formatLines = format.Split('\n');
-			List<FormattedLine> lines = new List<FormattedLine>();
+			List<FormattedLine> lines = [];
 
 			// Split all the multi-line properties into separate properties
 			JsonObject? properties = obj["properties"] as JsonObject;
diff --git a/Engine/Source/Programs/Shared/EpicGames.Horde/Logs/ServerLogger.cs b/Engine/Source/Programs/Shared/EpicGames.Horde/Logs/ServerLogger.cs
--- a/Engine/Source/Programs/Shared/EpicGames.Horde/Logs/ServerLogger.cs
+++ b/Engine/Source/Programs/Shared/EpicGames.Horde/Logs/ServerLogger.cs
@@ -147,7 +147,7 @@ namespace EpicGames.Horde.Logs
 
 			// Buffer for events read in a single iteration
 			ServerLogPacketBuilder writer = new ServerLogPacketBuilder();
-			List<RpcCreateLogEventRequest> events = new List<RpcCreateLogEventRequest>();
+			List<RpcCreateLogEventRequest> events = [];
 
 			// Whether we've written the flush command
 			for (; ; )
diff --git a/Engine/Source/Programs/Shared/EpicGames.Horde/Projects/HttpProjectCollection.cs b/Engine/Source/Programs/Shared/EpicGames.Horde/Projects/HttpProjectCollection.cs
--- a/Engine/Source/Programs/Shared/EpicGames.Horde/Projects/HttpProjectCollection.cs
+++ b/Engine/Source/Programs/Shared/EpicGames.Horde/Projects/HttpProjectCollection.cs
@@ -1,6 +1,5 @@
 // Copyright Epic Games, Inc. All Rights Reserved.
 
-using System;
 using System.Collections.Generic;
 using System.Diagnostics;
 using System.Threading;
@@ -25,7 +24,7 @@ namespace EpicGames.Horde.Projects
 			public Project(GetProjectResponse response)
 			{
 				_response = response;
-				_projectStreams = (IReadOnlyList<ProjectStream>?)response.Streams?.ConvertAll(x => new ProjectStream(x)) ?? Array.Empty<ProjectStream>();
+				_projectStreams = (IReadOnlyList<ProjectStream>?)response.Streams?.ConvertAll(x => new ProjectStream(x)) ?? [];
 			}
 		}
 
diff --git a/Engine/Source/Programs/Shared/EpicGames.Horde/Projects/ProjectMessages.cs b/Engine/Source/Programs/Shared/EpicGames.Horde/Projects/ProjectMessages.cs
--- a/Engine/Source/Programs/Shared/EpicGames.Horde/Projects/ProjectMessages.cs
+++ b/Engine/Source/Programs/Shared/EpicGames.Horde/Projects/ProjectMessages.cs
@@ -102,17 +102,17 @@ namespace EpicGames.Horde.Projects
 		/// <summary>
 		/// Patterns for stream names to include
 		/// </summary>
-		public List<string> IncludePatterns { get; init; } = new List<string>();
+		public List<string> IncludePatterns { get; init; } = [];
 
 		/// <summary>
 		/// Patterns for stream names to exclude
 		/// </summary>
-		public List<string> ExcludePatterns { get; init; } = new List<string>();
+		public List<string> ExcludePatterns { get; init; } = [];
 
 		/// <summary>
 		/// Streams to include in this category
 		/// </summary>
-		public List<string> Streams { get; init; } = new List<string>();
+		public List<string> Streams { get; init; } = [];
 
 		/// <summary>
 		/// Constructor
diff --git a/Engine/Source/Programs/Shared/EpicGames.Horde/Server/ServerMessages.cs b/Engine/Source/Programs/Shared/EpicGames.Horde/Server/ServerMessages.cs
--- a/Engine/Source/Programs/Shared/EpicGames.Horde/Server/ServerMessages.cs
+++ b/Engine/Source/Programs/Shared/EpicGames.Horde/Server/ServerMessages.cs
@@ -31,7 +31,7 @@ namespace EpicGames.Horde.Server
 		/// <summary>
 		/// List of plugins
 		/// </summary>
-		public ServerPluginInfoResponse[] Plugins { get; set; } = Array.Empty<ServerPluginInfoResponse>();
+		public ServerPluginInfoResponse[] Plugins { get; set; } = [];
 
 	}
 
@@ -187,7 +187,7 @@ namespace EpicGames.Horde.Server
 		/// <summary>
 		/// List of updates
 		/// </summary>
-		public ServerStatusUpdate[] Updates { get; set; } = Array.Empty<ServerStatusUpdate>();
+		public ServerStatusUpdate[] Updates { get; set; } = [];
 	}
 
 	/// <summary>
@@ -240,7 +240,7 @@ namespace EpicGames.Horde.Server
 		/// <summary>
 		/// List of subsystem statuses
 		/// </summary>
-		public ServerStatusSubsystem[] Statuses { get; set; } = Array.Empty<ServerStatusSubsystem>();
+		public ServerStatusSubsystem[] Statuses { get; set; } = [];
 	}
 
 	/// <summary>
diff --git a/Engine/Source/Programs/Shared/EpicGames.Horde/Storage/Backends/HttpStorageBackend.cs b/Engine/Source/Programs/Shared/EpicGames.Horde/Storage/Backends/HttpStorageBackend.cs
--- a/Engine/Source/Programs/Shared/EpicGames.Horde/Storage/Backends/HttpStorageBackend.cs
+++ b/Engine/Source/Programs/Shared/EpicGames.Horde/Storage/Backends/HttpStorageBackend.cs
@@ -86,7 +86,7 @@ namespace EpicGames.Horde.Storage.Backends
 
 				if (length.HasValue && length.Value == 0)
 				{
-					return new MemoryStream(Array.Empty<byte>());
+					return new MemoryStream([]);
 				}
 
 				using (HttpClient httpClient = _createClient())
@@ -200,7 +200,7 @@ namespace EpicGames.Horde.Storage.Backends
 					? new HttpRequestMessage(HttpMethod.Put, $"{_basePath}/blobs/{locator.Value}")
 					: new HttpRequestMessage(HttpMethod.Post, $"{_basePath}/blobs"))
 					{
-						using MultipartFormDataContent form = new MultipartFormDataContent();
+						using MultipartFormDataContent form = [];
 						if (streamContent != null)
 						{
 							form.Add(streamContent, "file", "filename");
diff --git a/Engine/Source/Programs/Shared/EpicGames.Horde/Storage/Backends/JupiterStorageBackend.cs b/Engine/Source/Programs/Shared/EpicGames.Horde/Storage/Backends/JupiterStorageBackend.cs
--- a/Engine/Source/Programs/Shared/EpicGames.Horde/Storage/Backends/JupiterStorageBackend.cs
+++ b/Engine/Source/Programs/Shared/EpicGames.Horde/Storage/Backends/JupiterStorageBackend.cs
@@ -200,8 +200,8 @@ namespace EpicGames.Horde.Storage.Backends
 			writer.BeginObject();
 			writer.WriteBinaryArray(ObjectFieldNames.Data, payload);
 
-			List<CbBinaryAttachment> binaryAttachments = new List<CbBinaryAttachment>();
-			List<CbObjectAttachment> objectAttachments = new List<CbObjectAttachment>();
+			List<CbBinaryAttachment> binaryAttachments = [];
+			List<CbObjectAttachment> objectAttachments = [];
 			foreach (BlobLocator import in imports)
 			{
 				if (TryGetBinaryAttachment(import.Path, out CbBinaryAttachment binaryAttachment))
