Shared: Enable IDE0028, IDE0031, IDE0034, IDE0054, IDE0062, IDE0071, IDE0074, IDE0075, IDE0280, IDE0300, IDE0301 IDE0305

diff --git a/Engine/Source/Programs/Shared/EpicGames.MsBuild/CsProjBuilder.cs b/Engine/Source/Programs/Shared/EpicGames.MsBuild/CsProjBuilder.cs
--- a/Engine/Source/Programs/Shared/EpicGames.MsBuild/CsProjBuilder.cs
+++ b/Engine/Source/Programs/Shared/EpicGames.MsBuild/CsProjBuilder.cs
@@ -109,7 +109,7 @@ namespace EpicGames.MsBuild
 
 			using ProjectCollection projectCollection = new ProjectCollection(globalProperties);
 			ConcurrentDictionary<FileReference, Project> projects = new();
-			ConcurrentBag<FileReference> platformSpecificProjects = new();
+			ConcurrentBag<FileReference> platformSpecificProjects = [];
 			logger.LogDebug("Setup {Count} projects time: {TimeSeconds:0.00} s", foundProjects.Count, stopwatch.Elapsed.TotalSeconds);
 
 			stopwatch.Restart();
@@ -117,7 +117,7 @@ namespace EpicGames.MsBuild
 			IEnumerable<FileReference> toProcess = outOfDateProjects;
 			while (toProcess.Any())
 			{
-				ConcurrentBag<FileReference> referencedProjects = new();
+				ConcurrentBag<FileReference> referencedProjects = [];
 				Parallel.ForEach(toProcess, (foundProject) =>
 				{
 					string projectPath = Path.GetFullPath(foundProject.FullName);
@@ -476,7 +476,7 @@ namespace EpicGames.MsBuild
 			}
 
 			// Acquire a mutex to prevent running builds from different processes concurrently
-			Task<IDisposable> mutexTask = SingleInstanceMutex.AcquireAsync($"Global\\CsProj_{Sha1Hash.Compute(Encoding.Default.GetBytes(hook.EngineDirectory.FullName)).ToString()}", cancellationToken);
+			Task<IDisposable> mutexTask = SingleInstanceMutex.AcquireAsync($"Global\\CsProj_{Sha1Hash.Compute(Encoding.Default.GetBytes(hook.EngineDirectory.FullName))}", cancellationToken);
 
 			Task delayTask = Task.Delay(TimeSpan.FromSeconds(1.0), cancellationToken);
 			if (Task.WhenAny(mutexTask, delayTask) == delayTask)
diff --git a/Engine/Source/Programs/Shared/EpicGames.OIDC/FakeOidcTokenManager.cs b/Engine/Source/Programs/Shared/EpicGames.OIDC/FakeOidcTokenManager.cs
--- a/Engine/Source/Programs/Shared/EpicGames.OIDC/FakeOidcTokenManager.cs
+++ b/Engine/Source/Programs/Shared/EpicGames.OIDC/FakeOidcTokenManager.cs
@@ -130,7 +130,7 @@ public class FakeOidcTokenManager : IOidcTokenManager
 	private void RefreshAccessToken()
 	{
 		AccessTokenExpiry = _utcNow() + AccessTokenTtl;
-		AccessToken = $"{AccessTokenPrefix}-{_accessCounter++}-{AccessTokenExpiry.ToUnixTimeMilliseconds()}-{AccessTokenExpiry.ToString()}";
+		AccessToken = $"{AccessTokenPrefix}-{_accessCounter++}-{AccessTokenExpiry.ToUnixTimeMilliseconds()}-{AccessTokenExpiry}";
 	}
 
 	private IOidcTokenInfo GetOidcTokenInfo()
diff --git a/Engine/Source/Programs/Shared/EpicGames.OIDC/InMemoryTokenStore.cs b/Engine/Source/Programs/Shared/EpicGames.OIDC/InMemoryTokenStore.cs
--- a/Engine/Source/Programs/Shared/EpicGames.OIDC/InMemoryTokenStore.cs
+++ b/Engine/Source/Programs/Shared/EpicGames.OIDC/InMemoryTokenStore.cs
@@ -11,7 +11,7 @@ namespace EpicGames.OIDC;
 // ReSharper disable once ClassNeverInstantiated.Global
 public class InMemoryTokenStore : ITokenStore
 {
-	private readonly Dictionary<string, string> _providerToRefreshToken = new();
+	private readonly Dictionary<string, string> _providerToRefreshToken = [];
 
 	/// <inheritdoc/>
 	public virtual bool TryGetRefreshToken(string oidcProvider, out string refreshToken)
diff --git a/Engine/Source/Programs/Shared/EpicGames.OIDC/MacOSTokenStore.cs b/Engine/Source/Programs/Shared/EpicGames.OIDC/MacOSTokenStore.cs
--- a/Engine/Source/Programs/Shared/EpicGames.OIDC/MacOSTokenStore.cs
+++ b/Engine/Source/Programs/Shared/EpicGames.OIDC/MacOSTokenStore.cs
@@ -88,7 +88,7 @@ namespace EpicGames.OIDC
 
 		public static void AddKeychainItem(string dataToEncrypt, string itemIdentifier)
 		{
-			RunKeychainHandler("add-generic-password", new List<string> {
+			RunKeychainHandler("add-generic-password", [
 				"-a",
 				itemIdentifier,
 				"-s",
@@ -96,18 +96,18 @@ namespace EpicGames.OIDC
 				"-U",
 				"-w",
 				dataToEncrypt
-			});
+			]);
 		}
 
 		public static string? GetKeychainItem(string itemIdentifier)
 		{
-			string? output = RunKeychainHandler("find-generic-password", new List<string> {
+			string? output = RunKeychainHandler("find-generic-password", [
 				"-a",
 				itemIdentifier,
 				"-s",
 				OidcTokenIdentifier,
 				"-w"
-			});
+			]);
 
 			return output;
 		}
diff --git a/Engine/Source/Programs/Shared/EpicGames.OIDC/OidcTokenManager.cs b/Engine/Source/Programs/Shared/EpicGames.OIDC/OidcTokenManager.cs
--- a/Engine/Source/Programs/Shared/EpicGames.OIDC/OidcTokenManager.cs
+++ b/Engine/Source/Programs/Shared/EpicGames.OIDC/OidcTokenManager.cs
@@ -64,8 +64,8 @@ namespace EpicGames.OIDC
 		public OidcTokenManager(IOptionsMonitor<OidcTokenOptions> settings, ITokenStore tokenStore, IOidcTokenClientFactory tokenClientFactory, List<string>? allowedProviders = null)
 			: this(tokenStore)
 		{
-			Dictionary<string, string> refreshTokens = new Dictionary<string, string>();
-			List<string> providerKeys = settings.CurrentValue.Providers.Keys.Where(k => allowedProviders == null || allowedProviders.Contains(k)).ToList();
+			Dictionary<string, string> refreshTokens = [];
+			List<string> providerKeys = [.. settings.CurrentValue.Providers.Keys.Where(k => allowedProviders == null || allowedProviders.Contains(k))];
 			foreach (string oidcProvider in providerKeys)
 			{
 				if (tokenStore.TryGetRefreshToken(oidcProvider, out string? refreshToken))
@@ -100,12 +100,12 @@ namespace EpicGames.OIDC
 		private OidcTokenManager(IConfiguration providerConfiguration, ITokenStore tokenStore, List<string>? allowedProviders = null)
 			: this(tokenStore)
 		{
-			Dictionary<string, string> refreshTokens = new Dictionary<string, string>();
+			Dictionary<string, string> refreshTokens = [];
 
 			OidcTokenOptions options = new OidcTokenOptions();
 			providerConfiguration.Bind(options);
 
-			List<string> providerKeys = options.Providers.Keys.Where(k => allowedProviders == null || allowedProviders.Contains(k)).ToList();
+			List<string> providerKeys = [.. options.Providers.Keys.Where(k => allowedProviders == null || allowedProviders.Contains(k))];
 
 			foreach (string oidcProvider in providerKeys)
 			{
@@ -402,7 +402,7 @@ namespace EpicGames.OIDC
 			_authorityUri = providerInfo.ServerUri;
 			_clientId = providerInfo.ClientId;
 
-			List<Uri> possibleRedirectUris = new List<Uri>();
+			List<Uri> possibleRedirectUris = [];
 			if (providerInfo.RedirectUri != null)
 			{
 				possibleRedirectUris.Add(providerInfo.RedirectUri);
@@ -884,7 +884,7 @@ namespace EpicGames.OIDC
 
 	public class OidcTokenOptions
 	{
-		public Dictionary<string, ProviderInfo> Providers { get; set; } = new Dictionary<string, ProviderInfo>();
+		public Dictionary<string, ProviderInfo> Providers { get; set; } = [];
 
 		public TimeSpan LoginTimeout { get; set; } = TimeSpan.FromMinutes(20);
 
diff --git a/Engine/Source/Programs/Shared/EpicGames.OIDC/ProviderConfigurationFactory.cs b/Engine/Source/Programs/Shared/EpicGames.OIDC/ProviderConfigurationFactory.cs
--- a/Engine/Source/Programs/Shared/EpicGames.OIDC/ProviderConfigurationFactory.cs
+++ b/Engine/Source/Programs/Shared/EpicGames.OIDC/ProviderConfigurationFactory.cs
@@ -19,12 +19,12 @@ namespace EpicGames.OIDC
 	{
 		const string ConfigFileName = "oidc-configuration.json";
 
-		public static IReadOnlyList<string> ConfigPaths { get; } = new string[]
-		{
+		public static IReadOnlyList<string> ConfigPaths { get; } =
+		[
 			$"Programs/OidcToken/{ConfigFileName}",
 			$"Restricted/NoRedist/Programs/OidcToken/{ConfigFileName}",
 			$"Restricted/NotForLicensees/Programs/OidcToken/{ConfigFileName}"
-		};
+		];
 
 		public static IConfiguration ReadConfiguration(DirectoryInfo engineDir, DirectoryInfo? gameDir)
 		{
diff --git a/Engine/Source/Programs/Shared/EpicGames.OIDC/TokenStore.cs b/Engine/Source/Programs/Shared/EpicGames.OIDC/TokenStore.cs
--- a/Engine/Source/Programs/Shared/EpicGames.OIDC/TokenStore.cs
+++ b/Engine/Source/Programs/Shared/EpicGames.OIDC/TokenStore.cs
@@ -19,7 +19,7 @@ namespace EpicGames.OIDC
 
 	internal class TokenStoreState
 	{
-		public Dictionary<string, string> Providers { get; set; } = new Dictionary<string, string>();
+		public Dictionary<string, string> Providers { get; set; } = [];
 		public string Key { get; set; }
 
 		[JsonConstructor]
@@ -46,7 +46,7 @@ namespace EpicGames.OIDC
 	public class FilesystemTokenStore : ITokenStore, IDisposable
 	{
 		private readonly Dictionary<string, byte[]> _providerToRefreshToken;
-		private readonly List<string> _dirtyProviders = new List<string>();
+		private readonly List<string> _dirtyProviders = [];
 		private readonly SymmetricAlgorithm _crypt = Aes.Create();
 
 		public FilesystemTokenStore()
@@ -65,7 +65,7 @@ namespace EpicGames.OIDC
 			if (!fi.Exists)
 			{
 				// if we have no store on disk then we just initialize it to empty
-				return new Dictionary<string, byte[]>();
+				return [];
 			}
 
 			using FileStream fs = fi.OpenRead();
@@ -82,9 +82,9 @@ namespace EpicGames.OIDC
 			if (state == null)
 			{
 				// if we fail to deserialize the state just drop it, will mean users will need to login again
-				return new Dictionary<string, byte[]>();
+				return [];
 			}
-			Dictionary<string, byte[]> providers = new Dictionary<string, byte[]>();
+			Dictionary<string, byte[]> providers = [];
 
 			foreach ((string key, string value) in state.Providers)
 			{
@@ -170,11 +170,11 @@ namespace EpicGames.OIDC
 
 		static byte[] GetSeed()
 		{
-			return new byte[]
-			{
+			return
+			[
 				0x1e, 0x72, 0x5e, 0xe7, 0x08, 0x9e, 0x29, 0x5e, 0xcb, 0xbe, 0x1b, 0xdf, 0x0e, 0xf9, 0x4a, 0x30, 0xd1,
 				0xa9, 0x9b, 0xa2, 0xee, 0x58, 0xc4, 0x8e
-			};
+			];
 		}
 
 		protected virtual void Dispose(bool disposing)
diff --git a/Engine/Source/Programs/Shared/EpicGames.OIDC/WindowsTokenStore.cs b/Engine/Source/Programs/Shared/EpicGames.OIDC/WindowsTokenStore.cs
--- a/Engine/Source/Programs/Shared/EpicGames.OIDC/WindowsTokenStore.cs
+++ b/Engine/Source/Programs/Shared/EpicGames.OIDC/WindowsTokenStore.cs
@@ -23,7 +23,7 @@ namespace EpicGames.OIDC
 
 	internal class WindowsTokenStoreState
 	{
-		public Dictionary<string, string> Providers { get; set; } = new Dictionary<string, string>();
+		public Dictionary<string, string> Providers { get; set; } = [];
 
 		[JsonConstructor]
 		public WindowsTokenStoreState(Dictionary<string, string> providers)
@@ -44,8 +44,8 @@ namespace EpicGames.OIDC
 	{
 		private readonly ILogger<WindowsTokenStore>? _logger = null;
 
-		private readonly Dictionary<string, byte[]> _providerToRefreshToken = new Dictionary<string, byte[]>();
-		private readonly List<string> _dirtyProviders = new List<string>();
+		private readonly Dictionary<string, byte[]> _providerToRefreshToken = [];
+		private readonly List<string> _dirtyProviders = [];
 
 		public WindowsTokenStore()
 		{
@@ -72,7 +72,7 @@ namespace EpicGames.OIDC
 			{
 				_logger?.LogDebug("No existing token store found at {Path}. Assuming empty store.", fi.FullName);
 				// if we have no store on disk then we just initialize it to empty
-				return new Dictionary<string, byte[]>();
+				return [];
 			}
 
 			using FileStream fs = fi.Open(FileMode.Open, FileAccess.Read, FileShare.Read);
@@ -92,10 +92,10 @@ namespace EpicGames.OIDC
 			{
 				_logger?.LogDebug("Failed to deserialize state. Dropping the existing state.");
 				// if we fail to deserialize the state just drop it, will mean users will need to login again
-				return new Dictionary<string, byte[]>();
+				return [];
 			}
 
-			Dictionary<string, byte[]> providers = new Dictionary<string, byte[]>();
+			Dictionary<string, byte[]> providers = [];
 
 			foreach ((string key, string value) in state.Providers)
 			{
diff --git a/Engine/Source/Programs/Shared/EpicGames.Perforce.Fixture/PerforceFixture.cs b/Engine/Source/Programs/Shared/EpicGames.Perforce.Fixture/PerforceFixture.cs
--- a/Engine/Source/Programs/Shared/EpicGames.Perforce.Fixture/PerforceFixture.cs
+++ b/Engine/Source/Programs/Shared/EpicGames.Perforce.Fixture/PerforceFixture.cs
@@ -80,8 +80,8 @@ public class ChangelistFixture
 	{
 		if (!expected.SetEquals(actual))
 		{
-			List<(string clientFile, long size, string digest)> expectedList = new(expected);
-			List<(string clientFile, long size, string digest)> actualList = new(actual);
+			List<(string clientFile, long size, string digest)> expectedList = [.. expected];
+			List<(string clientFile, long size, string digest)> actualList = [.. actual];
 
 			expectedList.Sort();
 			actualList.Sort();
@@ -111,8 +111,8 @@ public class ChangelistFixture
 	public async Task AssertHaveTableAsync(IPerforceConnection perforce, bool useHaveTable = true)
 	{
 		List<HaveRecord> haveRecords = await perforce.HaveAsync(new FileSpecList(), CancellationToken.None).ToListAsync();
-		HashSet<(string depotFile, int rev)> actual = new(haveRecords.Select((x) => (x.DepotFile, x.HaveRev)));
-		HashSet<(string depotFile, int rev)> expected = new(StreamFiles.Select((x) => (x.DepotFile, x.Revision)));
+		HashSet<(string depotFile, int rev)> actual = [.. haveRecords.Select((x) => (x.DepotFile, x.HaveRev))];
+		HashSet<(string depotFile, int rev)> expected = [.. StreamFiles.Select((x) => (x.DepotFile, x.Revision))];
 
 		if (!useHaveTable)
 		{
@@ -121,8 +121,8 @@ public class ChangelistFixture
 
 		if (!expected.SetEquals(actual))
 		{
-			List<(string depotFile, int rev)> expectedList = new(expected);
-			List<(string depotFile, int rev)> actualList = new(actual);
+			List<(string depotFile, int rev)> expectedList = [.. expected];
+			List<(string depotFile, int rev)> actualList = [.. actual];
 
 			expectedList.Sort();
 			actualList.Sort();
@@ -146,7 +146,7 @@ public class ChangelistFixture
 
 	public (FileSet localFiles, FileSet streamFiles) GetFileSets(string clientRoot)
 	{
-		FileSet streamFiles = new(StreamFiles.Select(x => (x.ClientFile, x.Size, x.Digest)));
+		FileSet streamFiles = [.. StreamFiles.Select(x => (x.ClientFile, x.Size, x.Digest))];
 		FileSet localFiles = GetLocalFileSet(clientRoot);
 		return (localFiles, streamFiles);
 	}
@@ -155,7 +155,7 @@ public class ChangelistFixture
 	{
 		EnumerationOptions options = new() { RecurseSubdirectories = true };
 
-		FileSet localFiles = new(Directory.EnumerateFiles(clientRoot, "*", options)
+		FileSet localFiles = [.. Directory.EnumerateFiles(clientRoot, "*", options)
 			.Select(x => Path.GetRelativePath(clientRoot, x))
 			.Select(x => x.Replace("\\", "/", StringComparison.Ordinal))
 			.Select(clientFile =>
@@ -168,7 +168,7 @@ public class ChangelistFixture
 				content = content.Replace("\r\n", "\n", StringComparison.Ordinal);
 				long size = content.Length;
 				return (clientFile, size, PerforceFixture.CalcMd5(content).ToUpperInvariant());
-			}));
+			})];
 
 		return localFiles;
 	}
@@ -176,11 +176,10 @@ public class ChangelistFixture
 	public (List<string> localFiles, List<string> streamFiles) GetFiles(string clientRoot)
 	{
 		EnumerationOptions options = new() { RecurseSubdirectories = true };
-		List<string> localFiles = Directory.EnumerateFiles(clientRoot, "*", options)
+		List<string> localFiles = [.. Directory.EnumerateFiles(clientRoot, "*", options)
 			.Select(x => Path.GetRelativePath(clientRoot, x))
-			.Select(x => x.Replace("\\", "/", StringComparison.Ordinal))
-			.ToList();
-		List<string> streamFiles = new(StreamFiles.Select(x => x.ClientFile));
+			.Select(x => x.Replace("\\", "/", StringComparison.Ordinal))];
+		List<string> streamFiles = [.. StreamFiles.Select(x => x.ClientFile)];
 
 		localFiles.Sort();
 		streamFiles.Sort();
@@ -227,117 +226,103 @@ public class StreamFixture
 public class PerforceFixture
 {
 	public StreamFixture StreamFooMain { get; } = new("//Foo/Main", 9,
-		new List<ChangelistFixture>
-		{
-			new(0, ChangelistFixture.Placeholder, new List<DepotFileFixture>()),
-			new(1, ChangelistFixture.Placeholder, new List<DepotFileFixture>()),
+		[
+			new(0, ChangelistFixture.Placeholder, []),
+			new(1, ChangelistFixture.Placeholder, []),
 			new(2, "Initial import",
-				new List<DepotFileFixture>()
-				{
+				[
 					new("//Foo/Main/main.cpp", "main.cpp", "This is change main.cpp #1\n", 1),
 					new("//Foo/Main/main.h", "main.h", "This is change main.h #1\n", 1),
 					new("//Foo/Main/common.h", "common.h", "This is change common.h #1\n", 1),
 					new("//Foo/Main/unused.cpp", "unused.cpp", "This is change unused.cpp #1\n", 1),
 					new("//Foo/Main/Data/data.txt", "Data/data.txt", "This is change data.txt #1\n", 1),
-				}),
+				]),
 			new(3, "Improvement to main.cpp",
-				new List<DepotFileFixture>()
-				{
+				[
 					new("//Foo/Main/main.cpp", "main.cpp", "This is change main.cpp #2\n", 2),
 					new("//Foo/Main/main.h", "main.h", "This is change main.h #1\n", 1),
 					new("//Foo/Main/common.h", "common.h", "This is change common.h #1\n", 1),
 					new("//Foo/Main/unused.cpp", "unused.cpp", "This is change unused.cpp #1\n", 1),
 					new("//Foo/Main/Data/data.txt", "Data/data.txt", "This is change data.txt #1\n", 1),
-				}),
+				]),
 			new(4, "Delete an unused file",
-				new List<DepotFileFixture>()
-				{
+				[
 					new("//Foo/Main/main.cpp", "main.cpp", "This is change main.cpp #2\n", 2),
 					new("//Foo/Main/main.h", "main.h", "This is change main.h #1\n", 1),
 					new("//Foo/Main/common.h", "common.h", "This is change common.h #1\n", 1),
 					new("//Foo/Main/Data/data.txt", "Data/data.txt", "This is change data.txt #1\n", 1),
-				}),
+				]),
 			new(5, "Rename common.h to shared.h",
-				new List<DepotFileFixture>()
-				{
+				[
 					new("//Foo/Main/main.cpp", "main.cpp", "This is change main.cpp #2\n", 2),
 					new("//Foo/Main/main.h", "main.h", "This is change main.h #1\n", 1),
 					new("//Foo/Main/shared.h", "shared.h", "This is change common.h #1\n", 1),
 					new("//Foo/Main/Data/data.txt", "Data/data.txt", "This is change data.txt #1\n", 1),
-				}),
+				]),
 			new(6, "Some updates to main",
-				new List<DepotFileFixture>()
-				{
+				[
 					new("//Foo/Main/main.cpp", "main.cpp", "This is change main.cpp #3\n", 3),
 					new("//Foo/Main/main.h", "main.h", "This is change main.h #2\n", 2),
 					new("//Foo/Main/shared.h", "shared.h", "This is change common.h #1\n", 1),
 					new("//Foo/Main/Data/data.txt", "Data/data.txt", "This is change data.txt #1\n", 1),
-				}),
+				]),
 			new(7, "Add more data",
-				new List<DepotFileFixture>()
-				{
+				[
 					new("//Foo/Main/main.cpp", "main.cpp", "This is change main.cpp #3\n", 3),
 					new("//Foo/Main/main.h", "main.h", "This is change main.h #2\n", 2),
 					new("//Foo/Main/shared.h", "shared.h", "This is change common.h #1\n", 1),
 					new("//Foo/Main/Data/data.txt", "Data/data.txt", "This is change data.txt #1\n", 1),
 					new("//Foo/Main/Data/moredata.txt", "Data/moredata.txt", "This is change moredata.txt #1\n", 1),
-				}),
+				]),
 			new(8, "A shelved CL", // Assumes base CL is 7
-				new List<DepotFileFixture>()
-				{
+				[
 					new("//Foo/Main/main.h", "main.h", "This is change main.h #3\n", 3),
 					new("//Foo/Main/shared.h", "shared.h", "This is change common.h #1\n", 1),
 					new("//Foo/Main/shelved.cpp", "shelved.cpp", "This is change shelved.cpp #1\n", 1),
 					new("//Foo/Main/unused.cpp", "unused.cpp", "This is change unused.cpp #2\n", 1),
 					new("//Foo/Main/Data/data.txt", "Data/data.txt", "This is change data.txt #1\n", 1),
 					new("//Foo/Main/Data/moredata.txt", "Data/moredata.txt", "This is change moredata.txt #1\n", 1),
-				}, true),
+				], true),
 			new(9, "Did more changes to main",
-				new List<DepotFileFixture>()
-				{
+				[
 					new("//Foo/Main/main.cpp", "main.cpp", "This is change main.cpp #4\n", 3),
 					new("//Foo/Main/main.h", "main.h", "This is change main.h #4\n", 2),
 					new("//Foo/Main/shared.h", "shared.h", "This is change common.h #1\n", 1),
 					new("//Foo/Main/Data/data.txt", "Data/data.txt", "This is change data.txt #1\n", 1),
 					new("//Foo/Main/Data/moredata.txt", "Data/moredata.txt", "This is change moredata.txt #1\n", 1),
-				})
-		});
+				])
+		]);
 
 	public StreamFixture StreamFooRelease { get; } = new("//Foo/Release", 11,
-		new List<ChangelistFixture>
-		{
+		[
 			new(11, "Initial import",
-				new List<DepotFileFixture>()
-				{
+				[
 					new("//Foo/Main/main.cpp", "main.cpp", "This is change main.cpp #2\n", 1),
 					new("//Foo/Main/main.h", "main.h", "This is change main.h #1\n", 1),
 					new("//Foo/Main/shared.h", "shared.h", "This is change common.h #1\n", 1),
 					// data.txt is excluded in view
-				}),
+				]),
 			new(12, "Initial import",
-				new List<DepotFileFixture>()
-				{
+				[
 					new("//Foo/Main/main.cpp", "main.cpp", "This is change main.cpp #2 in //Foo/Release\n", 1),
 					new("//Foo/Main/main.h", "main.h", "This is change main.h #1\n", 1),
 					new("//Foo/Main/shared.h", "shared.h", "This is change common.h #1\n", 1),
 					new("//Foo/Main/release.cpp", "release.cpp", "This is change release.cpp #1 in //Foo/Release\n", 1),
 					// data.txt is excluded in view
-				}),
-		});
+				]),
+		]);
 
 	public StreamFixture StreamFooSparseDev { get; } = new("//Foo/Main-SparseDev", 13,
-		new List<ChangelistFixture>
-		{
+		[
 			new(13, "Stream //Foo/Main-SparseDev saved.", // TODO: add action
-				new List<DepotFileFixture>()
-				{
+				[
 					new("//Foo/Main/main.cpp", "main.cpp", "This is change main.cpp #4\n", 3),
 					new("//Foo/Main/main.h", "main.h", "This is change main.h #4\n", 2),
 					new("//Foo/Main/shared.h", "shared.h", "This is change common.h #1\n", 1),
 					new("//Foo/Main/Data/data.txt", "Data/data.txt", "This is change data.txt #1\n", 1),
 					new("//Foo/Main/Data/moredata.txt", "Data/moredata.txt", "This is change moredata.txt #1\n", 1),
-				})
-		});
+				])
+		]);
 
 	public static string CalcMd5(string content)
 	{
@@ -354,13 +339,13 @@ public class PerforceFixture
 	{
 		HashSet<(string cacheId, string digest, long fileSize)> cachedFiles = GetCachedFiles(cacheDir);
 
-		HashSet<string> actual = new(cachedFiles.Select(x => x.digest));
-		HashSet<string> expected = new(digests);
+		HashSet<string> actual = [.. cachedFiles.Select(x => x.digest)];
+		HashSet<string> expected = [.. digests];
 
 		if (!expected.SetEquals(actual))
 		{
-			List<string> expectedList = new(expected);
-			List<string> actualList = new(actual);
+			List<string> expectedList = [.. expected];
+			List<string> actualList = [.. actual];
 
 			expectedList.Sort();
 			actualList.Sort();
@@ -386,7 +371,7 @@ public class PerforceFixture
 	{
 		EnumerationOptions options = new() { RecurseSubdirectories = true };
 
-		return new(Directory.EnumerateFiles(cacheDir, "*", options)
+		return [.. Directory.EnumerateFiles(cacheDir, "*", options)
 			.Select(x => Path.GetRelativePath(cacheDir, x))
 			.Select(clientFile =>
 			{
@@ -404,6 +389,6 @@ public class PerforceFixture
 				// See cache ID implementation in ManagedWorkspace
 				Assert.AreEqual(fileName, digest[..16]);
 				return (fileName, digest, size);
-			}));
+			})];
 	}
 }
\ No newline at end of file
diff --git a/Engine/Source/Programs/Shared/EpicGames.Perforce.Fixture/PerforceSnapshot.cs b/Engine/Source/Programs/Shared/EpicGames.Perforce.Fixture/PerforceSnapshot.cs
--- a/Engine/Source/Programs/Shared/EpicGames.Perforce.Fixture/PerforceSnapshot.cs
+++ b/Engine/Source/Programs/Shared/EpicGames.Perforce.Fixture/PerforceSnapshot.cs
@@ -76,13 +76,13 @@ public record WorkspaceSnapshot(int Change, int ShelvedChange, HashSet<Workspace
 	{
 		if (!expected.SetEquals(actual))
 		{
-			List<WorkspaceFile> expectedList = new(expected);
-			List<WorkspaceFile> actualList = new(actual);
+			List<WorkspaceFile> expectedList = [.. expected];
+			List<WorkspaceFile> actualList = [.. actual];
 
 			expectedList.Sort((a, b) => String.Compare(a.Path, b.Path, StringComparison.Ordinal));
 			actualList.Sort((a, b) => String.Compare(a.Path, b.Path, StringComparison.Ordinal));
 
-			string ReplaceLineEndings(string s) => s.Replace("\n", "\\n").Replace("\r", "\\r");
+			static string ReplaceLineEndings(string s) => s.Replace("\n", "\\n").Replace("\r", "\\r");
 
 			Console.WriteLine("Expected " + new string('-', 90));
 			foreach (WorkspaceFile wf in expectedList)
@@ -148,7 +148,7 @@ public class PerforceSnapshot(PerforceConnection perforce)
 
 	public async Task TakeSnapshotsAsync(string snapshotJsonFile)
 	{
-		Dictionary<string, StreamSnapshot> streamSnapshots = new();
+		Dictionary<string, StreamSnapshot> streamSnapshots = [];
 		async Task AddStream(IReadOnlyList<WorkspaceState> states)
 		{
 			string streamName = states[0].Stream;
@@ -191,7 +191,7 @@ public class PerforceSnapshot(PerforceConnection perforce)
 				UnshelveOptions.ForceOverwrite, Array.Empty<string>(), CancellationToken.None);
 		}
 
-		List<string> localFiles = Directory.EnumerateFiles(clientRoot, "*", SearchOption.AllDirectories).ToList();
+		List<string> localFiles = [.. Directory.EnumerateFiles(clientRoot, "*", SearchOption.AllDirectories)];
 		foreach (string absPath in localFiles)
 		{
 			string localFile = Path.GetRelativePath(clientRoot, absPath).Replace('\\', '/');
@@ -206,7 +206,7 @@ public class PerforceSnapshot(PerforceConnection perforce)
 	{
 		EnumerationOptions options = new() { RecurseSubdirectories = true };
 
-		HashSet<WorkspaceFile> localFiles = new(Directory.EnumerateFiles(clientRoot, "*", options)
+		HashSet<WorkspaceFile> localFiles = [.. Directory.EnumerateFiles(clientRoot, "*", options)
 			.Select(x => Path.GetRelativePath(clientRoot, x))
 			.Select(x => x.Replace("\\", "/", StringComparison.Ordinal))
 			.Select(clientFile =>
@@ -220,7 +220,7 @@ public class PerforceSnapshot(PerforceConnection perforce)
 				long size = content.Length;
 				// return WorkspaceFile.C(clientFile, size, PerforceFixture.CalcMd5(content).ToUpperInvariant());
 				return WorkspaceFile.Create(clientFile, content);
-			}));
+			})];
 
 		return localFiles;
 	}
diff --git a/Engine/Source/Programs/Shared/EpicGames.Perforce.Managed.Tests/ManagedWorkspaceTest.cs b/Engine/Source/Programs/Shared/EpicGames.Perforce.Managed.Tests/ManagedWorkspaceTest.cs
--- a/Engine/Source/Programs/Shared/EpicGames.Perforce.Managed.Tests/ManagedWorkspaceTest.cs
+++ b/Engine/Source/Programs/Shared/EpicGames.Perforce.Managed.Tests/ManagedWorkspaceTest.cs
@@ -176,12 +176,12 @@ public class ManagedWorkspaceTest : BasePerforceFixtureTest
 		ManagedWorkspace ws = await CreateManagedWorkspaceAsync(useHaveTable);
 		ChangelistFixture cl = Stream.GetChangelist(6);
 
-		List<string> view = new() { "-/Data/...", "-/shared.h", };
+		List<string> view = ["-/Data/...", "-/shared.h",];
 		await ws.SyncAsync(PerforceConnection, StreamName, cl.Number, view, true, false, null, CancellationToken.None);
 
-		List<DepotFileFixture> filtered = cl.StreamFiles
+		List<DepotFileFixture> filtered = [.. cl.StreamFiles
 			.Where(x => !x.DepotFile.Contains("shared.h", StringComparison.Ordinal))
-			.Where(x => !x.DepotFile.Contains("Data/", StringComparison.Ordinal)).ToList();
+			.Where(x => !x.DepotFile.Contains("Data/", StringComparison.Ordinal))];
 
 		ChangelistFixture clViewApplied = new(cl.Number, cl.Description, filtered, cl.IsShelved);
 		clViewApplied.AssertDepotFiles(SyncDir);
@@ -196,11 +196,10 @@ public class ManagedWorkspaceTest : BasePerforceFixtureTest
 		ManagedWorkspace ws = await CreateManagedWorkspaceAsync(useHaveTable);
 		ChangelistFixture cl = Stream.GetChangelist(6);
 
-		List<string> view = new() { "/...", "-/Data/..." };
+		List<string> view = ["/...", "-/Data/..."];
 		await ws.SyncAsync(PerforceConnection, StreamName, cl.Number, view, true, false, null, CancellationToken.None);
 
-		List<DepotFileFixture> filtered = cl.StreamFiles
-			.Where(x => !x.DepotFile.Contains("Data/", StringComparison.Ordinal)).ToList();
+		List<DepotFileFixture> filtered = [.. cl.StreamFiles.Where(x => !x.DepotFile.Contains("Data/", StringComparison.Ordinal))];
 
 		ChangelistFixture clViewApplied = new(cl.Number, cl.Description, filtered, cl.IsShelved);
 		clViewApplied.AssertDepotFiles(SyncDir);
@@ -213,7 +212,7 @@ public class ManagedWorkspaceTest : BasePerforceFixtureTest
 	public async Task PopulateAsync(bool useHaveTable)
 	{
 		ManagedWorkspace ws = await CreateManagedWorkspaceAsync(useHaveTable);
-		List<PopulateRequest> populateRequests = new() { new PopulateRequest(PerforceConnection, StreamName, new List<string>()) };
+		List<PopulateRequest> populateRequests = [new PopulateRequest(PerforceConnection, StreamName, [])];
 		await ws.PopulateAsync(populateRequests, false, CancellationToken.None);
 		Stream.LatestChangelist.AssertDepotFiles(SyncDir);
 		await Stream.LatestChangelist.AssertHaveTableAsync(PerforceConnection, useHaveTable);
@@ -225,13 +224,13 @@ public class ManagedWorkspaceTest : BasePerforceFixtureTest
 	public async Task PopulateWithViewAsync(bool useHaveTable)
 	{
 		ManagedWorkspace ws = await CreateManagedWorkspaceAsync(useHaveTable);
-		List<string> view = new() { "-/Data/...", "-/shared.h" };
+		List<string> view = ["-/Data/...", "-/shared.h"];
 		ChangelistFixture cl = Stream.LatestChangelist;
-		List<DepotFileFixture> filtered = cl.StreamFiles
+		List<DepotFileFixture> filtered = [.. cl.StreamFiles
 			.Where(x => !x.DepotFile.Contains("shared.h", StringComparison.Ordinal))
-			.Where(x => !x.DepotFile.Contains("Data/", StringComparison.Ordinal)).ToList();
+			.Where(x => !x.DepotFile.Contains("Data/", StringComparison.Ordinal))];
 
-		List<PopulateRequest> populateRequests = new() { new PopulateRequest(PerforceConnection, StreamName, view) };
+		List<PopulateRequest> populateRequests = [new PopulateRequest(PerforceConnection, StreamName, view)];
 		await ws.PopulateAsync(populateRequests, false, CancellationToken.None);
 
 		ChangelistFixture clViewApplied = new(cl.Number, cl.Description, filtered, cl.IsShelved);
@@ -266,7 +265,7 @@ public class ManagedWorkspaceTest : BasePerforceFixtureTest
 		await SyncAsync(ws, 1);
 
 		// Assert - all files from CL 7 are moved to the cache
-		PerforceFixture.AssertCacheEquals(CacheDir, Stream.GetChangelist(7).StreamFiles.Select(x => x.Digest).ToArray());
+		PerforceFixture.AssertCacheEquals(CacheDir, [.. Stream.GetChangelist(7).StreamFiles.Select(x => x.Digest)]);
 	}
 
 	[TestMethod]
@@ -283,7 +282,7 @@ public class ManagedWorkspaceTest : BasePerforceFixtureTest
 		await SyncAsync(ws, 7);
 
 		// Assert - cache becomes empty
-		PerforceFixture.AssertCacheEquals(CacheDir, Array.Empty<string>());
+		PerforceFixture.AssertCacheEquals(CacheDir, []);
 	}
 
 	[TestMethod]
@@ -347,7 +346,7 @@ public class ManagedWorkspaceTest : BasePerforceFixtureTest
 		List<ChangesRecord> shelvedChanges = await PerforceConnection.GetChangesAsync(ChangesOptions.None, -1, ChangeStatus.Shelved, fileSpec);
 		foreach (ChangesRecord cr in submittedChanges.Concat(shelvedChanges).OrderBy(x => x.Number))
 		{
-			HashSet<(string clientFile, int rev, long size, string digest)> depotFiles = new();
+			HashSet<(string clientFile, int rev, long size, string digest)> depotFiles = [];
 			List<FStatRecord> fstatRecords = await PerforceConnection.FStatAsync(options, fileSpec + "@" + cr.Number).ToListAsync();
 			foreach (FStatRecord fsr in fstatRecords)
 			{
@@ -357,7 +356,7 @@ public class ManagedWorkspaceTest : BasePerforceFixtureTest
 				}
 			}
 
-			HashSet<(string clientFile, int rev, long size, string digest)> fixtureFiles = new();
+			HashSet<(string clientFile, int rev, long size, string digest)> fixtureFiles = [];
 			ChangelistFixture changelist = Stream.GetChangelist(cr.Number);
 			if (changelist.IsShelved)
 			{
@@ -411,7 +410,7 @@ public class ManagedWorkspaceTest : BasePerforceFixtureTest
 
 	private async Task SyncAsync(ManagedWorkspace managedWorkspace, int changeNumber, FileReference? cacheFile = null, bool removeUntracked = true, bool assertFiles = false, bool useHaveTable = true)
 	{
-		await managedWorkspace.SyncAsync(PerforceConnection, StreamName, changeNumber, Array.Empty<string>(), removeUntracked, false, cacheFile, CancellationToken.None);
+		await managedWorkspace.SyncAsync(PerforceConnection, StreamName, changeNumber, [], removeUntracked, false, cacheFile, CancellationToken.None);
 		if (assertFiles)
 		{
 			Stream.GetChangelist(changeNumber).AssertDepotFiles(SyncDir);
diff --git a/Engine/Source/Programs/Shared/EpicGames.Perforce.Managed/ManagedWorkspace.cs b/Engine/Source/Programs/Shared/EpicGames.Perforce.Managed/ManagedWorkspace.cs
--- a/Engine/Source/Programs/Shared/EpicGames.Perforce.Managed/ManagedWorkspace.cs
+++ b/Engine/Source/Programs/Shared/EpicGames.Perforce.Managed/ManagedWorkspace.cs
@@ -243,12 +243,12 @@ namespace EpicGames.Perforce.Managed
 		/// <summary>
 		/// Set of clients that we're created. Used to avoid updating multiple times during one run.
 		/// </summary>
-		readonly Dictionary<string, ClientRecord> _createdClients = new Dictionary<string, ClientRecord>();
+		readonly Dictionary<string, ClientRecord> _createdClients = [];
 
 		/// <summary>
 		/// Set of unique cache entries. We use this to ensure new names in the cache are unique.
 		/// </summary>
-		readonly HashSet<ulong> _cacheEntries = new HashSet<ulong>();
+		readonly HashSet<ulong> _cacheEntries = [];
 
 		/// <summary>
 		/// List of all the staged files
@@ -258,7 +258,7 @@ namespace EpicGames.Perforce.Managed
 		/// <summary>
 		/// All the files which are currently being tracked
 		/// </summary>
-		Dictionary<FileContentId, CachedFileInfo> _contentIdToTrackedFile = new Dictionary<FileContentId, CachedFileInfo>();
+		Dictionary<FileContentId, CachedFileInfo> _contentIdToTrackedFile = [];
 
 		/// <summary>
 		/// Notification that a clean has been performed
@@ -610,9 +610,11 @@ namespace EpicGames.Perforce.Managed
 
 			if (filesToDelete.Length > 0 || directoriesToDelete.Length > 0)
 			{
-				List<string> paths = new List<string>();
-				paths.AddRange(directoriesToDelete.Select(x => String.Format("/{0}/...", new DirectoryReference(x).MakeRelativeTo(_workspaceDir).Replace(Path.DirectorySeparatorChar, '/'))));
-				paths.AddRange(filesToDelete.Select(x => String.Format("/{0}", new FileReference(x).MakeRelativeTo(_workspaceDir).Replace(Path.DirectorySeparatorChar, '/'))));
+				List<string> paths =
+				[
+					.. directoriesToDelete.Select(x => String.Format("/{0}/...", new DirectoryReference(x).MakeRelativeTo(_workspaceDir).Replace(Path.DirectorySeparatorChar, '/'))),
+					.. filesToDelete.Select(x => String.Format("/{0}", new FileReference(x).MakeRelativeTo(_workspaceDir).Replace(Path.DirectorySeparatorChar, '/'))),
+				];
 
 				const int MaxDisplay = 1000;
 				foreach (string path in paths.OrderBy(x => x).Take(MaxDisplay))
@@ -686,7 +688,7 @@ namespace EpicGames.Perforce.Managed
 
 			_logger.LogInformation("Dumping repository to log...");
 
-			WorkspaceFileInfo[] workspaceFiles = _workspace.GetFiles().OrderBy(x => x.GetLocation().FullName).ToArray();
+			WorkspaceFileInfo[] workspaceFiles = [.. _workspace.GetFiles().OrderBy(x => x.GetLocation().FullName)];
 			if (workspaceFiles.Length > 0)
 			{
 				_logger.LogDebug("  Workspace:");
@@ -721,7 +723,7 @@ namespace EpicGames.Perforce.Managed
 				// Make sure all the folders exist in the cache
 				CreateCacheHierarchy();
 
-				List<CachedFileInfo> trackedFiles = _contentIdToTrackedFile.Values.ToList();
+				List<CachedFileInfo> trackedFiles = [.. _contentIdToTrackedFile.Values];
 
 				// Check that all the files in the cache appear as we expect them to
 				const int MaxLoggedMissingFiles = 250;
@@ -783,7 +785,7 @@ namespace EpicGames.Perforce.Managed
 			using (Trace("Purge"))
 			using (_logger.BeginIndentScope("  "))
 			{
-				List<CachedFileInfo> cachedFiles = _contentIdToTrackedFile.Values.OrderBy(x => x.SequenceNumber).ToList();
+				List<CachedFileInfo> cachedFiles = [.. _contentIdToTrackedFile.Values.OrderBy(x => x.SequenceNumber)];
 
 				int numRemovedFiles = 0;
 				long totalSize = cachedFiles.Sum(x => x.Length);
@@ -863,7 +865,7 @@ namespace EpicGames.Perforce.Managed
 					for (int idx = 0; idx < streamNames.Count; idx++)
 					{
 						List<StreamFile> files = streamState[idx].Item2.GetFiles();
-						filesInStream[idx] = new HashSet<FileContentId>(files.Select(x => x.ContentId));
+						filesInStream[idx] = [.. files.Select(x => x.ContentId)];
 					}
 
 					// Build a table showing amount of unique content in each stream
@@ -1066,7 +1068,7 @@ namespace EpicGames.Perforce.Managed
 			_logger.LogInformation("Unshelving changelist {Change}...", unshelveChangelist);
 
 			// query the contents of the shelved changelist
-			List<DescribeRecord> records = await perforce.DescribeAsync(DescribeOptions.Shelved, -1, new int[] { unshelveChangelist }, cancellationToken);
+			List<DescribeRecord> records = await perforce.DescribeAsync(DescribeOptions.Shelved, -1, [unshelveChangelist], cancellationToken);
 			if (records.Count != 1)
 			{
 				throw new PerforceException($"Changelist {unshelveChangelist} is not shelved");
@@ -1082,8 +1084,8 @@ namespace EpicGames.Perforce.Managed
 			Dictionary<string, WhereRecord> whereRecords = whereResponseList.Where(x => x.Succeeded).Select(x => x.Data).ToDictionary(x => x.DepotFile, x => x, StringComparer.OrdinalIgnoreCase);
 
 			// parse out all the list of deleted and modified files
-			List<WhereRecord> deleteFiles = new List<WhereRecord>();
-			List<WhereRecord> writeFiles = new List<WhereRecord>();
+			List<WhereRecord> deleteFiles = [];
+			List<WhereRecord> writeFiles = [];
 			foreach (DescribeFileRecord fileRecord in lastRecord.Files)
 			{
 				WhereRecord? whereRecord = null;
@@ -1230,8 +1232,8 @@ namespace EpicGames.Perforce.Managed
 			{
 				Stopwatch timer = Stopwatch.StartNew();
 
-				HashSet<FileContentId> commonContentIds = new HashSet<FileContentId>();
-				Dictionary<FileContentId, long> contentIdToLength = new Dictionary<FileContentId, long>();
+				HashSet<FileContentId> commonContentIds = [];
+				Dictionary<FileContentId, long> contentIdToLength = [];
 				for (int idx = 0; idx < requests.Count; idx++)
 				{
 					List<StreamFile> files = streamState[idx].Item2.GetFiles();
@@ -1250,7 +1252,7 @@ namespace EpicGames.Perforce.Managed
 					}
 				}
 
-				List<CachedFileInfo> trackedFiles = _contentIdToTrackedFile.Values.ToList();
+				List<CachedFileInfo> trackedFiles = [.. _contentIdToTrackedFile.Values];
 				foreach (CachedFileInfo trackedFile in trackedFiles)
 				{
 					if (!contentIdToLength.ContainsKey(trackedFile.ContentId))
@@ -1479,7 +1481,7 @@ namespace EpicGames.Perforce.Managed
 				foreach (ChangesRecord change in changes)
 				{
 					// Delete the shelved files
-					List<DescribeRecord> describeResponse = await perforceClient.DescribeAsync(DescribeOptions.Shelved, -1, new[] { change.Number }, cancellationToken);
+					List<DescribeRecord> describeResponse = await perforceClient.DescribeAsync(DescribeOptions.Shelved, -1, [change.Number], cancellationToken);
 					foreach (DescribeRecord record in describeResponse)
 					{
 						if (record.Files.Count > 0)
@@ -1639,7 +1641,7 @@ namespace EpicGames.Perforce.Managed
 
 				// List of the last path fragments. Since file records that are returned are typically sorted by their position in the tree, we can save quite a lot of processing by
 				// reusing as many fragemnts as possible.
-				List<(Utf8String, StreamTreeBuilder)> fragments = new List<(Utf8String, StreamTreeBuilder)>();
+				List<(Utf8String, StreamTreeBuilder)> fragments = [];
 
 				int numFiles = 0;
 
@@ -1737,14 +1739,16 @@ namespace EpicGames.Perforce.Managed
 				}
 
 				// Create the workspace, and add records for all the files. Exclude deleted files with digest = null.
-				List<string> arguments = new List<string>();
-				arguments.Add("-Ol"); // Output fileSize and digest field
-				arguments.Add("-Op"); // Output clientFile field in both server and local path syntax
-				arguments.Add("-Os"); // Shorten output by excluding client workspace data (for instance, the clientFile field).
-				arguments.Add("-Rh"); // Limit output to files on your have list;
-				arguments.Add("-T"); // Include only the fields listed below
-				arguments.Add(String.Join(",", FStatIndexedRecord.FieldNames));
-				arguments.Add($"//{perforceClient.Settings.ClientName}/...@{changeNumber}");
+				List<string> arguments =
+				[
+					"-Ol", // Output fileSize and digest field
+					"-Op", // Output clientFile field in both server and local path syntax
+					"-Os", // Shorten output by excluding client workspace data (for instance, the clientFile field).
+					"-Rh", // Limit output to files on your have list;
+					"-T", // Include only the fields listed below
+					String.Join(",", FStatIndexedRecord.FieldNames),
+					$"//{perforceClient.Settings.ClientName}/...@{changeNumber}",
+				];
 				await perforceClient.RecordCommandAsync("fstat", arguments, null, HandleRecord, cancellationToken);
 
 				// Output the elapsed time
@@ -1840,14 +1844,16 @@ namespace EpicGames.Perforce.Managed
 				}
 
 				string fileSpec = $"//{perforceClient.Settings.ClientName}/...@{changeNumber}";
-				List<string> arguments = new();
-				arguments.Add("-Ol"); // Output fileSize and digest field
-				arguments.Add("-Os"); // Shorten output by excluding client workspace data (for instance, the clientFile field).
-				arguments.Add("-F"); // Filter any files not existing at current revision (filter below)
-				arguments.Add("^headAction=delete&^headAction=move/delete&^headAction=purge");
-				arguments.Add("-T"); // Include only the fields listed below
-				arguments.Add(String.Join(",", FStatRecordWithoutHaveTable.FieldNames));
-				arguments.Add(fileSpec);
+				List<string> arguments =
+				[
+					"-Ol", // Output fileSize and digest field
+					"-Os", // Shorten output by excluding client workspace data (for instance, the clientFile field).
+					"-F", // Filter any files not existing at current revision (filter below)
+					"^headAction=delete&^headAction=move/delete&^headAction=purge",
+					"-T", // Include only the fields listed below
+					String.Join(",", FStatRecordWithoutHaveTable.FieldNames),
+					fileSpec,
+				];
 
 				await perforceClient.RecordCommandAsync("fstat", arguments, null, HandleRecord, cancellationToken);
 
@@ -1950,7 +1956,7 @@ namespace EpicGames.Perforce.Managed
 			}
 
 			// Move files into the cache
-			KeyValuePair<FileContentId, WorkspaceFileInfo>[] filesToMove = transaction._filesToMove.ToArray();
+			KeyValuePair<FileContentId, WorkspaceFileInfo>[] filesToMove = [.. transaction._filesToMove];
 			if (filesToMove.Length > 0)
 			{
 				using (Trace("MoveToCache"))
@@ -1959,7 +1965,7 @@ namespace EpicGames.Perforce.Managed
 					Stopwatch timer = Stopwatch.StartNew();
 
 					// Add any new files to the cache
-					List<(FileContentId ContentId, FileReference Source, FileReference Target)> files = new List<(FileContentId, FileReference, FileReference)>();
+					List<(FileContentId ContentId, FileReference Source, FileReference Target)> files = [];
 					foreach (KeyValuePair<FileContentId, WorkspaceFileInfo> fileToMove in filesToMove)
 					{
 						ulong cacheId = GetUniqueCacheId(fileToMove.Key);
@@ -1981,7 +1987,7 @@ namespace EpicGames.Perforce.Managed
 			}
 
 			// Remove files which are no longer needed
-			WorkspaceFileInfo[] filesToDelete = transaction._filesToDelete.ToArray();
+			WorkspaceFileInfo[] filesToDelete = [.. transaction._filesToDelete];
 			if (filesToDelete.Length > 0)
 			{
 				using (Trace("DeleteFiles"))
@@ -2001,7 +2007,7 @@ namespace EpicGames.Perforce.Managed
 			}
 
 			// Remove directories which are no longer needed
-			WorkspaceDirectoryInfo[] directoriesToDelete = transaction._directoriesToDelete.ToArray();
+			WorkspaceDirectoryInfo[] directoriesToDelete = [.. transaction._directoriesToDelete];
 			if (directoriesToDelete.Length > 0)
 			{
 				using (Trace("DeleteDirectories"))
@@ -2114,7 +2120,7 @@ namespace EpicGames.Perforce.Managed
 			}
 
 			// Swap files in and out of the cache
-			WorkspaceFileToMove[] filesToMove = transaction._filesToMove.Values.ToArray();
+			WorkspaceFileToMove[] filesToMove = [.. transaction._filesToMove.Values];
 			if (filesToMove.Length > 0)
 			{
 				using (Trace("MoveFromCache"))
@@ -2134,7 +2140,7 @@ namespace EpicGames.Perforce.Managed
 			}
 
 			// Swap files in and out of the cache
-			WorkspaceFileToCopy[] filesToCopy = transaction._filesToCopy.ToArray();
+			WorkspaceFileToCopy[] filesToCopy = [.. transaction._filesToCopy];
 			if (filesToCopy.Length > 0)
 			{
 				using (Trace("CopyFiles"))
@@ -2152,7 +2158,7 @@ namespace EpicGames.Perforce.Managed
 			}
 
 			// Find all the files we want to sync
-			WorkspaceFileToSync[] filesToSync = transaction._filesToSync.ToArray();
+			WorkspaceFileToSync[] filesToSync = [.. transaction._filesToSync];
 			if (filesToSync.Length > 0)
 			{
 				long syncSize = filesToSync.Sum(x => x._streamFile.Length);
@@ -2186,7 +2192,7 @@ namespace EpicGames.Perforce.Managed
 					}
 
 					// Create a list of all the batches that we want to sync
-					List<(int, int)> batches = new List<(int, int)>();
+					List<(int, int)> batches = [];
 					for (int endIdx = 0; endIdx < filesToSync.Length;)
 					{
 						int beginIdx = endIdx;
@@ -2209,7 +2215,7 @@ namespace EpicGames.Perforce.Managed
 					long syncedSize = 0;
 
 					// Spawn some background threads to sync them
-					Dictionary<Task, int> tasks = new Dictionary<Task, int>();
+					Dictionary<Task, int> tasks = [];
 					Stack<IPerforceConnection> connectionPool = new Stack<IPerforceConnection>();
 					try
 					{
@@ -2234,7 +2240,7 @@ namespace EpicGames.Perforce.Managed
 							// Update metadata for the complete batch
 							(int beginIdx, int endIdx) = batches[batchIdx];
 
-							int[] indexesToUpdate = Enumerable.Range(beginIdx, endIdx - beginIdx).ToArray();
+							int[] indexesToUpdate = [.. Enumerable.Range(beginIdx, endIdx - beginIdx)];
 							ParallelOptions options = new() { MaxDegreeOfParallelism = _options.MaxFileScanConcurrency, CancellationToken = cancellationToken };
 							await Parallel.ForEachAsync(indexesToUpdate, options, (idx, ct) =>
 							{
@@ -2296,12 +2302,12 @@ namespace EpicGames.Perforce.Managed
 				{
 					FileReference localFile = filesToSync[idx]._workspaceFile.GetLocation();
 					DirectoryReference.CreateDirectory(localFile.Directory);
-					await FileReference.WriteAllBytesAsync(localFile, Array.Empty<byte>(), cancellationToken);
+					await FileReference.WriteAllBytesAsync(localFile, [], cancellationToken);
 				}
 			}
 			else if (client is NativePerforceConnection)
 			{
-				List<string> files = new List<string>();
+				List<string> files = [];
 				for (int idx = beginIdx; idx < endIdx; idx++)
 				{
 					files.Add($"{filesToSync[idx]._streamFile.Path}#{filesToSync[idx]._streamFile.Revision}");
diff --git a/Engine/Source/Programs/Shared/EpicGames.Perforce.Managed/StreamSnapshot.cs b/Engine/Source/Programs/Shared/EpicGames.Perforce.Managed/StreamSnapshot.cs
--- a/Engine/Source/Programs/Shared/EpicGames.Perforce.Managed/StreamSnapshot.cs
+++ b/Engine/Source/Programs/Shared/EpicGames.Perforce.Managed/StreamSnapshot.cs
@@ -38,7 +38,7 @@ namespace EpicGames.Perforce.Managed
 		/// <returns>List of files</returns>
 		public static List<StreamFile> GetFiles(this StreamSnapshot snapshot)
 		{
-			List<StreamFile> files = new List<StreamFile>();
+			List<StreamFile> files = [];
 			AppendFiles(snapshot, snapshot.Root, files);
 			return files;
 		}
diff --git a/Engine/Source/Programs/Shared/EpicGames.Perforce.Managed/StreamSnapshotFromMemory.cs b/Engine/Source/Programs/Shared/EpicGames.Perforce.Managed/StreamSnapshotFromMemory.cs
--- a/Engine/Source/Programs/Shared/EpicGames.Perforce.Managed/StreamSnapshotFromMemory.cs
+++ b/Engine/Source/Programs/Shared/EpicGames.Perforce.Managed/StreamSnapshotFromMemory.cs
@@ -18,7 +18,7 @@ namespace EpicGames.Perforce.Managed
 		/// <summary>
 		/// The current signature for saved directory objects
 		/// </summary>
-		static readonly byte[] s_currentSignature = { (byte)'W', (byte)'S', (byte)'D', 5 };
+		static readonly byte[] s_currentSignature = [(byte)'W', (byte)'S', (byte)'D', 5];
 
 		/// <summary>
 		/// The root digest
@@ -47,7 +47,7 @@ namespace EpicGames.Perforce.Managed
 		/// <param name="builder"></param>
 		public StreamSnapshotFromMemory(StreamTreeBuilder builder)
 		{
-			Dictionary<IoHash, CbObject> hashToTree = new Dictionary<IoHash, CbObject>();
+			Dictionary<IoHash, CbObject> hashToTree = [];
 			Root = builder.EncodeRef(tree => EncodeObject(tree, hashToTree));
 			HashToTree = hashToTree;
 		}
diff --git a/Engine/Source/Programs/Shared/EpicGames.Perforce.Managed/StreamTree.cs b/Engine/Source/Programs/Shared/EpicGames.Perforce.Managed/StreamTree.cs
--- a/Engine/Source/Programs/Shared/EpicGames.Perforce.Managed/StreamTree.cs
+++ b/Engine/Source/Programs/Shared/EpicGames.Perforce.Managed/StreamTree.cs
@@ -175,7 +175,7 @@ namespace EpicGames.Perforce.Managed
 		/// <summary>
 		/// Map of name to file within the directory
 		/// </summary>
-		public Dictionary<Utf8String, StreamFile> NameToFile { get; } = new Dictionary<Utf8String, StreamFile>();
+		public Dictionary<Utf8String, StreamFile> NameToFile { get; } = [];
 
 		/// <summary>
 		/// Map of name to subdirectory
diff --git a/Engine/Source/Programs/Shared/EpicGames.Perforce.Managed/StreamTreeBuilder.cs b/Engine/Source/Programs/Shared/EpicGames.Perforce.Managed/StreamTreeBuilder.cs
--- a/Engine/Source/Programs/Shared/EpicGames.Perforce.Managed/StreamTreeBuilder.cs
+++ b/Engine/Source/Programs/Shared/EpicGames.Perforce.Managed/StreamTreeBuilder.cs
@@ -61,7 +61,7 @@ namespace EpicGames.Perforce.Managed
 			EncodeChildren(writeTree);
 
 			// Find the common base path for all items in this tree.
-			Dictionary<Utf8String, int> basePathToCount = new Dictionary<Utf8String, int>();
+			Dictionary<Utf8String, int> basePathToCount = [];
 			foreach ((Utf8String name, StreamFile file) in NameToFile)
 			{
 				AddBasePath(basePathToCount, file.Path, name);
diff --git a/Engine/Source/Programs/Shared/EpicGames.Perforce.Managed/Transactions/AddTransaction.cs b/Engine/Source/Programs/Shared/EpicGames.Perforce.Managed/Transactions/AddTransaction.cs
--- a/Engine/Source/Programs/Shared/EpicGames.Perforce.Managed/Transactions/AddTransaction.cs
+++ b/Engine/Source/Programs/Shared/EpicGames.Perforce.Managed/Transactions/AddTransaction.cs
@@ -56,7 +56,7 @@ namespace EpicGames.Perforce.Managed
 		public ConcurrentQueue<WorkspaceFileToCopy> _filesToCopy = new ConcurrentQueue<WorkspaceFileToCopy>();
 		public ConcurrentQueue<WorkspaceFileToSync> _filesToSync = new ConcurrentQueue<WorkspaceFileToSync>();
 		readonly Dictionary<FileContentId, CachedFileInfo> _contentIdToTrackedFile;
-		readonly Dictionary<FileContentId, WorkspaceFileInfo> _contentIdToWorkspaceFile = new Dictionary<FileContentId, WorkspaceFileInfo>();
+		readonly Dictionary<FileContentId, WorkspaceFileInfo> _contentIdToWorkspaceFile = [];
 
 		private AddTransaction(WorkspaceDirectoryInfo workspaceRootDir, StreamSnapshot streamSnapshot, Dictionary<FileContentId, CachedFileInfo> contentIdToTrackedFile)
 		{
diff --git a/Engine/Source/Programs/Shared/EpicGames.Perforce.Managed/WorkspaceDirectoryInfo.cs b/Engine/Source/Programs/Shared/EpicGames.Perforce.Managed/WorkspaceDirectoryInfo.cs
--- a/Engine/Source/Programs/Shared/EpicGames.Perforce.Managed/WorkspaceDirectoryInfo.cs
+++ b/Engine/Source/Programs/Shared/EpicGames.Perforce.Managed/WorkspaceDirectoryInfo.cs
@@ -103,7 +103,7 @@ namespace EpicGames.Perforce.Managed
 		/// <returns>List of files</returns>
 		public List<WorkspaceFileInfo> GetFiles()
 		{
-			List<WorkspaceFileInfo> files = new List<WorkspaceFileInfo>();
+			List<WorkspaceFileInfo> files = [];
 			GetFilesInternal(files);
 			return files;
 		}
@@ -217,7 +217,7 @@ namespace EpicGames.Perforce.Managed
 			using AsyncThreadPoolWorkQueue queue = new(numWorkers);
 			await queue.EnqueueAsync(_ => FindDifferencesAsync(new DirectoryInfo(GetFullName()), "/", paths, queue));
 			await queue.ExecuteAsync();
-			return paths.OrderBy(x => x).ToArray();
+			return [.. paths.OrderBy(x => x)];
 		}
 
 		/// <summary>
@@ -230,7 +230,7 @@ namespace EpicGames.Perforce.Managed
 		async Task FindDifferencesAsync(DirectoryInfo directory, string path, ConcurrentQueue<string> paths, AsyncThreadPoolWorkQueue queue)
 		{
 			// Recurse through subdirectories
-			HashSet<Utf8String> remainingSubDirectoryNames = new HashSet<Utf8String>(NameToSubDirectory.Keys);
+			HashSet<Utf8String> remainingSubDirectoryNames = [.. NameToSubDirectory.Keys];
 			foreach (DirectoryInfo subDirectory in directory.EnumerateDirectories())
 			{
 				WorkspaceDirectoryInfo? stagedSubDirectory;
@@ -248,7 +248,7 @@ namespace EpicGames.Perforce.Managed
 			}
 
 			// Search through files
-			HashSet<Utf8String> remainingFileNames = new HashSet<Utf8String>(NameToFile.Keys);
+			HashSet<Utf8String> remainingFileNames = [.. NameToFile.Keys];
 			foreach (FileInfo file in directory.EnumerateFiles())
 			{
 				WorkspaceFileInfo? stagedFile;
diff --git a/Engine/Source/Programs/Shared/EpicGames.Perforce.Tests/PerforceChangeViewTest.cs b/Engine/Source/Programs/Shared/EpicGames.Perforce.Tests/PerforceChangeViewTest.cs
--- a/Engine/Source/Programs/Shared/EpicGames.Perforce.Tests/PerforceChangeViewTest.cs
+++ b/Engine/Source/Programs/Shared/EpicGames.Perforce.Tests/PerforceChangeViewTest.cs
@@ -11,11 +11,11 @@ namespace EpicGames.Perforce.Tests
 		public void TestImportAtChange()
 		{
 			string[] viewLines =
-			{
+			[
 				"//UE5/Main/Testing/...@123",
 				"//UE5/Main/Testing/Foo/...@456",
 				"//UE5/Main/Testing/Foo/Bar/...@100",
-			};
+			];
 
 			PerforceChangeView changeView = PerforceChangeView.Parse(viewLines, true);
 
diff --git a/Engine/Source/Programs/Shared/EpicGames.Perforce/FileSpecList.cs b/Engine/Source/Programs/Shared/EpicGames.Perforce/FileSpecList.cs
--- a/Engine/Source/Programs/Shared/EpicGames.Perforce/FileSpecList.cs
+++ b/Engine/Source/Programs/Shared/EpicGames.Perforce/FileSpecList.cs
@@ -14,12 +14,12 @@ namespace EpicGames.Perforce
 		/// <summary>
 		/// Empty filespec list
 		/// </summary>
-		public static FileSpecList Empty { get; } = new FileSpecList(new List<string>());
+		public static FileSpecList Empty { get; } = new FileSpecList([]);
 
 		/// <summary>
 		/// Matches any files in the depot
 		/// </summary>
-		public static FileSpecList Any { get; } = new FileSpecList(new List<string> { "//..." });
+		public static FileSpecList Any { get; } = new FileSpecList(["//..."]);
 
 		/// <summary>
 		/// The list of filespecs
@@ -86,7 +86,7 @@ namespace EpicGames.Perforce
 		/// <param name="fileSpec">The single filespec to construct from</param>
 		public static implicit operator FileSpecList(string fileSpec)
 		{
-			return new FileSpecList(new string[] { fileSpec });
+			return new FileSpecList([fileSpec]);
 		}
 
 		/// <inheritdoc/>
diff --git a/Engine/Source/Programs/Shared/EpicGames.Perforce/IPerforceOutput.cs b/Engine/Source/Programs/Shared/EpicGames.Perforce/IPerforceOutput.cs
--- a/Engine/Source/Programs/Shared/EpicGames.Perforce/IPerforceOutput.cs
+++ b/Engine/Source/Programs/Shared/EpicGames.Perforce/IPerforceOutput.cs
@@ -59,7 +59,7 @@ namespace EpicGames.Perforce
 		/// <summary>
 		/// Standard prefix for a returned record: record indicator, string, 4 bytes, 'code', string, [value]
 		/// </summary>
-		static readonly byte[] s_recordPrefix = { (byte)'{', (byte)'s', 4, 0, 0, 0, (byte)'c', (byte)'o', (byte)'d', (byte)'e', (byte)'s' };
+		static readonly byte[] s_recordPrefix = [(byte)'{', (byte)'s', 4, 0, 0, 0, (byte)'c', (byte)'o', (byte)'d', (byte)'e', (byte)'s'];
 
 		class BufferedPerforceOutput : IPerforceOutput
 		{
@@ -87,7 +87,7 @@ namespace EpicGames.Perforce
 		/// </summary>
 		/// <param name="response">Response to construct from</param>
 		/// <returns>Output object</returns>
-		public static IPerforceOutput FromResponse(PerforceResponse response) => FromResponses(new[] { response });
+		public static IPerforceOutput FromResponse(PerforceResponse response) => FromResponses([response]);
 
 		/// <summary>
 		/// Constructs an <see cref="IPerforceOutput"/> object from a sequence of responses
@@ -223,7 +223,7 @@ namespace EpicGames.Perforce
 		{
 			CachedRecordInfo? statRecordInfo = (statRecordType == null) ? null : PerforceReflection.GetCachedRecordInfo(statRecordType);
 
-			List<PerforceResponse> responses = new List<PerforceResponse>();
+			List<PerforceResponse> responses = [];
 
 			// Read all the records into a list
 			long parsedLen = 0;
@@ -286,7 +286,7 @@ namespace EpicGames.Perforce
 		{
 			CachedRecordInfo? statRecordInfo = (statRecordType == null) ? null : PerforceReflection.GetCachedRecordInfo(statRecordType);
 
-			List<PerforceResponse> responses = new List<PerforceResponse>();
+			List<PerforceResponse> responses = [];
 
 			// Read all the records into a list
 			long maxParsedLen;
diff --git a/Engine/Source/Programs/Shared/EpicGames.Perforce/NativePerforceConnection.cs b/Engine/Source/Programs/Shared/EpicGames.Perforce/NativePerforceConnection.cs
--- a/Engine/Source/Programs/Shared/EpicGames.Perforce/NativePerforceConnection.cs
+++ b/Engine/Source/Programs/Shared/EpicGames.Perforce/NativePerforceConnection.cs
@@ -86,7 +86,7 @@ namespace EpicGames.Perforce
 		/// </summary>
 		class PinnedBuffer : IDisposable
 		{
-			static readonly byte[] s_guardBytes = Enumerable.Repeat<byte>(0xfd, 16).ToArray();
+			static readonly byte[] s_guardBytes = [.. Enumerable.Repeat<byte>(0xfd, 16)];
 
 			public byte[] AllocatedData { get; private set; }
 			public Memory<byte> Data { get; private set; }
@@ -292,8 +292,8 @@ namespace EpicGames.Perforce
 		readonly OnBufferReadyFn _onBufferReadyInst;
 		readonly IntPtr _onBufferReadyFnPtr;
 		Thread? _backgroundThread;
-		readonly BlockingCollection<(Action, Response)?> _requests = new BlockingCollection<(Action, Response)?>();
-		readonly BlockingCollection<PinnedBuffer> _writeBuffers = new BlockingCollection<PinnedBuffer>();
+		readonly BlockingCollection<(Action, Response)?> _requests = [];
+		readonly BlockingCollection<PinnedBuffer> _writeBuffers = [];
 		Response? _currentResponse;
 		readonly ManualResetEvent _responseCompleteEvent;
 		readonly Stopwatch _stallTimer = new Stopwatch();
@@ -606,7 +606,7 @@ namespace EpicGames.Perforce
 				specData = Encoding.UTF8.GetBytes(FormatSpec(inputData));
 			}
 
-			List<string> allArguments = new List<string>(arguments);
+			List<string> allArguments = [.. arguments];
 			if (fileArguments != null)
 			{
 				allArguments.AddRange(fileArguments);
@@ -636,7 +636,7 @@ namespace EpicGames.Perforce
 			using IDisposable scope = _hangMonitor.Start($"{command} {argList}");
 			long initialAllocatedSize = GetAllocatedSize();
 
-			List<IntPtr> nativeArgs = new List<IntPtr>();
+			List<IntPtr> nativeArgs = [];
 			try
 			{
 				foreach (string arg in args)
@@ -657,7 +657,7 @@ namespace EpicGames.Perforce
 
 				try
 				{
-					Client_Command(_client, command, nativeArgs.Count, nativeArgs.ToArray(), inputData, inputData?.Length ?? 0, promptResponseBytes, interceptIo);
+					Client_Command(_client, command, nativeArgs.Count, [.. nativeArgs], inputData, inputData?.Length ?? 0, promptResponseBytes, interceptIo);
 				}
 				catch (Exception ex)
 				{
@@ -701,7 +701,7 @@ namespace EpicGames.Perforce
 		{
 			int pos = 0;
 
-			List<KeyValuePair<Utf8String, PerforceValue>> rows = new List<KeyValuePair<Utf8String, PerforceValue>>();
+			List<KeyValuePair<Utf8String, PerforceValue>> rows = [];
 			if (!PerforceOutput.ParseRecord(inputData, ref pos, rows))
 			{
 				throw new PerforceException("Unable to parse input data as record");
diff --git a/Engine/Source/Programs/Shared/EpicGames.Perforce/PerforceChangeView.cs b/Engine/Source/Programs/Shared/EpicGames.Perforce/PerforceChangeView.cs
--- a/Engine/Source/Programs/Shared/EpicGames.Perforce/PerforceChangeView.cs
+++ b/Engine/Source/Programs/Shared/EpicGames.Perforce/PerforceChangeView.cs
@@ -15,7 +15,7 @@ namespace EpicGames.Perforce
 		/// <summary>
 		/// List of entries for this change view
 		/// </summary>
-		public List<PerforceChangeViewEntry> Entries { get; } = new List<PerforceChangeViewEntry>();
+		public List<PerforceChangeViewEntry> Entries { get; } = [];
 
 		/// <summary>
 		/// Determines if a file revision is visible 
diff --git a/Engine/Source/Programs/Shared/EpicGames.Perforce/PerforceChildProcess.cs b/Engine/Source/Programs/Shared/EpicGames.Perforce/PerforceChildProcess.cs
--- a/Engine/Source/Programs/Shared/EpicGames.Perforce/PerforceChildProcess.cs
+++ b/Engine/Source/Programs/Shared/EpicGames.Perforce/PerforceChildProcess.cs
@@ -63,9 +63,7 @@ namespace EpicGames.Perforce
 		{
 			string perforceFileName = GetExecutable();
 
-			List<string> fullArguments = new List<string>();
-			fullArguments.Add("-G");
-			fullArguments.AddRange(globalOptions);
+			List<string> fullArguments = ["-G", .. globalOptions];
 			if (fileArguments != null)
 			{
 				_tempFileName = Path.GetTempFileName();
@@ -96,11 +94,11 @@ namespace EpicGames.Perforce
 			string perforceFileName = "p4.exe";
 			if (!OperatingSystem.IsWindows())
 			{
-				string[] p4Paths = {
+				string[] p4Paths = [
 					"/usr/bin/p4", // Default path
 					"/opt/homebrew/bin/p4", // Apple Silicon Homebrew Path
 					"/usr/local/bin/p4", // Apple Intel Homebrew Path
-					"[REDACTED_USER_HOME]/.linuxbrew/bin/p4" }; // Linux Homebrew path
+					"[REDACTED_USER_HOME]/.linuxbrew/bin/p4" ]; // Linux Homebrew path
 				foreach (string path in p4Paths)
 				{
 					if (File.Exists(path))
diff --git a/Engine/Source/Programs/Shared/EpicGames.Perforce/PerforceConnection.cs b/Engine/Source/Programs/Shared/EpicGames.Perforce/PerforceConnection.cs
--- a/Engine/Source/Programs/Shared/EpicGames.Perforce/PerforceConnection.cs
+++ b/Engine/Source/Programs/Shared/EpicGames.Perforce/PerforceConnection.cs
@@ -208,7 +208,7 @@ namespace EpicGames.Perforce
 		/// <summary>
 		/// Additional options to append to the command line
 		/// </summary>
-		public List<string> GlobalOptions { get; } = new List<string>();
+		public List<string> GlobalOptions { get; } = [];
 
 		/// <inheritdoc/>
 		public ILogger Logger { get; set; }
@@ -294,7 +294,7 @@ namespace EpicGames.Perforce
 
 		List<string> GetGlobalArguments()
 		{
-			List<string> arguments = new List<string>();
+			List<string> arguments = [];
 			if (ServerAndPort != null)
 			{
 				arguments.Add($"-p{ServerAndPort}");
@@ -362,8 +362,7 @@ namespace EpicGames.Perforce
 		/// <returns>Response from the server</returns>
 		public async Task<Tuple<bool, string>> TrySetAsync(string name, string value, CancellationToken cancellationToken = default)
 		{
-			List<string> arguments = new List<string>();
-			arguments.Add($"{name}={value}");
+			List<string> arguments = [$"{name}={value}"];
 
 			using (PerforceChildProcess childProcess = new PerforceChildProcess("set", arguments, null, null, GetGlobalArguments(), Logger))
 			{
@@ -555,7 +554,7 @@ namespace EpicGames.Perforce
 		{
 			List<PerforceResponse> responses = await connection.CommandAsync(command, arguments, fileArguments, inputData, typeof(T), cancellationToken);
 
-			PerforceResponseList<T> typedResponses = new PerforceResponseList<T>();
+			PerforceResponseList<T> typedResponses = [];
 			foreach (PerforceResponse response in responses)
 			{
 				typedResponses.Add(new PerforceResponse<T>(response));
@@ -588,11 +587,10 @@ namespace EpicGames.Perforce
 		/// <returns>List of objects returned by the server</returns>
 		public static async Task<PerforceResponseList<T>> BatchedCommandAsync<T>(IPerforceConnection connection, string command, IReadOnlyList<string> commonArguments, IReadOnlyList<string> batchedArguments, CancellationToken cancellationToken = default) where T : class
 		{
-			PerforceResponseList<T> responses = new PerforceResponseList<T>();
+			PerforceResponseList<T> responses = [];
 			for (int fileSpecIdx = 0; fileSpecIdx < batchedArguments.Count;)
 			{
-				List<string> arguments = new List<string>();
-				arguments.AddRange(commonArguments);
+				List<string> arguments = [.. commonArguments];
 
 				const int PerArgumentExtra = 5;
 				int length = (command.Length + PerArgumentExtra) + arguments.Sum(x => x.Length + PerArgumentExtra);
@@ -703,7 +701,7 @@ namespace EpicGames.Perforce
 		/// <returns>Response from the server</returns>
 		public static Task<PerforceResponseList<AddRecord>> TryAddAsync(this IPerforceConnection connection, int changeNumber, string? fileType, AddOptions options, FileSpecList fileNames, CancellationToken cancellationToken = default)
 		{
-			List<string> arguments = new List<string>();
+			List<string> arguments = [];
 			if (changeNumber != -1)
 			{
 				arguments.Add($"-c{changeNumber}");
@@ -762,7 +760,7 @@ namespace EpicGames.Perforce
 				throw new PerforceException("'Number' field should be set to -1 to create a new changelist.");
 			}
 
-			PerforceResponse response = await SingleResponseCommandAsync(connection, "change", new List<string> { "-i" }, SerializeRecord(connection, record), null, cancellationToken);
+			PerforceResponse response = await SingleResponseCommandAsync(connection, "change", ["-i"], SerializeRecord(connection, record), null, cancellationToken);
 
 			PerforceError? error = response.Error;
 			if (error != null)
@@ -814,8 +812,7 @@ namespace EpicGames.Perforce
 				throw new PerforceException("'Number' field must be set to update a changelist.");
 			}
 
-			List<string> arguments = new List<string>();
-			arguments.Add("-i");
+			List<string> arguments = ["-i"];
 			if ((options & UpdateChangeOptions.Force) != 0)
 			{
 				arguments.Add("-f");
@@ -851,7 +848,7 @@ namespace EpicGames.Perforce
 		/// <returns>Response from the server</returns>
 		public static Task<PerforceResponse> TryDeleteChangeAsync(this IPerforceConnection connection, DeleteChangeOptions options, int changeNumber, CancellationToken cancellationToken = default)
 		{
-			List<string> arguments = new List<string> { "-d" };
+			List<string> arguments = ["-d"];
 			if ((options & DeleteChangeOptions.Submitted) != 0)
 			{
 				arguments.Add("-f");
@@ -888,7 +885,7 @@ namespace EpicGames.Perforce
 		/// <returns>Response from the server</returns>
 		public static Task<PerforceResponse<ChangeRecord>> TryGetChange(this IPerforceConnection connection, GetChangeOptions options, int changeNumber, CancellationToken cancellationToken = default)
 		{
-			List<string> arguments = new List<string> { "-o" };
+			List<string> arguments = ["-o"];
 			if ((options & GetChangeOptions.BeforeRenumber) != 0)
 			{
 				arguments.Add("-O");
@@ -909,15 +906,10 @@ namespace EpicGames.Perforce
 		/// <returns>Serialized record</returns>
 		static byte[] SerializeRecord(this IPerforceConnection connection, ChangeRecord input)
 		{
-			List<KeyValuePair<string, object>> nameToValue = new List<KeyValuePair<string, object>>();
-			if (input.Number == -1)
-			{
-				nameToValue.Add(new KeyValuePair<string, object>("Change", "new"));
-			}
-			else
-			{
-				nameToValue.Add(new KeyValuePair<string, object>("Change", input.Number.ToString()));
-			}
+			List<KeyValuePair<string, object>> nameToValue =
+			[
+				input.Number == -1 ? new KeyValuePair<string, object>("Change", "new") : new KeyValuePair<string, object>("Change", input.Number.ToString()),
+			];
 			if (input.Type != ChangeType.Unspecified)
 			{
 				nameToValue.Add(new KeyValuePair<string, object>("Type", input.Type.ToString()));
@@ -1046,7 +1038,7 @@ namespace EpicGames.Perforce
 		/// <returns>List of responses from the server.</returns>
 		public static Task<PerforceResponseList<ChangesRecord>> TryGetChangesAsync(this IPerforceConnection connection, ChangesOptions options, string? clientName, int minChangeNumber, int maxChanges, ChangeStatus status, string? userName, FileSpecList fileSpecs, CancellationToken cancellationToken = default)
 		{
-			List<string> arguments = new List<string>();
+			List<string> arguments = [];
 			if ((options & ChangesOptions.IncludeIntegrations) != 0)
 			{
 				arguments.Add("-i");
@@ -1126,7 +1118,7 @@ namespace EpicGames.Perforce
 		/// <returns>Response from the server</returns>
 		public static async Task<PerforceResponseList<CleanRecord>> TryCleanAsync(this IPerforceConnection connection, CleanOptions options, FileSpecList fileSpecs, CancellationToken cancellationToken = default)
 		{
-			List<string> arguments = new List<string>();
+			List<string> arguments = [];
 			if ((options & CleanOptions.Edited) != 0)
 			{
 				arguments.Add("-e");
@@ -1211,7 +1203,7 @@ namespace EpicGames.Perforce
 		/// <returns>Response from the server</returns>
 		public static Task<PerforceResponse> TryUpdateClientAsync(this IPerforceConnection connection, ClientRecord record, CancellationToken cancellationToken = default)
 		{
-			return SingleResponseCommandAsync(connection, "client", new List<string> { "-i" }, connection.SerializeRecord(record), null, cancellationToken);
+			return SingleResponseCommandAsync(connection, "client", ["-i"], connection.SerializeRecord(record), null, cancellationToken);
 		}
 
 		/// <summary>
@@ -1237,7 +1229,7 @@ namespace EpicGames.Perforce
 		/// <returns>Response from the server</returns>
 		public static Task<PerforceResponse> TryDeleteClientAsync(this IPerforceConnection connection, DeleteClientOptions options, string clientName, CancellationToken cancellationToken = default)
 		{
-			List<string> arguments = new List<string> { "-d" };
+			List<string> arguments = ["-d"];
 			if ((options & DeleteClientOptions.Force) != 0)
 			{
 				arguments.Add("-f");
@@ -1274,7 +1266,7 @@ namespace EpicGames.Perforce
 		/// <returns>Response from the server</returns>
 		public static Task<PerforceResponse> TrySwitchClientToStreamAsync(this IPerforceConnection connection, string streamName, SwitchClientOptions options, CancellationToken cancellationToken = default)
 		{
-			List<string> arguments = new List<string> { "-s" };
+			List<string> arguments = ["-s"];
 			if ((options & SwitchClientOptions.IgnoreOpenFiles) != 0)
 			{
 				arguments.Add("-f");
@@ -1305,9 +1297,7 @@ namespace EpicGames.Perforce
 		/// <returns>Response from the server</returns>
 		public static Task<PerforceResponse> TrySwitchClientToTemplateAsync(this IPerforceConnection connection, string templateName, CancellationToken cancellationToken = default)
 		{
-			List<string> arguments = new List<string>();
-			arguments.Add("-s");
-			arguments.Add($"-t{templateName}");
+			List<string> arguments = ["-s", $"-t{templateName}"];
 
 			return SingleResponseCommandAsync(connection, "client", arguments, null, null, cancellationToken);
 		}
@@ -1333,7 +1323,7 @@ namespace EpicGames.Perforce
 		/// <returns>Response from the server; either a client record or error code</returns>
 		public static Task<PerforceResponse<ClientRecord>> TryGetClientAsync(this IPerforceConnection connection, string? clientName, CancellationToken cancellationToken = default)
 		{
-			List<string> arguments = new List<string> { "-o" };
+			List<string> arguments = ["-o"];
 			if (clientName != null)
 			{
 				arguments.Add(clientName);
@@ -1392,8 +1382,7 @@ namespace EpicGames.Perforce
 		/// <returns>Response from the server; either a client record or error code</returns>
 		public static Task<PerforceResponse<ClientRecord>> TryGetStreamViewAsync(this IPerforceConnection connection, string? clientName, string streamName, int changeNumber, CancellationToken cancellationToken = default)
 		{
-			List<string> arguments = new List<string> { "-o" };
-			arguments.Add($"-S{streamName}");
+			List<string> arguments = ["-o", $"-S{streamName}"];
 			if (changeNumber != -1)
 			{
 				arguments.Add($"-c{changeNumber}");
@@ -1415,7 +1404,7 @@ namespace EpicGames.Perforce
 		/// <returns>Serialized record data</returns>
 		static byte[] SerializeRecord(this IPerforceConnection connection, ClientRecord input)
 		{
-			List<KeyValuePair<string, object>> nameToValue = new List<KeyValuePair<string, object>>();
+			List<KeyValuePair<string, object>> nameToValue = [];
 			if (input.Name != null)
 			{
 				nameToValue.Add(new KeyValuePair<string, object>("Client", input.Name));
@@ -1522,7 +1511,7 @@ namespace EpicGames.Perforce
 		/// <returns>Response from the server; either a client record or error code</returns>
 		public static Task<PerforceResponseList<ClientsRecord>> TryGetClientsAsync(this IPerforceConnection connection, ClientsOptions options, string? filter, int maxResults, string? stream, string? userName, CancellationToken cancellationToken = default)
 		{
-			List<string> arguments = new List<string>();
+			List<string> arguments = [];
 			if ((options & ClientsOptions.All) != 0)
 			{
 				arguments.Add("-a");
@@ -1593,7 +1582,7 @@ namespace EpicGames.Perforce
 		/// <returns>List of responses from the server</returns>
 		public static async Task<PerforceResponseList<DeleteRecord>> TryDeleteAsync(this IPerforceConnection connection, int changeNumber, DeleteOptions options, FileSpecList fileSpecs, CancellationToken cancellationToken = default)
 		{
-			List<string> arguments = new List<string>();
+			List<string> arguments = [];
 			if (changeNumber != -1)
 			{
 				arguments.Add($"-c{changeNumber}");
@@ -1641,7 +1630,7 @@ namespace EpicGames.Perforce
 		/// <returns>Response from the server; either a client record or error code</returns>
 		public static Task<PerforceResponse<DepotRecord>> TryGetDepotAsync(this IPerforceConnection connection, string depotName, CancellationToken cancellationToken = default)
 		{
-			List<string> arguments = new List<string> { "-o", depotName };
+			List<string> arguments = ["-o", depotName];
 			return SingleResponseCommandAsync<DepotRecord>(connection, "depot", arguments, null, cancellationToken);
 		}
 
@@ -1670,7 +1659,7 @@ namespace EpicGames.Perforce
 		/// <returns>Response from the server; either a describe record or error code</returns>
 		public static async Task<PerforceResponse<DescribeRecord>> TryDescribeAsync(this IPerforceConnection connection, int changeNumber, CancellationToken cancellationToken = default)
 		{
-			PerforceResponseList<DescribeRecord> records = await TryDescribeAsync(connection, new int[] { changeNumber }, cancellationToken);
+			PerforceResponseList<DescribeRecord> records = await TryDescribeAsync(connection, [changeNumber], cancellationToken);
 			if (records.Count != 1)
 			{
 				throw new PerforceException("Expected only one record returned from p4 describe command, got {0}", records.Count);
@@ -1704,7 +1693,7 @@ namespace EpicGames.Perforce
 		/// <returns>Response from the server; either a describe record or error code</returns>
 		public static async Task<PerforceResponse<DescribeRecord>> TryDescribeAsync(this IPerforceConnection connection, DescribeOptions options, int maxNumFiles, int changeNumber, CancellationToken cancellationToken = default)
 		{
-			PerforceResponseList<DescribeRecord> records = await TryDescribeAsync(connection, options, maxNumFiles, new int[] { changeNumber }, cancellationToken);
+			PerforceResponseList<DescribeRecord> records = await TryDescribeAsync(connection, options, maxNumFiles, [changeNumber], cancellationToken);
 			if (records.Count != 1)
 			{
 				throw new PerforceException("Expected only one record returned from p4 describe command, got {0}", records.Count);
@@ -1733,7 +1722,7 @@ namespace EpicGames.Perforce
 		/// <returns>List of responses from the server</returns>
 		public static Task<PerforceResponseList<DescribeRecord>> TryDescribeAsync(this IPerforceConnection connection, int[] changeNumbers, CancellationToken cancellationToken = default)
 		{
-			List<string> arguments = new List<string> { "-s" };
+			List<string> arguments = ["-s"];
 			foreach (int changeNumber in changeNumbers)
 			{
 				arguments.Add($"{changeNumber}");
@@ -1766,7 +1755,7 @@ namespace EpicGames.Perforce
 		/// <returns>List of responses from the server</returns>
 		public static Task<PerforceResponseList<DescribeRecord>> TryDescribeAsync(this IPerforceConnection connection, DescribeOptions options, int maxNumFiles, int[] changeNumbers, CancellationToken cancellationToken = default)
 		{
-			List<string> arguments = new List<string> { "-s" };
+			List<string> arguments = ["-s"];
 			if ((options & DescribeOptions.ShowDescriptionForRestrictedChanges) != 0)
 			{
 				arguments.Add("-f");
@@ -1823,7 +1812,7 @@ namespace EpicGames.Perforce
 		/// <returns>Response from the server</returns>
 		public static Task<PerforceResponseList<DirsRecord>> TryGetDirsAsync(this IPerforceConnection connection, DirsOptions options, string? stream, FileSpecList fileSpecs, CancellationToken cancellationToken = default)
 		{
-			List<string> arguments = new List<string>();
+			List<string> arguments = [];
 			if ((options & DirsOptions.OnlyMapped) != 0)
 			{
 				arguments.Add("-C");
@@ -1906,7 +1895,7 @@ namespace EpicGames.Perforce
 		/// <returns>Response from the server</returns>
 		public static Task<PerforceResponseList<EditRecord>> TryEditAsync(this IPerforceConnection connection, int changeNumber, string? fileType, EditOptions options, FileSpecList fileSpecs, CancellationToken cancellationToken = default)
 		{
-			List<string> arguments = new List<string>();
+			List<string> arguments = [];
 			if (changeNumber != -1)
 			{
 				arguments.Add($"-c{changeNumber}");
@@ -2013,7 +2002,7 @@ namespace EpicGames.Perforce
 		public static async Task<PerforceResponseList<FileLogRecord>> TryFileLogAsync(this IPerforceConnection connection, int changeNumber, int maxChanges, FileLogOptions options, FileSpecList fileSpecs, CancellationToken cancellationToken = default)
 		{
 			// Build the argument list
-			List<string> arguments = new List<string>();
+			List<string> arguments = [];
 			if (changeNumber > 0)
 			{
 				arguments.Add($"-c{changeNumber}");
@@ -2111,7 +2100,7 @@ namespace EpicGames.Perforce
 		/// <returns>List of response objects</returns>
 		public static async Task<PerforceResponseList<FilesRecord>> TryFilesAsync(this IPerforceConnection connection, FilesOptions options, int maxFiles, FileSpecList fileSpecs, CancellationToken cancellationToken = default)
 		{
-			List<string> arguments = new List<string>();
+			List<string> arguments = [];
 			if ((options & FilesOptions.AllRevisions) != 0)
 			{
 				arguments.Add("-a");
@@ -2242,7 +2231,7 @@ namespace EpicGames.Perforce
 		public static IAsyncEnumerable<PerforceResponse<FStatRecord>> TryFStatAsync(this IPerforceConnection connection, int afterChangeNumber, int onlyChangeNumber, string? filter, string? fields, int maxFiles, FStatOptions options, FileSpecList fileSpecs, CancellationToken cancellationToken = default)
 		{
 			// Build the argument list
-			List<string> arguments = new List<string>();
+			List<string> arguments = [];
 			if (afterChangeNumber != -1)
 			{
 				arguments.Add($"-c{afterChangeNumber}");
@@ -2382,7 +2371,7 @@ namespace EpicGames.Perforce
 		/// <returns>List of file records</returns>
 		public static IAsyncEnumerable<PerforceResponse<HaveRecord>> TryHaveAsync(this IPerforceConnection connection, FileSpecList fileSpec, CancellationToken cancellationToken = default)
 		{
-			IAsyncEnumerable<PerforceResponse<HaveRecord>> records = StreamCommandAsync<HaveRecord>(connection, "have", new List<string>(), fileSpec.List, null, cancellationToken);
+			IAsyncEnumerable<PerforceResponse<HaveRecord>> records = StreamCommandAsync<HaveRecord>(connection, "have", [], fileSpec.List, null, cancellationToken);
 			records = records.Where(x => x.Error == null || x.Error.Generic != PerforceGenericCode.Empty);
 			return records;
 		}
@@ -2413,7 +2402,7 @@ namespace EpicGames.Perforce
 		public static Task<PerforceResponse<InfoRecord>> TryGetInfoAsync(this IPerforceConnection connection, InfoOptions options, CancellationToken cancellationToken = default)
 		{
 			// Build the argument list
-			List<string> arguments = new List<string>();
+			List<string> arguments = [];
 			if ((options & InfoOptions.ShortOutput) != 0)
 			{
 				arguments.Add("-s");
@@ -2461,7 +2450,7 @@ namespace EpicGames.Perforce
 		/// <returns></returns>
 		public static async Task<PerforceResponse<LoginRecord>> TryLoginAsync(this IPerforceConnection connection, LoginOptions options, string? user, string? password, string? host, CancellationToken cancellationToken = default)
 		{
-			List<string> arguments = new List<string>();
+			List<string> arguments = [];
 			if ((options & LoginOptions.AllHosts) != 0)
 			{
 				arguments.Add("-a");
@@ -2577,7 +2566,7 @@ namespace EpicGames.Perforce
 		/// <returns></returns>
 		public static async Task<PerforceResponse<LoginRecord>> TryGetLoginStateAsync(this IPerforceConnection connection, CancellationToken cancellationToken = default)
 		{
-			return await SingleResponseCommandAsync<LoginRecord>(connection, "login", new List<string> { "-s" }, null, cancellationToken);
+			return await SingleResponseCommandAsync<LoginRecord>(connection, "login", ["-s"], null, cancellationToken);
 		}
 
 		#endregion
@@ -2597,7 +2586,7 @@ namespace EpicGames.Perforce
 		/// <returns>List of records</returns>
 		public static async Task<PerforceResponseList<MergeRecord>> MergeAsync(this IPerforceConnection connection, MergeOptions options, int change, int maxFiles, string sourceFileSpec, string targetFileSpec, CancellationToken cancellationToken = default)
 		{
-			List<string> arguments = new List<string>();
+			List<string> arguments = [];
 			if ((options & MergeOptions.Preview) != 0)
 			{
 				arguments.Add($"-n");
@@ -2678,7 +2667,7 @@ namespace EpicGames.Perforce
 		/// <returns>Response from the server</returns>
 		public static async Task<PerforceResponseList<MoveRecord>> TryMoveAsync(this IPerforceConnection connection, int changeNumber, string? fileType, MoveOptions options, string sourceFileSpec, string targetFileSpec, CancellationToken cancellationToken = default)
 		{
-			List<string> arguments = new List<string>();
+			List<string> arguments = [];
 			if (changeNumber != -1)
 			{
 				arguments.Add($"-c{changeNumber}");
@@ -2769,7 +2758,7 @@ namespace EpicGames.Perforce
 		public static IAsyncEnumerable<PerforceResponse<OpenedRecord>> TryOpenedAsync(this IPerforceConnection connection, OpenedOptions options, int changeNumber, string? clientName, string? userName, int maxResults, FileSpecList fileSpecs, CancellationToken cancellationToken = default)
 		{
 			// Build the argument list
-			List<string> arguments = new List<string>();
+			List<string> arguments = [];
 			if ((options & OpenedOptions.AllWorkspaces) != 0)
 			{
 				arguments.Add($"-a");
@@ -2912,10 +2901,7 @@ namespace EpicGames.Perforce
 		static async Task<PerforceResponseList<PrintRecord>> TryPrintInternalAsync(this IPerforceConnection connection, string outputFile, string fileSpec, CancellationToken cancellationToken = default)
 		{
 			// Build the argument list
-			List<string> arguments = new List<string>();
-			arguments.Add("-o");
-			arguments.Add(outputFile);
-			arguments.Add(fileSpec);
+			List<string> arguments = ["-o", outputFile, fileSpec];
 
 			PerforceResponseList<PrintRecord> records = await CommandAsync<PrintRecord>(connection, "print", arguments, null, cancellationToken);
 			records.RemoveAll(x => x.Error != null && x.Error.Generic == PerforceGenericCode.Empty);
@@ -2933,10 +2919,7 @@ namespace EpicGames.Perforce
 		static Task<PerforceResponse<T>> TryPrintSingleInternalAsync<T>(this IPerforceConnection connection, string outputFile, string fileSpec, CancellationToken cancellationToken = default) where T : class
 		{
 			// Build the argument list
-			List<string> arguments = new List<string>();
-			arguments.Add("-o");
-			arguments.Add(outputFile);
-			arguments.Add(fileSpec);
+			List<string> arguments = ["-o", outputFile, fileSpec];
 			return SingleResponseCommandAsync<T>(connection, "print", arguments, null, cancellationToken);
 		}
 
@@ -2967,10 +2950,7 @@ namespace EpicGames.Perforce
 		/// <returns>Value of the property</returns>
 		public static async Task<PerforceResponse<PropertyRecord>> TryGetSinglePropertyAsync(this IPerforceConnection connection, string propertyName, CancellationToken cancellationToken = default)
 		{
-			List<string> arguments = new List<string>();
-			arguments.Add("-l");
-			arguments.Add("-n");
-			arguments.Add(propertyName);
+			List<string> arguments = ["-l", "-n", propertyName];
 
 			return await SingleResponseCommandAsync<PropertyRecord>(connection, "property", arguments, null, cancellationToken);
 		}
@@ -3006,7 +2986,7 @@ namespace EpicGames.Perforce
 		/// <returns>Response from the server</returns>
 		public static async Task<PerforceResponseList<ReconcileRecord>> TryReconcileAsync(this IPerforceConnection connection, int changeNumber, ReconcileOptions options, FileSpecList fileSpecs, CancellationToken cancellationToken = default)
 		{
-			List<string> arguments = new List<string>();
+			List<string> arguments = [];
 			if (changeNumber != -1)
 			{
 				arguments.Add($"-c{changeNumber}");
@@ -3079,7 +3059,7 @@ namespace EpicGames.Perforce
 		/// <returns>Response from the server</returns>
 		public static async Task<PerforceResponseList<ReopenRecord>> TryReopenAsync(this IPerforceConnection connection, int? changeNumber, string? fileType, string fileSpec, CancellationToken cancellationToken = default)
 		{
-			List<string> arguments = new List<string>();
+			List<string> arguments = [];
 			if (changeNumber != null)
 			{
 				if (changeNumber == PerforceReflection.DefaultChange)
@@ -3116,9 +3096,7 @@ namespace EpicGames.Perforce
 		/// <returns>Response from the server</returns>
 		public static Task<List<PerforceResponse>> ReloadClient(this IPerforceConnection connection, string clientName, string sourceServerId, CancellationToken cancellationToken = default)
 		{
-			List<string> arguments = new List<string>();
-			arguments.Add($"-c{clientName}");
-			arguments.Add($"-p{sourceServerId}");
+			List<string> arguments = [$"-c{clientName}", $"-p{sourceServerId}"];
 			return connection.CommandAsync("reload", arguments, null, null, cancellationToken);
 		}
 
@@ -3151,7 +3129,7 @@ namespace EpicGames.Perforce
 		/// <returns>Response from the server</returns>
 		public static async Task<PerforceResponseList<ResolveRecord>> TryResolveAsync(this IPerforceConnection connection, int changeNumber, ResolveOptions options, FileSpecList fileSpecs, CancellationToken cancellationToken = default)
 		{
-			List<string> arguments = new List<string>();
+			List<string> arguments = [];
 			if ((options & ResolveOptions.Automatic) != 0)
 			{
 				arguments.Add("-am");
@@ -3234,7 +3212,7 @@ namespace EpicGames.Perforce
 		/// <returns>Response from the server</returns>
 		public static async Task<PerforceResponseList<RevertRecord>> TryRevertAsync(this IPerforceConnection connection, int changeNumber, string? clientName, RevertOptions options, FileSpecList fileSpecs, CancellationToken cancellationToken = default)
 		{
-			List<string> arguments = new List<string>();
+			List<string> arguments = [];
 			if ((options & RevertOptions.Unchanged) != 0)
 			{
 				arguments.Add("-a");
@@ -3291,9 +3269,7 @@ namespace EpicGames.Perforce
 		/// <returns></returns>
 		public static Task<PerforceResponse<ServerRecord>> TryGetServerAsync(this IPerforceConnection connection, string serverId, CancellationToken cancellationToken = default)
 		{
-			List<string> arguments = new List<string>();
-			arguments.Add("-o");
-			arguments.Add(serverId);
+			List<string> arguments = ["-o", serverId];
 			return SingleResponseCommandAsync<ServerRecord>(connection, "server", arguments, null, cancellationToken);
 		}
 
@@ -3326,8 +3302,7 @@ namespace EpicGames.Perforce
 		/// <returns>Response from the server</returns>
 		public static async Task<PerforceResponseList<ShelveRecord>> TryShelveAsync(this IPerforceConnection connection, int changeNumber, ShelveOptions options, FileSpecList fileSpecs, CancellationToken cancellationToken = default)
 		{
-			List<string> arguments = new List<string>();
-			arguments.Add($"-c{changeNumber}");
+			List<string> arguments = [$"-c{changeNumber}"];
 			if ((options & ShelveOptions.OnlyChanged) != 0)
 			{
 				arguments.Add("-aleaveunchanged");
@@ -3372,8 +3347,7 @@ namespace EpicGames.Perforce
 		/// <returns>Response from the server</returns>
 		public static Task<List<PerforceResponse>> TryDeleteShelvedFilesAsync(this IPerforceConnection connection, int changeNumber, FileSpecList fileSpecs, CancellationToken cancellationToken = default)
 		{
-			List<string> arguments = new List<string>();
-			arguments.Add("-d");
+			List<string> arguments = ["-d"];
 			if (changeNumber != -1)
 			{
 				arguments.Add($"-c{changeNumber}");
@@ -3438,7 +3412,7 @@ namespace EpicGames.Perforce
 		/// <returns>List of response objects</returns>
 		public static async Task<PerforceResponseList<SizesRecord>> TrySizesAsync(this IPerforceConnection connection, SizesOptions options, int maxLines, FileSpecList fileSpecs, CancellationToken cancellationToken = default)
 		{
-			List<string> arguments = new List<string>();
+			List<string> arguments = [];
 			if ((options & SizesOptions.AllRevisions) != 0)
 			{
 				arguments.Add("-a");
@@ -3496,7 +3470,7 @@ namespace EpicGames.Perforce
 		/// <returns>Stream information record</returns>
 		public static Task<PerforceResponse<StreamRecord>> TryGetStreamAsync(this IPerforceConnection connection, string streamName, bool includeView, CancellationToken cancellationToken = default)
 		{
-			List<string> arguments = new List<string> { "-o" };
+			List<string> arguments = ["-o"];
 			if (includeView)
 			{
 				arguments.Add("-v");
@@ -3515,7 +3489,7 @@ namespace EpicGames.Perforce
 		/// <returns>Stream information record</returns>
 		public static Task<PerforceResponse> TryUpdateStreamAsync(this IPerforceConnection connection, StreamRecord record, CancellationToken cancellationToken = default)
 		{
-			List<string> arguments = new() { "-i" };
+			List<string> arguments = ["-i"];
 			return SingleResponseCommandAsync(connection, "stream", arguments, connection.SerializeRecord(record), null, cancellationToken);
 		}
 
@@ -3527,7 +3501,7 @@ namespace EpicGames.Perforce
 		/// <returns>Serialized record data</returns>
 		static byte[] SerializeRecord(this IPerforceConnection connection, StreamRecord input)
 		{
-			List<KeyValuePair<string, object>> nameToValue = new List<KeyValuePair<string, object>>();
+			List<KeyValuePair<string, object>> nameToValue = [];
 
 			void Add(string fieldName, string? value)
 			{
@@ -3687,7 +3661,7 @@ namespace EpicGames.Perforce
 		public static async Task<PerforceResponseList<StreamsRecord>> TryGetStreamsAsync(this IPerforceConnection connection, string? streamPath, int maxResults, string? filter, bool unloaded, CancellationToken cancellationToken = default)
 		{
 			// Build the command line
-			List<string> arguments = new List<string>();
+			List<string> arguments = [];
 			if (unloaded)
 			{
 				arguments.Add("-U");
@@ -3739,7 +3713,7 @@ namespace EpicGames.Perforce
 		/// <returns>Response from the server</returns>
 		public static async Task<PerforceResponse<SubmitRecord>> TrySubmitAsync(this IPerforceConnection connection, int changeNumber, SubmitOptions options, CancellationToken cancellationToken = default)
 		{
-			List<string> arguments = new List<string>();
+			List<string> arguments = [];
 			if ((options & SubmitOptions.ReopenAsEdit) != 0)
 			{
 				arguments.Add("-r");
@@ -3787,8 +3761,7 @@ namespace EpicGames.Perforce
 		/// <returns>Response from the server</returns>
 		public static async Task<PerforceResponse<SubmitRecord>> TrySubmitShelvedAsync(this IPerforceConnection connection, int changeNumber, CancellationToken cancellationToken = default)
 		{
-			List<string> arguments = new List<string>();
-			arguments.Add($"-e{changeNumber}");
+			List<string> arguments = [$"-e{changeNumber}"];
 
 			PerforceResponseList<SubmitRecord> responses = await CommandAsync<SubmitRecord>(connection, "submit", arguments, null, cancellationToken);
 			return ParseSubmitResponses(responses, connection.Logger);
@@ -4020,7 +3993,7 @@ namespace EpicGames.Perforce
 		/// <returns>Response from the server</returns>
 		public static async Task<PerforceResponseList<SyncSummaryRecord>> TrySyncQuietAsync(this IPerforceConnection connection, SyncOptions options, int maxFiles, int numThreads, int batch, int batchSize, int min, int minSize, FileSpecList fileSpecs, CancellationToken cancellationToken = default)
 		{
-			return new PerforceResponseList<SyncSummaryRecord>(await SyncInternalAsync<SyncSummaryRecord>(connection, options, maxFiles, numThreads, batch, batchSize, min, minSize, fileSpecs, true, cancellationToken).ToListAsync(cancellationToken));
+			return [.. await SyncInternalAsync<SyncSummaryRecord>(connection, options, maxFiles, numThreads, batch, batchSize, min, minSize, fileSpecs, true, cancellationToken).ToListAsync(cancellationToken)];
 		}
 
 		/// <summary>
@@ -4040,7 +4013,7 @@ namespace EpicGames.Perforce
 		/// <returns>Arguments for the command</returns>
 		private static IAsyncEnumerable<PerforceResponse<T>> SyncInternalAsync<T>(this IPerforceConnection connection, SyncOptions options, int maxFiles, int numThreads, int batch, int batchSize, int min, int minSize, FileSpecList fileSpecs, bool quiet, CancellationToken cancellationToken = default) where T : class
 		{
-			List<string> arguments = new List<string>();
+			List<string> arguments = [];
 			if ((options & SyncOptions.Force) != 0)
 			{
 				arguments.Add("-f");
@@ -4144,8 +4117,7 @@ namespace EpicGames.Perforce
 		/// <returns>Response from the server</returns>
 		public static async Task<PerforceResponseList<UnshelveRecord>> TryUnshelveAsync(this IPerforceConnection connection, int changeNumber, int intoChangeNumber, string? usingBranchSpec, string? usingStream, string? forceParentStream, UnshelveOptions options, FileSpecList fileSpecs, CancellationToken cancellationToken = default)
 		{
-			List<string> arguments = new List<string>();
-			arguments.Add($"-s{changeNumber}");
+			List<string> arguments = [$"-s{changeNumber}"];
 			if ((options & UnshelveOptions.ForceOverwrite) != 0)
 			{
 				arguments.Add("-f");
@@ -4202,9 +4174,7 @@ namespace EpicGames.Perforce
 		/// <returns>Response from the server</returns>
 		public static Task<PerforceResponseList<UserRecord>> TryGetUserAsync(this IPerforceConnection connection, string userName, CancellationToken cancellationToken = default)
 		{
-			List<string> arguments = new List<string>();
-			arguments.Add("-o");
-			arguments.Add(userName);
+			List<string> arguments = ["-o", userName];
 			return CommandAsync<UserRecord>(connection, "user", arguments, null, cancellationToken);
 		}
 
@@ -4235,7 +4205,7 @@ namespace EpicGames.Perforce
 		/// <returns>Response from the server</returns>
 		public static async Task<PerforceResponseList<UsersRecord>> TryGetUsersAsync(this IPerforceConnection connection, UsersOptions options, int maxResults, CancellationToken cancellationToken = default)
 		{
-			List<string> arguments = new List<string>();
+			List<string> arguments = [];
 			if ((options & UsersOptions.IncludeServiceUsers) != 0)
 			{
 				arguments.Add("-a");
@@ -4287,7 +4257,7 @@ namespace EpicGames.Perforce
 		/// <returns>List of responses from the server</returns>
 		public static IAsyncEnumerable<PerforceResponse<WhereRecord>> TryWhereAsync(this IPerforceConnection connection, FileSpecList fileSpecs, CancellationToken cancellationToken = default)
 		{
-			IAsyncEnumerable<PerforceResponse<WhereRecord>> records = StreamCommandAsync<WhereRecord>(connection, "where", new List<string>(), fileSpecs.List, null, cancellationToken);
+			IAsyncEnumerable<PerforceResponse<WhereRecord>> records = StreamCommandAsync<WhereRecord>(connection, "where", [], fileSpecs.List, null, cancellationToken);
 			return records.Where(x => x.Error == null || x.Error.Generic != PerforceGenericCode.Empty);
 		}
 
@@ -4318,9 +4288,7 @@ namespace EpicGames.Perforce
 		/// <returns>Response from the server</returns>
 		public static async Task<PerforceResponseList<UndoRecord>> TryUndoChangeAsync(this IPerforceConnection connection, int changeNumberToUndo, int changeNumber, CancellationToken cancellationToken = default)
 		{
-			List<string> arguments = new List<string>();
-			arguments.Add($"-c{changeNumber}");
-			arguments.Add($"//...@{changeNumberToUndo}");
+			List<string> arguments = [$"-c{changeNumber}", $"//...@{changeNumberToUndo}"];
 
 			return await CommandAsync<UndoRecord>(connection, "undo", arguments, null, cancellationToken);
 		}
@@ -4352,7 +4320,7 @@ namespace EpicGames.Perforce
 		/// <returns>List of annotate records</returns>
 		public static Task<PerforceResponseList<AnnotateRecord>> TryAnnotateAsync(this IPerforceConnection connection, string fileSpec, AnnotateOptions options, CancellationToken cancellationToken = default)
 		{
-			List<string> arguments = new List<string>();
+			List<string> arguments = [];
 			if ((options & AnnotateOptions.IncludeDeletedFilesAndLines) != 0)
 			{
 				arguments.Add("-a");
diff --git a/Engine/Source/Programs/Shared/EpicGames.Perforce/PerforceEnvironment.cs b/Engine/Source/Programs/Shared/EpicGames.Perforce/PerforceEnvironment.cs
--- a/Engine/Source/Programs/Shared/EpicGames.Perforce/PerforceEnvironment.cs
+++ b/Engine/Source/Programs/Shared/EpicGames.Perforce/PerforceEnvironment.cs
@@ -154,7 +154,7 @@ namespace EpicGames.Perforce
 		/// </summary>
 		public static IPerforceEnvironment Default { get; } = CreateDefaultEnvironment();
 
-		static readonly Dictionary<DirectoryReference, IPerforceEnvironment> s_directoryToEnvironment = new Dictionary<DirectoryReference, IPerforceEnvironment>();
+		static readonly Dictionary<DirectoryReference, IPerforceEnvironment> s_directoryToEnvironment = [];
 
 		static IPerforceEnvironment CreateDefaultEnvironment()
 		{
diff --git a/Engine/Source/Programs/Shared/EpicGames.Perforce/PerforceMetadataLogger.cs b/Engine/Source/Programs/Shared/EpicGames.Perforce/PerforceMetadataLogger.cs
--- a/Engine/Source/Programs/Shared/EpicGames.Perforce/PerforceMetadataLogger.cs
+++ b/Engine/Source/Programs/Shared/EpicGames.Perforce/PerforceMetadataLogger.cs
@@ -30,7 +30,7 @@ namespace EpicGames.Perforce
 		}
 
 		readonly ILogger _inner;
-		readonly List<ClientView> _clients = new List<ClientView>();
+		readonly List<ClientView> _clients = [];
 
 		/// <summary>
 		/// Accessor for the inner logger
@@ -160,8 +160,8 @@ namespace EpicGames.Perforce
 			Utf8JsonReader reader = new Utf8JsonReader(data.Span);
 			for (bool valid = ReadFirstLogProperty(ref reader); valid; valid = ReadNextLogProperty(ref reader))
 			{
-				ReadOnlySpan<byte> type = ReadOnlySpan<byte>.Empty;
-				ReadOnlySpan<byte> text = ReadOnlySpan<byte>.Empty;
+				ReadOnlySpan<byte> type = [];
+				ReadOnlySpan<byte> text = [];
 				string? file = null;
 
 				// Get the $type and $text properties
diff --git a/Engine/Source/Programs/Shared/EpicGames.Perforce/PerforceRecord.cs b/Engine/Source/Programs/Shared/EpicGames.Perforce/PerforceRecord.cs
--- a/Engine/Source/Programs/Shared/EpicGames.Perforce/PerforceRecord.cs
+++ b/Engine/Source/Programs/Shared/EpicGames.Perforce/PerforceRecord.cs
@@ -234,7 +234,7 @@ namespace EpicGames.Perforce
 		/// <summary>
 		/// The rows in this record
 		/// </summary>
-		public List<KeyValuePair<Utf8String, PerforceValue>> Rows { get; } = new List<KeyValuePair<Utf8String, PerforceValue>>();
+		public List<KeyValuePair<Utf8String, PerforceValue>> Rows { get; } = [];
 
 		/// <summary>
 		/// Create a record from a set of fields
