Shared: Enable IDE0028, IDE0031, IDE0034, IDE0054, IDE0062, IDE0071, IDE0074, IDE0075, IDE0280, IDE0300, IDE0301 IDE0305

diff --git a/Engine/Source/Programs/Shared/EpicGames.Perforce/PerforceReflection.cs b/Engine/Source/Programs/Shared/EpicGames.Perforce/PerforceReflection.cs
--- a/Engine/Source/Programs/Shared/EpicGames.Perforce/PerforceReflection.cs
+++ b/Engine/Source/Programs/Shared/EpicGames.Perforce/PerforceReflection.cs
@@ -141,12 +141,12 @@ namespace EpicGames.Perforce
 		/// <summary>
 		/// List of fields in the record. These should be ordered to match P4 output for maximum efficiency.
 		/// </summary>
-		public List<TaggedPropertyInfo> Properties { get; } = new List<TaggedPropertyInfo>();
+		public List<TaggedPropertyInfo> Properties { get; } = [];
 
 		/// <summary>
 		/// Map of name to tag info
 		/// </summary>
-		public Dictionary<TaggedPropertyNameAndRank, TaggedPropertyInfo> NameAndRankToInfo { get; set; } = new Dictionary<TaggedPropertyNameAndRank, TaggedPropertyInfo>();
+		public Dictionary<TaggedPropertyNameAndRank, TaggedPropertyInfo> NameAndRankToInfo { get; set; } = [];
 
 		/// <summary>
 		/// Bitmask of all the required tags. Formed by bitwise-or'ing the RequiredTagBitMask fields for each required CachedTagInfo.
@@ -183,17 +183,17 @@ namespace EpicGames.Perforce
 		/// <summary>
 		/// Map of name to value
 		/// </summary>
-		public Dictionary<Utf8String, int> _nameToValue = new Dictionary<Utf8String, int>();
+		public Dictionary<Utf8String, int> _nameToValue = [];
 
 		/// <summary>
 		/// Map of value to name
 		/// </summary>
-		public Dictionary<int, Utf8String> _valueToName = new Dictionary<int, Utf8String>();
+		public Dictionary<int, Utf8String> _valueToName = [];
 
 		/// <summary>
 		/// List of name/value pairs
 		/// </summary>
-		public List<KeyValuePair<string, int>> _nameValuePairs = new List<KeyValuePair<string, int>>();
+		public List<KeyValuePair<string, int>> _nameValuePairs = [];
 
 		/// <summary>
 		/// Constructor
@@ -306,7 +306,7 @@ namespace EpicGames.Perforce
 		{
 			if (_bHasFlagsAttribute)
 			{
-				List<string> names = new List<string>();
+				List<string> names = [];
 
 				int combinedIntegerValue = 0;
 				foreach (KeyValuePair<string, int> pair in _nameValuePairs)
@@ -472,7 +472,7 @@ namespace EpicGames.Perforce
 				record = new CachedRecordInfo(recordType, GetCreateRecordDelegate(recordType));
 
 				// Find all the properties in the record
-				AddRecordProperties(recordType, record, Array.Empty<NestedRecordInfo>());
+				AddRecordProperties(recordType, record, []);
 				record.NameAndRankToInfo = record.Properties.ToDictionary(x => new TaggedPropertyNameAndRank(x.Tag, x.ParentRecords.Length), x => x);
 
 				// Try to save the record info, or get the version that's already in the cache
@@ -580,7 +580,7 @@ namespace EpicGames.Perforce
 				{
 					Type newRecordType = property.PropertyType.GenericTypeArguments[0];
 					NestedRecordInfo newParentRecord = new NestedRecordInfo(property, GetCreateRecordDelegate(newRecordType));
-					AddRecordProperties(newRecordType, rootRecord, parentRecords.Append(newParentRecord).ToArray());
+					AddRecordProperties(newRecordType, rootRecord, [.. parentRecords, newParentRecord]);
 				}
 			}
 		}
diff --git a/Engine/Source/Programs/Shared/EpicGames.Perforce/PerforceResponseList.cs b/Engine/Source/Programs/Shared/EpicGames.Perforce/PerforceResponseList.cs
--- a/Engine/Source/Programs/Shared/EpicGames.Perforce/PerforceResponseList.cs
+++ b/Engine/Source/Programs/Shared/EpicGames.Perforce/PerforceResponseList.cs
@@ -50,10 +50,10 @@ namespace EpicGames.Perforce
 					PerforceError? error = this[0].Error;
 					if (error != null && error.Generic == PerforceGenericCode.Empty)
 					{
-						return new List<T>();
+						return [];
 					}
 				}
-				return this.Where(x => x.Info == null).Select(x => x.Data).Where(x => x != null).ToList();
+				return [.. this.Where(x => x.Info == null).Select(x => x.Data).Where(x => x != null)];
 			}
 		}
 
diff --git a/Engine/Source/Programs/Shared/EpicGames.Perforce/PerforceUtils.cs b/Engine/Source/Programs/Shared/EpicGames.Perforce/PerforceUtils.cs
--- a/Engine/Source/Programs/Shared/EpicGames.Perforce/PerforceUtils.cs
+++ b/Engine/Source/Programs/Shared/EpicGames.Perforce/PerforceUtils.cs
@@ -16,8 +16,8 @@ namespace EpicGames.Perforce
 		/// <summary>
 		/// Set of extensions to treat as code
 		/// </summary>
-		public static readonly HashSet<string> CodeExtensions = new HashSet<string>
-		{
+		public static readonly HashSet<string> CodeExtensions =
+		[
 			".c",
 			".cc",
 			".cpp",
@@ -36,7 +36,7 @@ namespace EpicGames.Perforce
 			".uplugin",
 			".sln",
 			".native.verse"
-		};
+		];
 
 		/// <summary>
 		/// Tests if a path is a code file
diff --git a/Engine/Source/Programs/Shared/EpicGames.Perforce/PerforceViewFilter.cs b/Engine/Source/Programs/Shared/EpicGames.Perforce/PerforceViewFilter.cs
--- a/Engine/Source/Programs/Shared/EpicGames.Perforce/PerforceViewFilter.cs
+++ b/Engine/Source/Programs/Shared/EpicGames.Perforce/PerforceViewFilter.cs
@@ -24,7 +24,7 @@ namespace EpicGames.Perforce
 		/// <param name="entries"></param>
 		public PerforceViewFilter(IEnumerable<PerforceViewFilterEntry> entries)
 		{
-			Entries = entries.ToList();
+			Entries = [.. entries];
 
 			// When first entry is excluding, start by including everything. Otherwise, default is to exclude.
 			// This is to emulate the behavior of ManagedWorkspace.UpdateClientHaveTableAsync()
diff --git a/Engine/Source/Programs/Shared/EpicGames.Perforce/PerforceViewMap.cs b/Engine/Source/Programs/Shared/EpicGames.Perforce/PerforceViewMap.cs
--- a/Engine/Source/Programs/Shared/EpicGames.Perforce/PerforceViewMap.cs
+++ b/Engine/Source/Programs/Shared/EpicGames.Perforce/PerforceViewMap.cs
@@ -23,7 +23,7 @@ namespace EpicGames.Perforce
 		/// </summary>
 		public PerforceViewMap()
 		{
-			Entries = new List<PerforceViewMapEntry>();
+			Entries = [];
 		}
 
 		/// <summary>
@@ -32,7 +32,7 @@ namespace EpicGames.Perforce
 		/// <param name="entries"></param>
 		public PerforceViewMap(IEnumerable<PerforceViewMapEntry> entries)
 		{
-			Entries = entries.ToList();
+			Entries = [.. entries];
 		}
 
 		/// <summary>
@@ -41,7 +41,7 @@ namespace EpicGames.Perforce
 		/// <param name="other"></param>
 		public PerforceViewMap(PerforceViewMap other)
 		{
-			Entries = new List<PerforceViewMapEntry>(other.Entries);
+			Entries = [.. other.Entries];
 		}
 
 		/// <summary>
@@ -59,7 +59,7 @@ namespace EpicGames.Perforce
 		/// <returns></returns>
 		public PerforceViewMap Invert()
 		{
-			List<PerforceViewMapEntry> entries = new List<PerforceViewMapEntry>();
+			List<PerforceViewMapEntry> entries = [];
 			foreach (PerforceViewMapEntry entry in Entries)
 			{
 				entries.Add(new PerforceViewMapEntry(entry.Include, entry.Target, entry.Source));
@@ -139,7 +139,7 @@ namespace EpicGames.Perforce
 		/// <returns></returns>
 		public List<string> GetRootPaths(StringComparison comparison)
 		{
-			List<string> rootPaths = new List<string>();
+			List<string> rootPaths = [];
 			foreach (PerforceViewMapEntry entry in Entries)
 			{
 				if (entry.Include)
diff --git a/Engine/Source/Programs/Shared/EpicGames.Perforce/Records/AnnotateRecord.cs b/Engine/Source/Programs/Shared/EpicGames.Perforce/Records/AnnotateRecord.cs
--- a/Engine/Source/Programs/Shared/EpicGames.Perforce/Records/AnnotateRecord.cs
+++ b/Engine/Source/Programs/Shared/EpicGames.Perforce/Records/AnnotateRecord.cs
@@ -68,7 +68,7 @@ namespace EpicGames.Perforce
 		/// 
 		/// </summary>
 		[PerforceRecordList]
-		public List<AnnotateLineRecord> LineRecord { get; } = new List<AnnotateLineRecord>();
+		public List<AnnotateLineRecord> LineRecord { get; } = [];
 
 		/// <summary>
 		/// Private constructor for serialization
diff --git a/Engine/Source/Programs/Shared/EpicGames.Perforce/Records/ChangeRecord.cs b/Engine/Source/Programs/Shared/EpicGames.Perforce/Records/ChangeRecord.cs
--- a/Engine/Source/Programs/Shared/EpicGames.Perforce/Records/ChangeRecord.cs
+++ b/Engine/Source/Programs/Shared/EpicGames.Perforce/Records/ChangeRecord.cs
@@ -56,6 +56,6 @@ namespace EpicGames.Perforce
 		/// Files that are open in this changelist
 		/// </summary>
 		[PerforceTag("Files", Optional = true)]
-		public List<string> Files { get; } = new List<string>();
+		public List<string> Files { get; } = [];
 	}
 }
diff --git a/Engine/Source/Programs/Shared/EpicGames.Perforce/Records/ClientRecord.cs b/Engine/Source/Programs/Shared/EpicGames.Perforce/Records/ClientRecord.cs
--- a/Engine/Source/Programs/Shared/EpicGames.Perforce/Records/ClientRecord.cs
+++ b/Engine/Source/Programs/Shared/EpicGames.Perforce/Records/ClientRecord.cs
@@ -92,7 +92,7 @@ namespace EpicGames.Perforce
 		/// Specifies the mappings between files in the depot and files in the workspace.
 		/// </summary>
 		[PerforceTag("View")]
-		public List<string> View { get; } = new List<string>();
+		public List<string> View { get; } = [];
 
 		/// <summary>
 		/// Private constructor for serialization
diff --git a/Engine/Source/Programs/Shared/EpicGames.Perforce/Records/DescribeRecord.cs b/Engine/Source/Programs/Shared/EpicGames.Perforce/Records/DescribeRecord.cs
--- a/Engine/Source/Programs/Shared/EpicGames.Perforce/Records/DescribeRecord.cs
+++ b/Engine/Source/Programs/Shared/EpicGames.Perforce/Records/DescribeRecord.cs
@@ -62,7 +62,7 @@ namespace EpicGames.Perforce
 		/// The files affected by this change
 		/// </summary>
 		[PerforceRecordList]
-		public List<DescribeFileRecord> Files { get; } = new List<DescribeFileRecord>();
+		public List<DescribeFileRecord> Files { get; } = [];
 
 		/// <summary>
 		/// Private constructor for serialization
diff --git a/Engine/Source/Programs/Shared/EpicGames.Perforce/Records/FileLogRecord.cs b/Engine/Source/Programs/Shared/EpicGames.Perforce/Records/FileLogRecord.cs
--- a/Engine/Source/Programs/Shared/EpicGames.Perforce/Records/FileLogRecord.cs
+++ b/Engine/Source/Programs/Shared/EpicGames.Perforce/Records/FileLogRecord.cs
@@ -19,7 +19,7 @@ namespace EpicGames.Perforce
 		/// Revisions of this file
 		/// </summary>
 		[PerforceRecordList]
-		public List<RevisionRecord> Revisions { get; } = new List<RevisionRecord>();
+		public List<RevisionRecord> Revisions { get; } = [];
 
 		/// <summary>
 		/// Private constructor for serialization
diff --git a/Engine/Source/Programs/Shared/EpicGames.Perforce/Records/RevisionRecord.cs b/Engine/Source/Programs/Shared/EpicGames.Perforce/Records/RevisionRecord.cs
--- a/Engine/Source/Programs/Shared/EpicGames.Perforce/Records/RevisionRecord.cs
+++ b/Engine/Source/Programs/Shared/EpicGames.Perforce/Records/RevisionRecord.cs
@@ -74,7 +74,7 @@ namespace EpicGames.Perforce
 		/// Integration records for this revision
 		/// </summary>
 		[PerforceRecordList]
-		public List<IntegrationRecord> Integrations { get; } = new List<IntegrationRecord>();
+		public List<IntegrationRecord> Integrations { get; } = [];
 
 		/// <summary>
 		/// Private constructor for serialization
diff --git a/Engine/Source/Programs/Shared/EpicGames.Perforce/Records/StreamRecord.cs b/Engine/Source/Programs/Shared/EpicGames.Perforce/Records/StreamRecord.cs
--- a/Engine/Source/Programs/Shared/EpicGames.Perforce/Records/StreamRecord.cs
+++ b/Engine/Source/Programs/Shared/EpicGames.Perforce/Records/StreamRecord.cs
@@ -74,19 +74,19 @@ namespace EpicGames.Perforce
 		/// List of paths in the stream spec
 		/// </summary>
 		[PerforceTag("Paths")]
-		public List<string> Paths { get; } = new List<string>();
+		public List<string> Paths { get; } = [];
 
 		/// <summary>
 		/// Computed view for the stream
 		/// </summary>
 		[PerforceTag("View", Optional = true)]
-		public List<string> View { get; } = new List<string>();
+		public List<string> View { get; } = [];
 
 		/// <summary>
 		/// Specific changelists to import
 		/// </summary>
 		[PerforceTag("ChangeView", Optional = true)]
-		public List<string> ChangeView { get; } = new List<string>();
+		public List<string> ChangeView { get; } = [];
 
 		/// <summary>
 		/// Private constructor for serialization
diff --git a/Engine/Source/Programs/Shared/EpicGames.Perforce/Records/UserRecord.cs b/Engine/Source/Programs/Shared/EpicGames.Perforce/Records/UserRecord.cs
--- a/Engine/Source/Programs/Shared/EpicGames.Perforce/Records/UserRecord.cs
+++ b/Engine/Source/Programs/Shared/EpicGames.Perforce/Records/UserRecord.cs
@@ -44,7 +44,7 @@ namespace EpicGames.Perforce
 		/// Paths which the user is watching
 		/// </summary>
 		[PerforceTag("Reviews", Optional = true)]
-		public List<string> Reviews { get; } = new List<string>();
+		public List<string> Reviews { get; } = [];
 
 		/// <summary>
 		/// The type of user
diff --git a/Engine/Source/Programs/Shared/EpicGames.ProjectStore/CloudStorage/CloudStorage.cs b/Engine/Source/Programs/Shared/EpicGames.ProjectStore/CloudStorage/CloudStorage.cs
--- a/Engine/Source/Programs/Shared/EpicGames.ProjectStore/CloudStorage/CloudStorage.cs
+++ b/Engine/Source/Programs/Shared/EpicGames.ProjectStore/CloudStorage/CloudStorage.cs
@@ -145,20 +145,20 @@ namespace EpicGames.ProjectStore
 		}
 
 		/// <inheritdoc />
-		public Task DownloadBuildAsync(StringId bucketId, CbObjectId buildId, DirectoryReference outputDir, List<FilePattern>? includeFiles = null, List<FilePattern>? excludeFiles = null, CloudBuildDownloadOptions? downloadOptions = null, IProgress<(float, string)>? progress = null, CloudBuildOptionsOverride? optionsOverride = null, CancellationToken cancellationToken = default(CancellationToken))
+		public Task DownloadBuildAsync(StringId bucketId, CbObjectId buildId, DirectoryReference outputDir, List<FilePattern>? includeFiles = null, List<FilePattern>? excludeFiles = null, CloudBuildDownloadOptions? downloadOptions = null, IProgress<(float, string)>? progress = null, CloudBuildOptionsOverride? optionsOverride = null, CancellationToken cancellationToken = default)
 		{
 			NamespaceId? ns = _options.Value.NamespaceId != null ? new NamespaceId(_options.Value.NamespaceId) : null;
 			return DownloadBuildAsync(ns, bucketId, buildId, outputDir, includeFiles, excludeFiles, downloadOptions, progress, optionsOverride, cancellationToken);
 		}
 
 		/// <inheritdoc />
-		public void DownloadBuild(StringId bucketId, CbObjectId buildId, DirectoryReference outputDir, List<FilePattern>? includeFiles = null, List<FilePattern>? excludeFiles = null, CloudBuildDownloadOptions? downloadOptions = null, IProgress<(float, string)>? progress = null, CloudBuildOptionsOverride? optionsOverride = null, CancellationToken cancellationToken = default(CancellationToken))
+		public void DownloadBuild(StringId bucketId, CbObjectId buildId, DirectoryReference outputDir, List<FilePattern>? includeFiles = null, List<FilePattern>? excludeFiles = null, CloudBuildDownloadOptions? downloadOptions = null, IProgress<(float, string)>? progress = null, CloudBuildOptionsOverride? optionsOverride = null, CancellationToken cancellationToken = default)
 		{
 			DownloadBuildAsync(bucketId, buildId, outputDir, includeFiles, excludeFiles, downloadOptions, progress, optionsOverride, cancellationToken).Wait(cancellationToken);
 		}
 
 		/// <inheritdoc />
-		public Task DownloadBuildFromUriAsync(Uri uri, DirectoryReference outputDir, List<FilePattern>? includeFiles = null, List<FilePattern>? excludeFiles = null, CloudBuildDownloadOptions? downloadOptions = null, IProgress<(float, string)>? progress = null, CloudBuildOptionsOverride? optionsOverride = null, CancellationToken cancellationToken = default(CancellationToken))
+		public Task DownloadBuildFromUriAsync(Uri uri, DirectoryReference outputDir, List<FilePattern>? includeFiles = null, List<FilePattern>? excludeFiles = null, CloudBuildDownloadOptions? downloadOptions = null, IProgress<(float, string)>? progress = null, CloudBuildOptionsOverride? optionsOverride = null, CancellationToken cancellationToken = default)
 		{
 			string path = uri.AbsolutePath;
 			// remove api path if specified
@@ -179,7 +179,7 @@ namespace EpicGames.ProjectStore
 			return DownloadBuildAsync(new NamespaceId(namespaceId), new StringId(bucketId), CbObjectId.Parse(buildId), outputDir, includeFiles, excludeFiles, downloadOptions, progress, optionsOverride, cancellationToken);
 		}
 
-		private async Task DownloadBuildAsync(NamespaceId? namespaceId, StringId bucketId, CbObjectId buildId, DirectoryReference outputDir, List<FilePattern>? includeFiles, List<FilePattern>? excludeFiles, CloudBuildDownloadOptions? downloadOptions = null, IProgress<(float, string)>? progress = null, CloudBuildOptionsOverride? optionsOverride = null, CancellationToken cancellationToken = default(CancellationToken))
+		private async Task DownloadBuildAsync(NamespaceId? namespaceId, StringId bucketId, CbObjectId buildId, DirectoryReference outputDir, List<FilePattern>? includeFiles, List<FilePattern>? excludeFiles, CloudBuildDownloadOptions? downloadOptions = null, IProgress<(float, string)>? progress = null, CloudBuildOptionsOverride? optionsOverride = null, CancellationToken cancellationToken = default)
 		{
 			NamespaceId? ns = optionsOverride?.NamespaceId ?? namespaceId;
 			if (ns == null)
@@ -207,7 +207,7 @@ namespace EpicGames.ProjectStore
 
 		private static string ToZenWildcard(DirectoryReference outputDirectory, FilePattern filePattern)
 		{
-			IEnumerable<string> GetTokens(FilePattern pattern)
+			static IEnumerable<string> GetTokens(FilePattern pattern)
 			{
 				foreach (string token in pattern.Tokens)
 				{
@@ -226,7 +226,7 @@ namespace EpicGames.ProjectStore
 		}
 
 		/// <inheritdoc />
-		public IEnumerable<FindBuildResponse> FindBuilds(Regex bucketRegex, CbObject query, CloudBuildSearchOptions? searchOptions = null, CloudBuildOptionsOverride? optionsOverride = null, CancellationToken cancellationToken = default(CancellationToken))
+		public IEnumerable<FindBuildResponse> FindBuilds(Regex bucketRegex, CbObject query, CloudBuildSearchOptions? searchOptions = null, CloudBuildOptionsOverride? optionsOverride = null, CancellationToken cancellationToken = default)
 		{
 			string? ns = GetCurrentNamespace(optionsOverride);
 			if (ns == null)
@@ -341,7 +341,7 @@ namespace EpicGames.ProjectStore
 		}
 
 		/// <inheritdoc />
-		public async IAsyncEnumerable<FindBuildResponse> FindBuildsAsync(Regex bucketRegex, CbObject query, CloudBuildSearchOptions? searchOptions = null, CloudBuildOptionsOverride? optionsOverride = null, [EnumeratorCancellation] CancellationToken cancellationToken = default(CancellationToken))
+		public async IAsyncEnumerable<FindBuildResponse> FindBuildsAsync(Regex bucketRegex, CbObject query, CloudBuildSearchOptions? searchOptions = null, CloudBuildOptionsOverride? optionsOverride = null, [EnumeratorCancellation] CancellationToken cancellationToken = default)
 		{
 			await Task.CompletedTask;
 
@@ -368,7 +368,7 @@ namespace EpicGames.ProjectStore
 		}
 
 		/// <inheritdoc />
-		public IAsyncEnumerable<FindBuildResponse> FindBuildsAsync(string project, string? type, string? branch, CbObject query, CloudBuildSearchOptions? searchOptions = null, CloudBuildOptionsOverride? optionsOverride = null, CancellationToken cancellationToken = default(CancellationToken))
+		public IAsyncEnumerable<FindBuildResponse> FindBuildsAsync(string project, string? type, string? branch, CbObject query, CloudBuildSearchOptions? searchOptions = null, CloudBuildOptionsOverride? optionsOverride = null, CancellationToken cancellationToken = default)
 		{
 			Regex regex = BuildRegex(project, type, branch);
 
@@ -376,7 +376,7 @@ namespace EpicGames.ProjectStore
 		}
 
 		/// <inheritdoc />
-		public IEnumerable<FindBuildResponse> FindBuilds(string project, string? type, string? branch, CbObject query, CloudBuildSearchOptions? searchOptions = null, CloudBuildOptionsOverride? optionsOverride = null, CancellationToken cancellationToken = default(CancellationToken))
+		public IEnumerable<FindBuildResponse> FindBuilds(string project, string? type, string? branch, CbObject query, CloudBuildSearchOptions? searchOptions = null, CloudBuildOptionsOverride? optionsOverride = null, CancellationToken cancellationToken = default)
 		{
 			Regex regex = BuildRegex(project, type, branch);
 
@@ -394,7 +394,7 @@ namespace EpicGames.ProjectStore
 		}
 
 		/// <inheritdoc />
-		public async Task<List<BuildContents>> ListBuildContentsAsync(StringId bucketId, CbObjectId buildId, CloudBuildOptionsOverride? optionsOverride = null, CancellationToken cancellationToken = default(CancellationToken))
+		public async Task<List<BuildContents>> ListBuildContentsAsync(StringId bucketId, CbObjectId buildId, CloudBuildOptionsOverride? optionsOverride = null, CancellationToken cancellationToken = default)
 		{
 			Uri uri = new Uri($"{_options.Value.Host}/api/v2/builds/{GetCurrentNamespace(optionsOverride)}/{bucketId}/{buildId}");
 			string cmdline = $"builds ls --cloud-url={uri} --quiet";
@@ -405,7 +405,7 @@ namespace EpicGames.ProjectStore
 
 			// TODO: we need structured output from zen for this listing
 
-			Dictionary<string, FileDescription> files = new Dictionary<string, FileDescription>();
+			Dictionary<string, FileDescription> files = [];
 			/*foreach (string line in stdout)
 			{
 				string[] segments = line.Split('\t');
@@ -463,7 +463,7 @@ namespace EpicGames.ProjectStore
 		}
 
 		/// <inheritdoc />
-		public async Task ImportSnapshotAsync(string localFileName, string localRootPath, string engineDir, string cookedPlatformDir, bool forceImport = false, bool asyncImport = false, CancellationToken cancellationToken = default(CancellationToken))
+		public async Task ImportSnapshotAsync(string localFileName, string localRootPath, string engineDir, string cookedPlatformDir, bool forceImport = false, bool asyncImport = false, CancellationToken cancellationToken = default)
 		{
 			/*string projectId = GetProjectID(projectFile);
 			CreateProject(projectFile, rootDirectory, engineDirectory);
@@ -518,7 +518,7 @@ namespace EpicGames.ProjectStore
 			{
 				cmdline += " --access-token-env UE_CloudDataCacheAccessToken ";
 
-				env = new Dictionary<string, string>();
+				env = [];
 				// forward the current environment and then override the token
 				foreach (DictionaryEntry entry in Environment.GetEnvironmentVariables())
 				{
diff --git a/Engine/Source/Programs/Shared/EpicGames.ProjectStore/CloudStorage/ICloudStorage.cs b/Engine/Source/Programs/Shared/EpicGames.ProjectStore/CloudStorage/ICloudStorage.cs
--- a/Engine/Source/Programs/Shared/EpicGames.ProjectStore/CloudStorage/ICloudStorage.cs
+++ b/Engine/Source/Programs/Shared/EpicGames.ProjectStore/CloudStorage/ICloudStorage.cs
@@ -278,7 +278,7 @@ namespace EpicGames.ProjectStore
 		/// <param name="optionsOverride">Optional per invocation overrides of options</param>
 		/// <param name="cancellationToken">A cancellation token</param>
 		/// <returns></returns>
-		Task DownloadBuildAsync(StringId bucketId, CbObjectId buildId, DirectoryReference outputDir, List<FilePattern>? includeFiles = null, List<FilePattern>? excludeFiles = null, CloudBuildDownloadOptions? downloadOptions = null, IProgress<(float, string)>? progress = null, CloudBuildOptionsOverride? optionsOverride = null, CancellationToken cancellationToken = default(CancellationToken));
+		Task DownloadBuildAsync(StringId bucketId, CbObjectId buildId, DirectoryReference outputDir, List<FilePattern>? includeFiles = null, List<FilePattern>? excludeFiles = null, CloudBuildDownloadOptions? downloadOptions = null, IProgress<(float, string)>? progress = null, CloudBuildOptionsOverride? optionsOverride = null, CancellationToken cancellationToken = default);
 
 		/// <summary>
 		/// Download a build using all the required fields
@@ -293,7 +293,7 @@ namespace EpicGames.ProjectStore
 		/// <param name="optionsOverride">Optional per invocation overrides of options</param>
 		/// <param name="cancellationToken">A cancellation token</param>
 		/// <returns></returns>
-		void DownloadBuild(StringId bucketId, CbObjectId buildId, DirectoryReference outputDir, List<FilePattern>? includeFiles = null, List<FilePattern>? excludeFiles = null, CloudBuildDownloadOptions? downloadOptions = null, IProgress<(float, string)>? progress = null, CloudBuildOptionsOverride? optionsOverride = null, CancellationToken cancellationToken = default(CancellationToken));
+		void DownloadBuild(StringId bucketId, CbObjectId buildId, DirectoryReference outputDir, List<FilePattern>? includeFiles = null, List<FilePattern>? excludeFiles = null, CloudBuildDownloadOptions? downloadOptions = null, IProgress<(float, string)>? progress = null, CloudBuildOptionsOverride? optionsOverride = null, CancellationToken cancellationToken = default);
 
 		/// <summary>
 		/// Download a build using a single uri value that encodes everything needed to locate the build
@@ -307,7 +307,7 @@ namespace EpicGames.ProjectStore
 		/// <param name="optionsOverride">Optional per invocation overrides of options</param>		
 		/// <param name="cancellationToken">A cancellation token</param>
 		/// <returns></returns>
-		Task DownloadBuildFromUriAsync(Uri uri, DirectoryReference outputDir, List<FilePattern>? includeFiles = null, List<FilePattern>? excludeFiles = null, CloudBuildDownloadOptions? downloadOptions = null, IProgress<(float, string)>? progress = null, CloudBuildOptionsOverride? optionsOverride = null, CancellationToken cancellationToken = default(CancellationToken));
+		Task DownloadBuildFromUriAsync(Uri uri, DirectoryReference outputDir, List<FilePattern>? includeFiles = null, List<FilePattern>? excludeFiles = null, CloudBuildDownloadOptions? downloadOptions = null, IProgress<(float, string)>? progress = null, CloudBuildOptionsOverride? optionsOverride = null, CancellationToken cancellationToken = default);
 
 		/// <summary>
 		/// Search for builds in UE Cloud Storage using a regular expression to describe which buckets the build might be in
@@ -318,7 +318,7 @@ namespace EpicGames.ProjectStore
 		/// <param name="optionsOverride">Optional per invocation overrides of options</param>		
 		/// <param name="cancellationToken">A cancellation token</param>
 		/// <returns></returns>
-		IAsyncEnumerable<FindBuildResponse> FindBuildsAsync(Regex bucketRegex, CbObject query, CloudBuildSearchOptions? searchOptions = null, CloudBuildOptionsOverride? optionsOverride = null, CancellationToken cancellationToken = default(CancellationToken));
+		IAsyncEnumerable<FindBuildResponse> FindBuildsAsync(Regex bucketRegex, CbObject query, CloudBuildSearchOptions? searchOptions = null, CloudBuildOptionsOverride? optionsOverride = null, CancellationToken cancellationToken = default);
 
 		/// <summary>
 		/// Search for builds in UE Cloud Storage using a regular expression to describe which buckets the build might be in
@@ -329,7 +329,7 @@ namespace EpicGames.ProjectStore
 		/// <param name="optionsOverride">Optional per invocation overrides of options</param>		
 		/// <param name="cancellationToken">A cancellation token</param>
 		/// <returns></returns>
-		IEnumerable<FindBuildResponse> FindBuilds(Regex bucketRegex, CbObject query, CloudBuildSearchOptions? searchOptions = null, CloudBuildOptionsOverride? optionsOverride = null, CancellationToken cancellationToken = default(CancellationToken));
+		IEnumerable<FindBuildResponse> FindBuilds(Regex bucketRegex, CbObject query, CloudBuildSearchOptions? searchOptions = null, CloudBuildOptionsOverride? optionsOverride = null, CancellationToken cancellationToken = default);
 
 		/// <summary>
 		/// Search for builds in UE Cloud Storage using common components
@@ -342,7 +342,7 @@ namespace EpicGames.ProjectStore
 		/// <param name="optionsOverride">Optional per invocation overrides of options</param>		
 		/// <param name="cancellationToken">A cancellation token</param>
 		/// <returns></returns>
-		IAsyncEnumerable<FindBuildResponse> FindBuildsAsync(string project, string? type, string? branch, CbObject query, CloudBuildSearchOptions? searchOptions = null, CloudBuildOptionsOverride? optionsOverride = null, CancellationToken cancellationToken = default(CancellationToken));
+		IAsyncEnumerable<FindBuildResponse> FindBuildsAsync(string project, string? type, string? branch, CbObject query, CloudBuildSearchOptions? searchOptions = null, CloudBuildOptionsOverride? optionsOverride = null, CancellationToken cancellationToken = default);
 
 		/// <summary>
 		/// Search for builds in UE Cloud Storage using common components
@@ -355,7 +355,7 @@ namespace EpicGames.ProjectStore
 		/// <param name="optionsOverride">Optional per invocation overrides of options</param>		
 		/// <param name="cancellationToken">A cancellation token</param>
 		/// <returns></returns>
-		IEnumerable<FindBuildResponse> FindBuilds(string project, string? type, string? branch, CbObject query, CloudBuildSearchOptions? searchOptions = null, CloudBuildOptionsOverride? optionsOverride = null, CancellationToken cancellationToken = default(CancellationToken));
+		IEnumerable<FindBuildResponse> FindBuilds(string project, string? type, string? branch, CbObject query, CloudBuildSearchOptions? searchOptions = null, CloudBuildOptionsOverride? optionsOverride = null, CancellationToken cancellationToken = default);
 
 		/// <summary>
 		/// List the contents of a build
@@ -365,7 +365,7 @@ namespace EpicGames.ProjectStore
 		/// <param name="optionsOverride">Optional per invocation overrides of options</param>		
 		/// <param name="cancellationToken">A cancellation token</param>
 		/// <returns></returns>
-		Task<List<BuildContents>> ListBuildContentsAsync(StringId bucketId, CbObjectId buildId, CloudBuildOptionsOverride? optionsOverride = null, CancellationToken cancellationToken = default(CancellationToken));
+		Task<List<BuildContents>> ListBuildContentsAsync(StringId bucketId, CbObjectId buildId, CloudBuildOptionsOverride? optionsOverride = null, CancellationToken cancellationToken = default);
 
 		/// <summary>
 		/// Override the configuration of this instance of cloud storage, mostly useful if you get configuration from another source then the default inis / environment variables
@@ -408,7 +408,7 @@ namespace EpicGames.ProjectStore
 		/// <param name="asyncImport"></param>
 		/// <param name="cancellationToken"></param>
 		/// <returns></returns>
-		Task ImportSnapshotAsync(string localFileName, string localRootPath, string engineDir, string cookedPlatformDir, bool forceImport = false, bool asyncImport = false, CancellationToken cancellationToken = default(CancellationToken));
+		Task ImportSnapshotAsync(string localFileName, string localRootPath, string engineDir, string cookedPlatformDir, bool forceImport = false, bool asyncImport = false, CancellationToken cancellationToken = default);
 	}
 
 	/// <summary>
@@ -488,7 +488,7 @@ namespace EpicGames.ProjectStore
 		/// <returns></returns>
 		public IEnumerable<ValidationResult> Validate(ValidationContext validationContext)
 		{
-			List<ValidationResult> results = new List<ValidationResult>();
+			List<ValidationResult> results = [];
 			if (String.IsNullOrEmpty(AccessToken) && OidcTokenExePath == null && FetchAccessTokenInProcess == false)
 			{
 				results.Add(new ValidationResult("AccessToken or OidcTokenExePath needs to be defined for auth to function"));
diff --git a/Engine/Source/Programs/Shared/EpicGames.ProjectStore/ProjectStoreTypes.cs b/Engine/Source/Programs/Shared/EpicGames.ProjectStore/ProjectStoreTypes.cs
--- a/Engine/Source/Programs/Shared/EpicGames.ProjectStore/ProjectStoreTypes.cs
+++ b/Engine/Source/Programs/Shared/EpicGames.ProjectStore/ProjectStoreTypes.cs
@@ -54,7 +54,7 @@ namespace EpicGames.ProjectStore
 		/// <summary>
 		/// The list of remote host names for zenserver
 		/// </summary>
-		public List<string> RemoteHostNames { get; } = new List<string>();
+		public List<string> RemoteHostNames { get; } = [];
 		/// <summary>
 		/// The port for zenserver on the host
 		/// </summary>
@@ -175,7 +175,7 @@ namespace EpicGames.ProjectStore
 
 		private static string[] GetHostAddresses()
 		{
-			return Dns.GetHostAddresses("", System.Net.Sockets.AddressFamily.InterNetwork).Select(addr => addr.ToString()).ToArray();
+			return [.. Dns.GetHostAddresses("", System.Net.Sockets.AddressFamily.InterNetwork).Select(addr => addr.ToString())];
 		}
 
 		/// <summary>
@@ -380,7 +380,7 @@ namespace EpicGames.ProjectStore
 			/// <summary>
 			/// The list of snapshots contained within this collection.
 			/// </summary>
-			public List<SnapshotDescriptor> Snapshots { get; set; } = new List<SnapshotDescriptor>();
+			public List<SnapshotDescriptor> Snapshots { get; set; } = [];
 		}
 
 		private readonly SnapshotDescriptorCollection _snapshotDescriptors = new SnapshotDescriptorCollection();
@@ -692,8 +692,8 @@ namespace EpicGames.ProjectStore
 		public override List<int> GetAvailableChangelists(string branch, string project, string platform, string runtime)
 		{
 			DirectoryReference buildsDir = GetBuildsDir(branch, project, platform, runtime);
-			List<DirectoryReference> directories = DirectoryReference.EnumerateDirectories(buildsDir).ToList();
-			List<int> changelists = new List<int>();
+			List<DirectoryReference> directories = [.. DirectoryReference.EnumerateDirectories(buildsDir)];
+			List<int> changelists = [];
 
 			foreach (DirectoryReference d in directories)
 			{
@@ -907,13 +907,13 @@ namespace EpicGames.ProjectStore
 
 		class QueryBuildsResponse
 		{
-			public List<UnrealCloudDDCBuildData> Results { get; set; } = new List<UnrealCloudDDCBuildData>();
+			public List<UnrealCloudDDCBuildData> Results { get; set; } = [];
 			public bool PartialResult { get; set; }
 		}
 
 		class QueryBucketsResponse
 		{
-			public List<string> Buckets { get; set; } = new List<string>();
+			public List<string> Buckets { get; set; } = [];
 		}
 
 		[System.Diagnostics.CodeAnalysis.SuppressMessage("Usage", "CA1308", Justification = "We know the provided strings are safe to lowercase using ToLowerInvariant()")]
@@ -987,15 +987,14 @@ namespace EpicGames.ProjectStore
 			// Bucket format for snapshot buckets is <project>.oplog.<branch>.<platform>.
 			// Find buckets matching the format, with the specified project and branch,
 			// then return the associated platform.
-			return buckets.Buckets
+			return [.. buckets.Buckets
 				.Select(b => b.Split("."))
 				.Where(b =>
 					b.Length == 4
 					&& b[0].Equals(project, StringComparison.OrdinalIgnoreCase)
 					&& b[1].Equals("oplog", StringComparison.OrdinalIgnoreCase)
 					&& b[2].Equals(branch, StringComparison.OrdinalIgnoreCase))
-				.Select(b => b[3].ToLower())
-				.ToList();
+				.Select(b => b[3].ToLower())];
 		}
 
 		private List<UnrealCloudDDCBuildData> QueryBuilds(string branch, string project, string platform, string runtime, int changelist = 0)
@@ -1049,9 +1048,7 @@ namespace EpicGames.ProjectStore
 		/// <inheritdoc />
 		public override List<int> GetAvailableChangelists(string branch, string project, string platform, string runtime)
 		{
-			return QueryBuilds(branch, project, platform, runtime)
-				.Select(b => b.Metadata.Changelist)
-				.ToList();
+			return [.. QueryBuilds(branch, project, platform, runtime).Select(b => b.Metadata.Changelist)];
 		}
 
 		/// <inheritdoc />
diff --git a/Engine/Source/Programs/Shared/EpicGames.Redis/RedisHashKey.cs b/Engine/Source/Programs/Shared/EpicGames.Redis/RedisHashKey.cs
--- a/Engine/Source/Programs/Shared/EpicGames.Redis/RedisHashKey.cs
+++ b/Engine/Source/Programs/Shared/EpicGames.Redis/RedisHashKey.cs
@@ -379,7 +379,7 @@ namespace EpicGames.Redis
 
 			public HashRecordInfo()
 			{
-				List<HashPropertyInfo<T>> properties = new List<HashPropertyInfo<T>>();
+				List<HashPropertyInfo<T>> properties = [];
 				foreach (PropertyInfo propertyInfo in typeof(T).GetProperties(BindingFlags.Public | BindingFlags.Instance))
 				{
 					properties.Add(new HashPropertyInfo<T>(propertyInfo));
@@ -446,7 +446,7 @@ namespace EpicGames.Redis
 		/// <inheritdoc cref="IDatabaseAsync.HashSetAsync(RedisKey, HashEntry[], CommandFlags)"/>
 		public static Task HashSetAsync<TName, TValue>(this IDatabaseAsync target, RedisHashKey<TName, TValue> key, IEnumerable<HashEntry<TName, TValue>> entries, CommandFlags flags = CommandFlags.None)
 		{
-			return target.HashSetAsync(key.Inner, entries.Select(x => (HashEntry)x).ToArray(), flags);
+			return target.HashSetAsync(key.Inner, [.. entries.Select(x => (HashEntry)x)], flags);
 		}
 
 		#endregion
diff --git a/Engine/Source/Programs/Shared/EpicGames.Redis/RedisSerializer.cs b/Engine/Source/Programs/Shared/EpicGames.Redis/RedisSerializer.cs
--- a/Engine/Source/Programs/Shared/EpicGames.Redis/RedisSerializer.cs
+++ b/Engine/Source/Programs/Shared/EpicGames.Redis/RedisSerializer.cs
@@ -147,7 +147,7 @@ namespace EpicGames.Redis
 		static Dictionary<Type, object> CreateNativeConverterLookup()
 		{
 			KeyValuePair<Type, object>[] converters =
-			{
+			[
 				CreateNativeConverter<RedisValue>(x => x, x => x),
 				CreateNativeConverter(x => (bool)x, x => x),
 				CreateNativeConverter(x => (int)x, x => x),
@@ -164,7 +164,7 @@ namespace EpicGames.Redis
 				CreateNativeConverter(x => (byte[])x!, x => x),
 				CreateNativeConverter(x => (string)x!, x => x),
 				CreateNativeConverter(x => new DateTime((long)x, DateTimeKind.Utc), x => x.ToUniversalTime().Ticks)
-			};
+			];
 			return new Dictionary<Type, object>(converters);
 		}
 
@@ -173,7 +173,7 @@ namespace EpicGames.Redis
 			return new KeyValuePair<Type, object>(typeof(T), new RedisNativeConverter<T>(fromRedisValueFunc, toRedisValueFunc));
 		}
 
-		static readonly Dictionary<Type, Type> s_typeToConverterType = new Dictionary<Type, Type>();
+		static readonly Dictionary<Type, Type> s_typeToConverterType = [];
 
 		class RedisObjectConverter<T> : IRedisConverter<object>
 		{
diff --git a/Engine/Source/Programs/Shared/EpicGames.Serialization.Tests/CbWriterTests.cs b/Engine/Source/Programs/Shared/EpicGames.Serialization.Tests/CbWriterTests.cs
--- a/Engine/Source/Programs/Shared/EpicGames.Serialization.Tests/CbWriterTests.cs
+++ b/Engine/Source/Programs/Shared/EpicGames.Serialization.Tests/CbWriterTests.cs
@@ -200,8 +200,8 @@ namespace EpicGames.Serialization.Tests
 		[TestMethod]
 		public void Attachments()
 		{
-			List<IoHash> objHashes = Enumerable.Range(0, 5).Select(x => IoHash.Compute(Encoding.UTF8.GetBytes($"obj{x}"))).ToList();
-			List<IoHash> binHashes = Enumerable.Range(0, 5).Select(x => IoHash.Compute(Encoding.UTF8.GetBytes($"bin{x}"))).ToList();
+			List<IoHash> objHashes = [.. Enumerable.Range(0, 5).Select(x => IoHash.Compute(Encoding.UTF8.GetBytes($"obj{x}")))];
+			List<IoHash> binHashes = [.. Enumerable.Range(0, 5).Select(x => IoHash.Compute(Encoding.UTF8.GetBytes($"bin{x}")))];
 
 			CbWriter writer = new CbWriter();
 			writer.BeginObject();
diff --git a/Engine/Source/Programs/Shared/EpicGames.Serialization.Tests/EngineTests.cs b/Engine/Source/Programs/Shared/EpicGames.Serialization.Tests/EngineTests.cs
--- a/Engine/Source/Programs/Shared/EpicGames.Serialization.Tests/EngineTests.cs
+++ b/Engine/Source/Programs/Shared/EpicGames.Serialization.Tests/EngineTests.cs
@@ -391,7 +391,7 @@ namespace EpicGames.Serialization.Tests
 			}
 
 			// Test CbField(Array, Empty)
-			TestField(CbFieldType.Array, new byte[] { 1, 0 });
+			TestField(CbFieldType.Array, [1, 0]);
 
 			// Test CbField(Array, Empty)
 			{
@@ -536,7 +536,7 @@ namespace EpicGames.Serialization.Tests
 		public void CbFieldBinaryTest()
 		{
 			// Test CbField(Binary, Empty)
-			TestField(CbFieldType.Binary, new byte[] { 0 });
+			TestField(CbFieldType.Binary, [0]);
 
 			// Test CbField(Binary, Value)
 			{
@@ -566,7 +566,7 @@ namespace EpicGames.Serialization.Tests
 		public void CbFieldStringTest()
 		{
 			// Test CbField(String, Empty)
-			TestField(CbFieldType.String, new byte[] { 0 });
+			TestField(CbFieldType.String, [0]);
 
 			// Test CbField(String, Value)
 			{
@@ -922,10 +922,10 @@ namespace EpicGames.Serialization.Tests
 		public void CbFieldDateTimeTest()
 		{
 			// Test CbField(DateTime, Zero)
-			TestField(CbFieldType.DateTime, new byte[] { 0, 0, 0, 0, 0, 0, 0, 0 });
+			TestField(CbFieldType.DateTime, [0, 0, 0, 0, 0, 0, 0, 0]);
 
 			// Test CbField(DateTime, 0x1020_3040_5060_7080)
-			TestField(CbFieldType.DateTime, new byte[] { 0x10, 0x20, 0x30, 0x40, 0x50, 0x60, 0x70, 0x80 }, new DateTime(0x1020_3040_5060_7080L, DateTimeKind.Utc));
+			TestField(CbFieldType.DateTime, [0x10, 0x20, 0x30, 0x40, 0x50, 0x60, 0x70, 0x80], new DateTime(0x1020_3040_5060_7080L, DateTimeKind.Utc));
 
 			// Test CbField(DateTime, Zero) as FDateTime
 			{
@@ -947,10 +947,10 @@ namespace EpicGames.Serialization.Tests
 		public void CbFieldTimeSpanTest()
 		{
 			// Test CbField(TimeSpan, Zero)
-			TestField(CbFieldType.TimeSpan, new byte[] { 0, 0, 0, 0, 0, 0, 0, 0 });
+			TestField(CbFieldType.TimeSpan, [0, 0, 0, 0, 0, 0, 0, 0]);
 
 			// Test CbField(TimeSpan, 0x1020_3040_5060_7080)
-			TestField(CbFieldType.TimeSpan, new byte[] { 0x10, 0x20, 0x30, 0x40, 0x50, 0x60, 0x70, 0x80 }, new TimeSpan(0x1020_3040_5060_7080L));
+			TestField(CbFieldType.TimeSpan, [0x10, 0x20, 0x30, 0x40, 0x50, 0x60, 0x70, 0x80], new TimeSpan(0x1020_3040_5060_7080L));
 
 			// Test CbField(TimeSpan, Zero) as FTimeSpan
 			{
diff --git a/Engine/Source/Programs/Shared/EpicGames.Serialization.Tests/JupiterTests.cs b/Engine/Source/Programs/Shared/EpicGames.Serialization.Tests/JupiterTests.cs
--- a/Engine/Source/Programs/Shared/EpicGames.Serialization.Tests/JupiterTests.cs
+++ b/Engine/Source/Programs/Shared/EpicGames.Serialization.Tests/JupiterTests.cs
@@ -137,7 +137,7 @@ namespace EpicGames.Serialization.Tests
 			Assert.AreEqual(1, fields.Count);
 			CbField? needs = o["needs"];
 			List<CbField> blobList = [.. needs!.AsArray()];
-			IoHash[] blobs = blobList.Select(field => field.AsHash()).ToArray();
+			IoHash[] blobs = [.. blobList.Select(field => field.AsHash())];
 			CollectionAssert.AreEqual(new IoHash[] { hash1, hash2 }, blobs);
 		}
 
diff --git a/Engine/Source/Programs/Shared/EpicGames.Slack/Blocks/ActionsBlock.cs b/Engine/Source/Programs/Shared/EpicGames.Slack/Blocks/ActionsBlock.cs
--- a/Engine/Source/Programs/Shared/EpicGames.Slack/Blocks/ActionsBlock.cs
+++ b/Engine/Source/Programs/Shared/EpicGames.Slack/Blocks/ActionsBlock.cs
@@ -15,7 +15,7 @@ namespace EpicGames.Slack.Blocks
 		/// A collection of interactive elements.
 		/// </summary>
 		[JsonPropertyName("elements"), JsonPropertyOrder(1)]
-		public List<Element> Elements { get; } = new List<Element>();
+		public List<Element> Elements { get; } = [];
 
 		/// <summary>
 		/// Constructor
diff --git a/Engine/Source/Programs/Shared/EpicGames.Slack/Blocks/ContextBlock.cs b/Engine/Source/Programs/Shared/EpicGames.Slack/Blocks/ContextBlock.cs
--- a/Engine/Source/Programs/Shared/EpicGames.Slack/Blocks/ContextBlock.cs
+++ b/Engine/Source/Programs/Shared/EpicGames.Slack/Blocks/ContextBlock.cs
@@ -16,7 +16,7 @@ namespace EpicGames.Slack.Blocks
 		/// An array of image elements and text objects. Maximum number of items is 10.
 		/// </summary>
 		[JsonPropertyName("elements")]
-		public List<Element> Elements { get; } = new List<Element>();
+		public List<Element> Elements { get; } = [];
 
 		/// <summary>
 		/// Constructor
diff --git a/Engine/Source/Programs/Shared/EpicGames.Slack/Elements/CheckboxGroupElement.cs b/Engine/Source/Programs/Shared/EpicGames.Slack/Elements/CheckboxGroupElement.cs
--- a/Engine/Source/Programs/Shared/EpicGames.Slack/Elements/CheckboxGroupElement.cs
+++ b/Engine/Source/Programs/Shared/EpicGames.Slack/Elements/CheckboxGroupElement.cs
@@ -21,13 +21,13 @@ namespace EpicGames.Slack.Elements
 		/// An array of option objects. A maximum of 10 options are allowed.
 		/// </summary>
 		[JsonPropertyName("options")]
-		public List<SlackOption> Options { get; } = new List<SlackOption>();
+		public List<SlackOption> Options { get; } = [];
 
 		/// <summary>
 		/// An array of option objects that exactly matches one or more of the options within options. These options will be selected when the checkbox group initially loads.
 		/// </summary>
 		[JsonPropertyName("initial_options")]
-		public List<SlackOption> InitialOptions { get; } = new List<SlackOption>();
+		public List<SlackOption> InitialOptions { get; } = [];
 
 		/// <summary>
 		/// A confirm object that defines an optional confirmation dialog that appears after clicking one of the checkboxes in this element.
diff --git a/Engine/Source/Programs/Shared/EpicGames.Slack/Elements/OverflowMenuElement.cs b/Engine/Source/Programs/Shared/EpicGames.Slack/Elements/OverflowMenuElement.cs
--- a/Engine/Source/Programs/Shared/EpicGames.Slack/Elements/OverflowMenuElement.cs
+++ b/Engine/Source/Programs/Shared/EpicGames.Slack/Elements/OverflowMenuElement.cs
@@ -25,7 +25,7 @@ namespace EpicGames.Slack.Elements
 		/// An array of up to five option objects to display in the menu.
 		/// </summary>
 		[JsonPropertyName("options")]
-		public List<SlackOption> Options { get; } = new List<SlackOption>();
+		public List<SlackOption> Options { get; } = [];
 
 		/// <summary>
 		/// A confirm object that defines an optional confirmation dialog that appears after clicking one of the checkboxes in this element.
diff --git a/Engine/Source/Programs/Shared/EpicGames.Slack/Elements/RadioButtonGroupElement.cs b/Engine/Source/Programs/Shared/EpicGames.Slack/Elements/RadioButtonGroupElement.cs
--- a/Engine/Source/Programs/Shared/EpicGames.Slack/Elements/RadioButtonGroupElement.cs
+++ b/Engine/Source/Programs/Shared/EpicGames.Slack/Elements/RadioButtonGroupElement.cs
@@ -21,7 +21,7 @@ namespace EpicGames.Slack.Elements
 		/// An array of option objects. A maximum of 10 options are allowed.
 		/// </summary>
 		[JsonPropertyName("options")]
-		public List<SlackOption> Options { get; } = new List<SlackOption>();
+		public List<SlackOption> Options { get; } = [];
 
 		/// <summary>
 		/// An array of option objects that exactly matches one or more of the options within options. These options will be selected when the checkbox group initially loads.
diff --git a/Engine/Source/Programs/Shared/EpicGames.Slack/Elements/SelectMenuElement.cs b/Engine/Source/Programs/Shared/EpicGames.Slack/Elements/SelectMenuElement.cs
--- a/Engine/Source/Programs/Shared/EpicGames.Slack/Elements/SelectMenuElement.cs
+++ b/Engine/Source/Programs/Shared/EpicGames.Slack/Elements/SelectMenuElement.cs
@@ -77,7 +77,7 @@ namespace EpicGames.Slack.Elements
 			: base("static_select", actionId)
 		{
 			Placeholder = placeholder;
-			Options = new List<SlackOption>(options);
+			Options = [.. options];
 		}
 
 		/// <summary>
@@ -87,7 +87,7 @@ namespace EpicGames.Slack.Elements
 			: base("static_select", actionId)
 		{
 			Placeholder = placeholder;
-			OptionGroups = new List<SlackOptionGroup>(optionGroups);
+			OptionGroups = [.. optionGroups];
 		}
 	}
 }
diff --git a/Engine/Source/Programs/Shared/EpicGames.Slack/SlackAttachment.cs b/Engine/Source/Programs/Shared/EpicGames.Slack/SlackAttachment.cs
--- a/Engine/Source/Programs/Shared/EpicGames.Slack/SlackAttachment.cs
+++ b/Engine/Source/Programs/Shared/EpicGames.Slack/SlackAttachment.cs
@@ -27,7 +27,7 @@ namespace EpicGames.Slack
 		/// A collection of BlockKit blocks the attachment is composed from.
 		/// </summary>
 		[JsonPropertyName("blocks"), JsonPropertyOrder(3)]
-		public List<Block> Blocks { get; } = new List<Block>();
+		public List<Block> Blocks { get; } = [];
 
 		/// <summary>
 		/// Implicit conversion to a message of its own
diff --git a/Engine/Source/Programs/Shared/EpicGames.Slack/SlackClient.cs b/Engine/Source/Programs/Shared/EpicGames.Slack/SlackClient.cs
--- a/Engine/Source/Programs/Shared/EpicGames.Slack/SlackClient.cs
+++ b/Engine/Source/Programs/Shared/EpicGames.Slack/SlackClient.cs
@@ -506,7 +506,7 @@ namespace EpicGames.Slack
 		/// <param name="channel">Channel identifier to invite the user to</param>
 		/// <param name="userId">The user id</param>
 		/// <param name="cancellationToken">Cancellation token for the operation</param>
-		public Task InviteUserAsync(string channel, string userId, CancellationToken cancellationToken = default) => InviteUsersAsync(channel, new[] { userId }, cancellationToken);
+		public Task InviteUserAsync(string channel, string userId, CancellationToken cancellationToken = default) => InviteUsersAsync(channel, [userId], cancellationToken);
 
 		/// <summary>
 		/// Invite a set of users to a channel
diff --git a/Engine/Source/Programs/Shared/EpicGames.Slack/SlackMessage.cs b/Engine/Source/Programs/Shared/EpicGames.Slack/SlackMessage.cs
--- a/Engine/Source/Programs/Shared/EpicGames.Slack/SlackMessage.cs
+++ b/Engine/Source/Programs/Shared/EpicGames.Slack/SlackMessage.cs
@@ -22,12 +22,12 @@ namespace EpicGames.Slack
 		/// <summary>
 		/// A JSON-based array of structured blocks, presented as a URL-encoded string.
 		/// </summary>
-		public List<Block> Blocks { get; } = new List<Block>();
+		public List<Block> Blocks { get; } = [];
 
 		/// <summary>
 		/// Attachments for the message
 		/// </summary>
-		public List<SlackAttachment> Attachments { get; } = new List<SlackAttachment>();
+		public List<SlackAttachment> Attachments { get; } = [];
 
 		/// <summary>
 		/// Pass true to enable unfurling of primarily text-based content.
diff --git a/Engine/Source/Programs/Shared/EpicGames.Slack/SlackView.cs b/Engine/Source/Programs/Shared/EpicGames.Slack/SlackView.cs
--- a/Engine/Source/Programs/Shared/EpicGames.Slack/SlackView.cs
+++ b/Engine/Source/Programs/Shared/EpicGames.Slack/SlackView.cs
@@ -29,7 +29,7 @@ namespace EpicGames.Slack
 		/// An array of blocks that defines the content of the view. Max of 100 blocks.
 		/// </summary>
 		[JsonPropertyName("blocks")]
-		public List<Block> Blocks { get; } = new List<Block>();
+		public List<Block> Blocks { get; } = [];
 
 		/// <summary>
 		/// An optional plain_text element that defines the text displayed in the close button at the bottom-right of the view. Max length of 24 characters.
@@ -105,7 +105,7 @@ namespace EpicGames.Slack
 		/// </summary>
 		[JsonPropertyName("values")]
 #pragma warning disable CA2227 // Collection properties should be read only
-		public Dictionary<string, Dictionary<string, SlackViewValue>> Values { get; set; } = new Dictionary<string, Dictionary<string, SlackViewValue>>();
+		public Dictionary<string, Dictionary<string, SlackViewValue>> Values { get; set; } = [];
 #pragma warning restore CA2227 // Collection properties should be read only
 
 		/// <summary>
diff --git a/Engine/Source/Programs/Shared/EpicGames.Tracing.Tests/UnrealInsights/EventTest.cs b/Engine/Source/Programs/Shared/EpicGames.Tracing.Tests/UnrealInsights/EventTest.cs
--- a/Engine/Source/Programs/Shared/EpicGames.Tracing.Tests/UnrealInsights/EventTest.cs
+++ b/Engine/Source/Programs/Shared/EpicGames.Tracing.Tests/UnrealInsights/EventTest.cs
@@ -13,14 +13,14 @@ namespace EpicGames.Tracing.Tests.UnrealInsights
 		public void EnterScopeEventTimestampDeserialize()
 		{
 			{
-				using MemoryStream ms = new MemoryStream(new byte[] { 0x10, 0x05, 0x00, 0x00, 0x00, 0x00, 0x00, 0x00 });
+				using MemoryStream ms = new MemoryStream([0x10, 0x05, 0x00, 0x00, 0x00, 0x00, 0x00, 0x00]);
 				using BinaryReader reader = new BinaryReader(ms);
 				EnterScopeEventTimestamp @event = EnterScopeEventTimestamp.Deserialize(reader);
 				Assert.AreEqual((ulong)0x05, @event.Timestamp);
 			}
 
 			{
-				using MemoryStream ms = new MemoryStream(new byte[] { 0x10, 0x00, 0x00, 0x00, 0x00, 0x00, 0x00, 0xFF });
+				using MemoryStream ms = new MemoryStream([0x10, 0x00, 0x00, 0x00, 0x00, 0x00, 0x00, 0xFF]);
 				using BinaryReader reader = new BinaryReader(ms);
 				EnterScopeEventTimestamp @event = EnterScopeEventTimestamp.Deserialize(reader);
 				Assert.AreEqual((ulong)0x00_FF_00_00_00_00_00_00, @event.Timestamp);
diff --git a/Engine/Source/Programs/Shared/EpicGames.Tracing.Tests/UnrealInsights/GenericEventTest.cs b/Engine/Source/Programs/Shared/EpicGames.Tracing.Tests/UnrealInsights/GenericEventTest.cs
--- a/Engine/Source/Programs/Shared/EpicGames.Tracing.Tests/UnrealInsights/GenericEventTest.cs
+++ b/Engine/Source/Programs/Shared/EpicGames.Tracing.Tests/UnrealInsights/GenericEventTest.cs
@@ -41,20 +41,20 @@ namespace EpicGames.Tracing.Tests.UnrealInsights
 		[TestMethod]
 		public void Serialize_Memory_MemoryScope()
 		{
-			SerializeAndAssertGenericEvent(_memoryMemoryScopeHex1, MemoryMemoryScopeEvent.EventType, false, 18, new[]
-			{
+			SerializeAndAssertGenericEvent(_memoryMemoryScopeHex1, MemoryMemoryScopeEvent.EventType, false, 18,
+			[
 				Field.FromInt(0x27),
-			}, 45);
+			], 45);
 
-			SerializeAndAssertGenericEvent(_memoryMemoryScopeHex2, MemoryMemoryScopeEvent.EventType, false, 18, new[]
-			{
+			SerializeAndAssertGenericEvent(_memoryMemoryScopeHex2, MemoryMemoryScopeEvent.EventType, false, 18,
+			[
 				Field.FromInt(0x27),
-			}, 46);
+			], 46);
 
-			SerializeAndAssertGenericEvent(_memoryMemoryScopeHex3, MemoryMemoryScopeEvent.EventType, false, 18, new[]
-			{
+			SerializeAndAssertGenericEvent(_memoryMemoryScopeHex3, MemoryMemoryScopeEvent.EventType, false, 18,
+			[
 				Field.FromInt(0x33),
-			}, 5832765);
+			], 5832765);
 		}
 
 		[TestMethod]
@@ -71,13 +71,13 @@ namespace EpicGames.Tracing.Tests.UnrealInsights
 		[TestMethod]
 		public void Serialize_Trace_NewTrace_Generic()
 		{
-			SerializeAndAssertGenericEvent(_traceNewTraceHex, TraceNewTraceEvent.EventType, true, 16, new[]
-			{
+			SerializeAndAssertGenericEvent(_traceNewTraceHex, TraceNewTraceEvent.EventType, true, 16,
+			[
 				Field.FromLong(25582215261913),
 				Field.FromLong(10000000),
 				Field.FromInt(21069),
 				Field.FromInt(8),
-			});
+			]);
 		}
 
 		[TestMethod]
@@ -155,8 +155,8 @@ namespace EpicGames.Tracing.Tests.UnrealInsights
 		[TestMethod]
 		public void Serialize_Diagnostics_Session2_Generic()
 		{
-			SerializeAndAssertGenericEvent(_diagnosticsSession2Hex, DiagnosticsSession2Event.EventType, true, 34, new[]
-			{
+			SerializeAndAssertGenericEvent(_diagnosticsSession2Hex, DiagnosticsSession2Event.EventType, true, 34,
+			[
 				Field.FromString("Win64"),
 				Field.FromString("UnrealEditor"),
 				Field.FromString(@" D:\depot\starship-main\QAGame\QAGame.uproject -trace=cpu -game"),
@@ -165,7 +165,7 @@ namespace EpicGames.Tracing.Tests.UnrealInsights
 				Field.FromInt(17442524),
 				Field.FromInt(3),
 				Field.FromInt(4)
-			});
+			]);
 		}
 
 		[TestMethod]
@@ -186,11 +186,11 @@ namespace EpicGames.Tracing.Tests.UnrealInsights
 		[TestMethod]
 		public void CpuProfiler_EventSpec_Serialize_Generic()
 		{
-			SerializeAndAssertGenericEvent(_cpuProfilerEventSpecHex, CpuProfilerEventSpecEvent.EventType, true, 33, new Field[]
-			{
+			SerializeAndAssertGenericEvent(_cpuProfilerEventSpecHex, CpuProfilerEventSpecEvent.EventType, true, 33,
+			[
 				Field.FromInt(9524),
 				Field.FromString("K2Node_VariableGet /Engine/EngineSky/BP_Sky_Sphere.BP_Sky_Sphere:UserConstructionScript.K2Node_VariableGet_983")
-			});
+			]);
 		}
 
 		[TestMethod]
@@ -271,10 +271,10 @@ namespace EpicGames.Tracing.Tests.UnrealInsights
 			using BinaryWriter writer = new BinaryWriter(ms);
 
 			Field[] fields =
-			{
+			[
 				Field.FromInt(9524),
 				Field.FromString("K2Node_VariableGet /Engine/EngineSky/BP_Sky_Sphere.BP_Sky_Sphere:UserConstructionScript.K2Node_VariableGet_983"),
-			};
+			];
 			GenericEvent @event = new GenericEvent(98989, fields, CpuProfilerEventSpecEvent.EventType);
 			Assert.AreEqual(0x77, @event.Size);
 
@@ -289,7 +289,7 @@ namespace EpicGames.Tracing.Tests.UnrealInsights
 		public void AuxHeaderDeserialize()
 		{
 			// Aux header from field "Name" in CpuProfiler.EventSpec
-			using MemoryStream ms = new MemoryStream(new byte[] { 0x01, 0xC1, 0x0D, 0x00 });
+			using MemoryStream ms = new MemoryStream([0x01, 0xC1, 0x0D, 0x00]);
 			using BinaryReader reader = new BinaryReader(ms);
 
 			(ushort uid, int fieldIndex, int size) = GenericEvent.DeserializeAuxHeader(reader.ReadUInt32());
diff --git a/Engine/Source/Programs/Shared/EpicGames.Tracing.Tests/UnrealInsights/TransportPacketTest.cs b/Engine/Source/Programs/Shared/EpicGames.Tracing.Tests/UnrealInsights/TransportPacketTest.cs
--- a/Engine/Source/Programs/Shared/EpicGames.Tracing.Tests/UnrealInsights/TransportPacketTest.cs
+++ b/Engine/Source/Programs/Shared/EpicGames.Tracing.Tests/UnrealInsights/TransportPacketTest.cs
@@ -177,7 +177,7 @@ namespace EpicGames.Tracing.Tests.UnrealInsights
 				using MemoryStream ms = new MemoryStream();
 				using BinaryWriter writer = new BinaryWriter(ms);
 				TransportPacket packet = TransportPacket.Create(0, 0);
-				packet.Serialize(writer, new[] { (uid1, event1 as ITraceEvent), (uid2, event2 as ITraceEvent) });
+				packet.Serialize(writer, [(uid1, event1 as ITraceEvent), (uid2, event2 as ITraceEvent)]);
 
 				ms.Position = 0;
 				using BinaryReader reader = new BinaryReader(ms);
@@ -191,7 +191,7 @@ namespace EpicGames.Tracing.Tests.UnrealInsights
 				using MemoryStream ms = new MemoryStream();
 				using BinaryWriter writer = new BinaryWriter(ms);
 				TransportPacket packet = TransportPacket.Create(0, 1);
-				packet.Serialize(writer, Array.Empty<(ushort, ITraceEvent)>());
+				packet.Serialize(writer, []);
 
 				ms.Position = 0;
 				using BinaryReader reader = new BinaryReader(ms);
@@ -205,7 +205,7 @@ namespace EpicGames.Tracing.Tests.UnrealInsights
 				using MemoryStream ms = new MemoryStream();
 				using BinaryWriter writer = new BinaryWriter(ms);
 				TransportPacket packet = TransportPacket.Create(0, 12345);
-				packet.Serialize(writer, Array.Empty<(ushort, ITraceEvent)>());
+				packet.Serialize(writer, []);
 
 				ms.Position = 0;
 				using BinaryReader reader = new BinaryReader(ms);
@@ -231,7 +231,7 @@ namespace EpicGames.Tracing.Tests.UnrealInsights
 				using MemoryStream ms = new MemoryStream();
 				using BinaryWriter writer = new BinaryWriter(ms);
 				TransportPacket packet = TransportPacket.Create(0, 0);
-				packet.Serialize(writer, new[] { (uid1, event1 as ITraceEvent), (uid2, event2 as ITraceEvent) });
+				packet.Serialize(writer, [(uid1, event1 as ITraceEvent), (uid2, event2 as ITraceEvent)]);
 
 				ms.Position = 0;
 				using BinaryReader reader = new BinaryReader(ms);
@@ -246,12 +246,12 @@ namespace EpicGames.Tracing.Tests.UnrealInsights
 				using MemoryStream ms = new MemoryStream();
 				using BinaryWriter writer = new BinaryWriter(ms);
 				TransportPacket packet = TransportPacket.Create(0, 0);
-				packet.Serialize(writer, new[]
-				{
+				packet.Serialize(writer,
+				[
 					(uid1, event1.Type as ITraceEvent),
 					(uid2, event2.Type as ITraceEvent),
 					(uid3, event3.Type as ITraceEvent)
-				});
+				]);
 
 				ms.Position = 0;
 				using BinaryReader reader = new BinaryReader(ms);
diff --git a/Engine/Source/Programs/Shared/EpicGames.Tracing/UnrealInsights/Events/CpuProfiler.cs b/Engine/Source/Programs/Shared/EpicGames.Tracing/UnrealInsights/Events/CpuProfiler.cs
--- a/Engine/Source/Programs/Shared/EpicGames.Tracing/UnrealInsights/Events/CpuProfiler.cs
+++ b/Engine/Source/Programs/Shared/EpicGames.Tracing/UnrealInsights/Events/CpuProfiler.cs
@@ -9,10 +9,10 @@ namespace EpicGames.Tracing.UnrealInsights.Events
 	public class CpuProfilerEventSpecEvent : ITraceEvent
 	{
 		public static readonly EventType EventType = new EventType(0, "CpuProfiler", "EventSpec", EventType.FlagImportant | EventType.FlagMaybeHasAux | EventType.FlagNoSync,
-			new List<EventTypeField>() {
+			[
 				new EventTypeField(0, 4, EventTypeField.TypeInt32, "Id"),
 				new EventTypeField(4, 0, EventTypeField.TypeAnsiString, "Name")
-			});
+			]);
 
 		public ushort Size => (ushort)(_genericEvent.Size + TraceImportantEventHeader.HeaderSize);
 		public EventType Type => EventType;
@@ -28,10 +28,10 @@ namespace EpicGames.Tracing.UnrealInsights.Events
 			_name = name;
 
 			Field[] fields =
-			{
+			[
 				Field.FromInt((int) Id),
 				Field.FromString(_name),
-			};
+			];
 
 			_genericEvent = new GenericEvent(0, fields, EventType);
 		}
@@ -46,7 +46,7 @@ namespace EpicGames.Tracing.UnrealInsights.Events
 	public class CpuProfilerEventBatchEvent : ITraceEvent
 	{
 		public static readonly EventType EventType = new EventType(0, "CpuProfiler", "EventBatch", EventType.FlagMaybeHasAux | EventType.FlagNoSync,
-			new List<EventTypeField>() { new EventTypeField(0, 0, EventTypeField.TypeArray, "Data") });
+			[new EventTypeField(0, 0, EventTypeField.TypeArray, "Data")]);
 
 		public ushort Size => (ushort)(_genericEvent.Size + sizeof(ushort));
 		public EventType Type => EventType;
@@ -55,7 +55,7 @@ namespace EpicGames.Tracing.UnrealInsights.Events
 
 		public CpuProfilerEventBatchEvent(byte[] data)
 		{
-			Field[] fields = { Field.FromArray(data) };
+			Field[] fields = [Field.FromArray(data)];
 			_genericEvent = new GenericEvent(0, fields, EventType);
 		}
 
@@ -92,7 +92,7 @@ namespace EpicGames.Tracing.UnrealInsights.Events
 	public class CpuProfilerSerializer
 	{
 		private readonly ulong _cycleFrequency;
-		public List<CpuProfilerScopeEvent> ScopeEvents { get; } = new List<CpuProfilerScopeEvent>();
+		public List<CpuProfilerScopeEvent> ScopeEvents { get; } = [];
 
 		public CpuProfilerSerializer(ulong cycleFrequency)
 		{
@@ -138,7 +138,7 @@ namespace EpicGames.Tracing.UnrealInsights.Events
 
 		public List<byte[]> Write()
 		{
-			List<MemoryStream> streams = new List<MemoryStream>();
+			List<MemoryStream> streams = [];
 			MemoryStream stream = new MemoryStream();
 			BinaryWriter writer = new BinaryWriter(stream);
 			streams.Add(stream);
@@ -175,7 +175,7 @@ namespace EpicGames.Tracing.UnrealInsights.Events
 
 			writer.Close();
 
-			return streams.Select(x => x.ToArray()).ToList();
+			return [.. streams.Select(x => x.ToArray())];
 		}
 
 		internal static ulong EncodeTimestamp(ulong timestamp, bool isScopeEnter)
diff --git a/Engine/Source/Programs/Shared/EpicGames.Tracing/UnrealInsights/Events/Diagnostics.cs b/Engine/Source/Programs/Shared/EpicGames.Tracing/UnrealInsights/Events/Diagnostics.cs
--- a/Engine/Source/Programs/Shared/EpicGames.Tracing/UnrealInsights/Events/Diagnostics.cs
+++ b/Engine/Source/Programs/Shared/EpicGames.Tracing/UnrealInsights/Events/Diagnostics.cs
@@ -1,6 +1,5 @@
 // Copyright Epic Games, Inc. All Rights Reserved.
 
-using System.Collections.Generic;
 using System.IO;
 
 namespace EpicGames.Tracing.UnrealInsights.Events
@@ -8,7 +7,7 @@ namespace EpicGames.Tracing.UnrealInsights.Events
 	public class DiagnosticsSession2Event : ITraceEvent
 	{
 		public static readonly EventType EventType = new EventType(0, "Diagnostics", "Session2", EventType.FlagImportant | EventType.FlagMaybeHasAux | EventType.FlagNoSync,
-			new List<EventTypeField>() {
+			[
 				new EventTypeField(0, 0, EventTypeField.TypeAnsiString, "Platform"),
 				new EventTypeField(0, 0, EventTypeField.TypeAnsiString, "AppName"),
 				new EventTypeField(0, 0, EventTypeField.TypeWideString, "CommandLine"),
@@ -17,7 +16,7 @@ namespace EpicGames.Tracing.UnrealInsights.Events
 				new EventTypeField(0, 4, EventTypeField.TypeInt32, "Changelist"),
 				new EventTypeField(4, 1, EventTypeField.TypeInt8, "ConfigurationType"),
 				new EventTypeField(5, 1, EventTypeField.TypeInt8, "TargetType"),
-			});
+			]);
 
 		public ushort Size => (ushort)(_genericEvent.Size + TraceImportantEventHeader.HeaderSize);
 		public EventType Type => EventType;
@@ -26,7 +25,7 @@ namespace EpicGames.Tracing.UnrealInsights.Events
 		public DiagnosticsSession2Event(string platform, string appName, string commandLine, string branch, string buildVersion, int changeList, int configurationType, int targetType)
 		{
 			Field[] fields =
-			{
+			[
 				Field.FromString(platform),
 				Field.FromString(appName),
 				Field.FromString(commandLine),
@@ -35,7 +34,7 @@ namespace EpicGames.Tracing.UnrealInsights.Events
 				Field.FromInt(changeList),
 				Field.FromInt(configurationType),
 				Field.FromInt(targetType),
-			};
+			];
 
 			_genericEvent = new GenericEvent(0, fields, EventType);
 		}
diff --git a/Engine/Source/Programs/Shared/EpicGames.Tracing/UnrealInsights/Events/Logging.cs b/Engine/Source/Programs/Shared/EpicGames.Tracing/UnrealInsights/Events/Logging.cs
--- a/Engine/Source/Programs/Shared/EpicGames.Tracing/UnrealInsights/Events/Logging.cs
+++ b/Engine/Source/Programs/Shared/EpicGames.Tracing/UnrealInsights/Events/Logging.cs
@@ -1,7 +1,6 @@
 // Copyright Epic Games, Inc. All Rights Reserved.
 
 using System;
-using System.Collections.Generic;
 using System.IO;
 
 namespace EpicGames.Tracing.UnrealInsights.Events
@@ -9,11 +8,11 @@ namespace EpicGames.Tracing.UnrealInsights.Events
 	public class LoggingLogMessage : ITraceEvent
 	{
 		public static readonly EventType EventType = new EventType(0, "Logging", "LogMessage", EventType.FlagMaybeHasAux | EventType.FlagNoSync,
-			new List<EventTypeField>() {
+			[
 				new EventTypeField(0, 8, EventTypeField.TypeInt64, "LogPoint"),
 				new EventTypeField(8, 8, EventTypeField.TypeInt64, "Cycle"),
 				new EventTypeField(16, 0, EventTypeField.TypeArray, "FormatArgs")
-			});
+			]);
 
 		public ushort Size => throw new NotImplementedException();
 		public EventType Type => EventType;
diff --git a/Engine/Source/Programs/Shared/EpicGames.Tracing/UnrealInsights/Events/Memory.cs b/Engine/Source/Programs/Shared/EpicGames.Tracing/UnrealInsights/Events/Memory.cs
--- a/Engine/Source/Programs/Shared/EpicGames.Tracing/UnrealInsights/Events/Memory.cs
+++ b/Engine/Source/Programs/Shared/EpicGames.Tracing/UnrealInsights/Events/Memory.cs
@@ -1,7 +1,6 @@
 // Copyright Epic Games, Inc. All Rights Reserved.
 
 using System;
-using System.Collections.Generic;
 using System.IO;
 
 namespace EpicGames.Tracing.UnrealInsights.Events
@@ -9,7 +8,7 @@ namespace EpicGames.Tracing.UnrealInsights.Events
 	public class MemoryMemoryScopeEvent : ITraceEvent
 	{
 		public static readonly EventType EventType = new EventType(0, "Memory", "MemoryScope", EventType.FlagNone,
-			new List<EventTypeField>() { new EventTypeField(0, 4, EventTypeField.TypeInt32, "Id") });
+			[new EventTypeField(0, 4, EventTypeField.TypeInt32, "Id")]);
 
 		public ushort Size => throw new NotImplementedException();
 		public EventType Type => EventType;
diff --git a/Engine/Source/Programs/Shared/EpicGames.Tracing/UnrealInsights/Events/NewEvent.cs b/Engine/Source/Programs/Shared/EpicGames.Tracing/UnrealInsights/Events/NewEvent.cs
--- a/Engine/Source/Programs/Shared/EpicGames.Tracing/UnrealInsights/Events/NewEvent.cs
+++ b/Engine/Source/Programs/Shared/EpicGames.Tracing/UnrealInsights/Events/NewEvent.cs
@@ -197,17 +197,17 @@ namespace EpicGames.Tracing.UnrealInsights.Events
 			LoggerName = loggerName;
 			EventName = eventName;
 			Flags = flags;
-			Fields = new List<EventTypeField>();
+			Fields = [];
 		}
 
 		public static EventType WellKnown(ushort uid, string name)
 		{
-			return new EventType(uid, "WellKnown", name, 0, new List<EventTypeField>());
+			return new EventType(uid, "WellKnown", name, 0, []);
 		}
 
 		private static EventType Self()
 		{
-			return new EventType(0, "EventType", "Self", 0, new List<EventTypeField>());
+			return new EventType(0, "EventType", "Self", 0, []);
 		}
 
 		public void AddEventType(ushort offset, ushort size, byte typeInfo, string name)
@@ -384,7 +384,7 @@ namespace EpicGames.Tracing.UnrealInsights.Events
 				fields[i].SetName(Encoding.UTF8.GetString(reader.ReadBytesStrict(fields[i].NameSize)));
 			}
 
-			return (newEventUid, new EventType(newEventUid, loggerName, eventName, flags, fields.ToList()));
+			return (newEventUid, new EventType(newEventUid, loggerName, eventName, flags, [.. fields]));
 		}
 	}
 }
\ No newline at end of file
diff --git a/Engine/Source/Programs/Shared/EpicGames.Tracing/UnrealInsights/Events/Trace.cs b/Engine/Source/Programs/Shared/EpicGames.Tracing/UnrealInsights/Events/Trace.cs
--- a/Engine/Source/Programs/Shared/EpicGames.Tracing/UnrealInsights/Events/Trace.cs
+++ b/Engine/Source/Programs/Shared/EpicGames.Tracing/UnrealInsights/Events/Trace.cs
@@ -1,6 +1,5 @@
 // Copyright Epic Games, Inc. All Rights Reserved.
 
-using System.Collections.Generic;
 using System.IO;
 
 namespace EpicGames.Tracing.UnrealInsights.Events
@@ -8,12 +7,12 @@ namespace EpicGames.Tracing.UnrealInsights.Events
 	public class TraceNewTraceEvent : ITraceEvent
 	{
 		public static readonly EventType EventType = new EventType(0, "$Trace", "NewTrace", EventType.FlagImportant | EventType.FlagNoSync,
-			new List<EventTypeField>() {
+			[
 				new EventTypeField(0, 8, EventTypeField.TypeInt64, "StartCycle"),
 				new EventTypeField(8, 8, EventTypeField.TypeInt64, "CycleFrequency"),
 				new EventTypeField(16, 2, EventTypeField.TypeInt16, "Endian"),
 				new EventTypeField(18, 1, EventTypeField.TypeInt8, "PointerSize")
-			});
+			]);
 
 		public ulong StartCycle { get; private set; }
 		public ulong CycleFrequency { get; private set; }
@@ -53,12 +52,12 @@ namespace EpicGames.Tracing.UnrealInsights.Events
 	public class TraceThreadInfoEvent : ITraceEvent
 	{
 		public static readonly EventType EventType = new EventType(0, "$Trace", "ThreadInfo", EventType.FlagImportant | EventType.FlagMaybeHasAux | EventType.FlagNoSync,
-			new List<EventTypeField>() {
+			[
 				new EventTypeField(0, 4, EventTypeField.TypeInt32, "ThreadId"),
 				new EventTypeField(4, 4, EventTypeField.TypeInt32, "SystemId"),
 				new EventTypeField(8, 4, EventTypeField.TypeInt32, "SortHint"),
 				new EventTypeField(12, 0, EventTypeField.TypeAnsiString, "Name")
-			});
+			]);
 
 		public ushort Size => (ushort)(_genericEvent.Size + TraceImportantEventHeader.HeaderSize);
 		public EventType Type => EventType;
@@ -67,12 +66,12 @@ namespace EpicGames.Tracing.UnrealInsights.Events
 		public TraceThreadInfoEvent(int threadId, int systemId, int sortHint, string name)
 		{
 			Field[] fields =
-			{
+			[
 				Field.FromInt((int) threadId),
 				Field.FromInt((int) systemId),
 				Field.FromInt((int) sortHint),
 				Field.FromString(name),
-			};
+			];
 
 			_genericEvent = new GenericEvent(0, fields, EventType);
 		}
diff --git a/Engine/Source/Programs/Shared/EpicGames.Tracing/UnrealInsights/Reader.cs b/Engine/Source/Programs/Shared/EpicGames.Tracing/UnrealInsights/Reader.cs
--- a/Engine/Source/Programs/Shared/EpicGames.Tracing/UnrealInsights/Reader.cs
+++ b/Engine/Source/Programs/Shared/EpicGames.Tracing/UnrealInsights/Reader.cs
@@ -32,8 +32,8 @@ namespace EpicGames.Tracing.UnrealInsights
 
 	public class UnrealInsightsReader
 	{
-		public Dictionary<ushort, EventType> EventTypes { get; } = new Dictionary<ushort, EventType>();
-		public Dictionary<ushort, List<ITraceEvent>> EventsPerThread { get; } = new Dictionary<ushort, List<ITraceEvent>>();
+		public Dictionary<ushort, EventType> EventTypes { get; } = [];
+		public Dictionary<ushort, List<ITraceEvent>> EventsPerThread { get; } = [];
 
 		internal int NumTransportPacketsRead { get; private set; } = 0;
 
@@ -83,7 +83,7 @@ namespace EpicGames.Tracing.UnrealInsights
 		/// <returns>A stream of data per thread</returns>
 		public Dictionary<ushort, MemoryStream> Demux(BinaryReader reader)
 		{
-			Dictionary<ushort, MemoryStream> threadStreams = new Dictionary<ushort, MemoryStream>();
+			Dictionary<ushort, MemoryStream> threadStreams = [];
 			while (reader.BaseStream.Position < reader.BaseStream.Length)
 			{
 				TransportPacket packet = TransportPacket.Deserialize(reader);
@@ -107,7 +107,7 @@ namespace EpicGames.Tracing.UnrealInsights
 
 		public List<ITraceEvent> ReadEvents(int threadId, BinaryReader reader)
 		{
-			List<ITraceEvent> eventsRead = new List<ITraceEvent>();
+			List<ITraceEvent> eventsRead = [];
 			while (reader.BaseStream.Position != reader.BaseStream.Length)
 			{
 				ushort uid;
@@ -172,7 +172,7 @@ namespace EpicGames.Tracing.UnrealInsights
 
 		public Dictionary<ushort, List<ITraceEvent>> GetEventsPerUid()
 		{
-			Dictionary<ushort, List<ITraceEvent>> eventsPerUid = new Dictionary<ushort, List<ITraceEvent>>();
+			Dictionary<ushort, List<ITraceEvent>> eventsPerUid = [];
 
 			foreach ((ushort _, List<ITraceEvent> threadEvents) in EventsPerThread)
 			{
@@ -181,7 +181,7 @@ namespace EpicGames.Tracing.UnrealInsights
 					ushort uid = EventTypes.First(x => @event.Type.Name == x.Value.Name).Key;
 					if (!eventsPerUid.TryGetValue(uid, out List<ITraceEvent>? events))
 					{
-						events = new List<ITraceEvent>();
+						events = [];
 						eventsPerUid[uid] = events;
 					}
 					events.Add(@event);
diff --git a/Engine/Source/Programs/Shared/EpicGames.Tracing/UnrealInsights/TransportPacket.cs b/Engine/Source/Programs/Shared/EpicGames.Tracing/UnrealInsights/TransportPacket.cs
--- a/Engine/Source/Programs/Shared/EpicGames.Tracing/UnrealInsights/TransportPacket.cs
+++ b/Engine/Source/Programs/Shared/EpicGames.Tracing/UnrealInsights/TransportPacket.cs
@@ -1,6 +1,5 @@
 // Copyright Epic Games, Inc. All Rights Reserved.
 
-using System;
 using System.Collections.Generic;
 using System.IO;
 using System.Linq;
@@ -21,7 +20,7 @@ namespace EpicGames.Tracing.UnrealInsights
 		public ushort PacketSize { get; private set; }
 		public ushort ThreadIdAndMarkers { get; private set; }
 		public ushort DecodedSize { get; private set; }
-		byte[] _data = Array.Empty<byte>();
+		byte[] _data = [];
 		public byte[] GetData()
 		{
 			return _data;
diff --git a/Engine/Source/Programs/Shared/EpicGames.Tracing/UnrealInsights/Writer.cs b/Engine/Source/Programs/Shared/EpicGames.Tracing/UnrealInsights/Writer.cs
--- a/Engine/Source/Programs/Shared/EpicGames.Tracing/UnrealInsights/Writer.cs
+++ b/Engine/Source/Programs/Shared/EpicGames.Tracing/UnrealInsights/Writer.cs
@@ -10,16 +10,16 @@ namespace EpicGames.Tracing.UnrealInsights
 {
 	public class UnrealInsightsWriter
 	{
-		readonly Dictionary<EventType, ushort> _eventTypeToUids = new Dictionary<EventType, ushort>();
-		readonly Dictionary<ushort, EventType> _uidToEventTypes = new Dictionary<ushort, EventType>();
-		readonly Dictionary<ushort, List<ITraceEvent>> _threadToEvents = new Dictionary<ushort, List<ITraceEvent>>();
+		readonly Dictionary<EventType, ushort> _eventTypeToUids = [];
+		readonly Dictionary<ushort, EventType> _uidToEventTypes = [];
+		readonly Dictionary<ushort, List<ITraceEvent>> _threadToEvents = [];
 
 		private ushort _uidCounter = PredefinedEventUid.WellKnownNum; // IDs below 16 are reserved for well-known events
 		private readonly object _lock = new object();
 
 		public UnrealInsightsWriter()
 		{
-			_threadToEvents[TransportPacket.ThreadIdEvents] = new List<ITraceEvent>();
+			_threadToEvents[TransportPacket.ThreadIdEvents] = [];
 
 			EnterScopeEvent enterScopeEvent = new EnterScopeEvent();
 			RegisterEventType(PredefinedEventUid.EnterScope, enterScopeEvent.Type);
@@ -45,7 +45,7 @@ namespace EpicGames.Tracing.UnrealInsights
 			{
 				if (!_threadToEvents.TryGetValue(threadId, out List<ITraceEvent>? events))
 				{
-					events = new List<ITraceEvent>();
+					events = [];
 					_threadToEvents[threadId] = events;
 				}
 
@@ -104,7 +104,7 @@ namespace EpicGames.Tracing.UnrealInsights
 			{
 				ushort uid = _eventTypeToUids[@event.Type];
 				TransportPacket transportPacket = TransportPacket.Create(0, threadId);
-				transportPacket.Serialize(writer, new[] { (uid, @event) });
+				transportPacket.Serialize(writer, [(uid, @event)]);
 			}
 		}
 	}
diff --git a/Engine/Source/Programs/Shared/EpicGames.UBA/Impl/CacheClientImpl.cs b/Engine/Source/Programs/Shared/EpicGames.UBA/Impl/CacheClientImpl.cs
--- a/Engine/Source/Programs/Shared/EpicGames.UBA/Impl/CacheClientImpl.cs
+++ b/Engine/Source/Programs/Shared/EpicGames.UBA/Impl/CacheClientImpl.cs
@@ -75,7 +75,7 @@ namespace EpicGames.UBA.Impl
 			{
 				if (_connectingTasks != null)
 				{
-					Task.WaitAll(_connectingTasks.ToArray());
+					Task.WaitAll([.. _connectingTasks]);
 				}
 				CacheClient_Destroy(_handle);
 				_handle = IntPtr.Zero;
@@ -98,7 +98,7 @@ namespace EpicGames.UBA.Impl
 
 			if (desiredConnectionCount > 1)
 			{
-				_connectingTasks = new();
+				_connectingTasks = [];
 				for (int i = 1; i < desiredConnectionCount; ++i)
 				{
 					_connectingTasks.Add(Task.Run(() =>
@@ -115,7 +115,7 @@ namespace EpicGames.UBA.Impl
 		{
 			if (_connectingTasks != null)
 			{
-				Task.WaitAll(_connectingTasks.ToArray());
+				Task.WaitAll([.. _connectingTasks]);
 			}
 
 			CacheClient_Disconnect(_handle);
