UnrealBuildTool.Tests: Missing files from last commit

#rnx

#Synced-CL 50769868

#Synced-CL 50769868

diff --git a/Engine/Source/Programs/Shared/UnrealBuildTool.Tests/PluginDescriptorTests.cs b/Engine/Source/Programs/Shared/UnrealBuildTool.Tests/PluginDescriptorTests.cs
--- a/Engine/Source/Programs/Shared/UnrealBuildTool.Tests/PluginDescriptorTests.cs
+++ b/Engine/Source/Programs/Shared/UnrealBuildTool.Tests/PluginDescriptorTests.cs
@@ -141,11 +141,7 @@ namespace UnrealBuildTool.Tests
 
 		private static string GetAbsolutePathToTempFile(string fileName)
 		{
-			if (String.IsNullOrEmpty(fileName) || s_tempDirectory == null)
-			{
-				return "";
-			}
-			return Path.Join(s_tempDirectory.FullName, fileName);
+			return String.IsNullOrEmpty(fileName) || s_tempDirectory == null ? "" : Path.Join(s_tempDirectory.FullName, fileName);
 		}
 
 		private static FileReference CreateTempUPluginFile(string fileName, string fileContent)
@@ -332,7 +328,7 @@ namespace UnrealBuildTool.Tests
 			{
 				return verbatimString;
 			}
-			string[] lines = verbatimString.Split(new[] { Environment.NewLine }, StringSplitOptions.None);
+			string[] lines = verbatimString.Split([Environment.NewLine], StringSplitOptions.None);
 			StringBuilder jsonStringBuilder = new StringBuilder();
 			int numberOfTabs = 0;
 			bool bFoundTabLevel = false;
diff --git a/Engine/Source/Programs/Shared/UnrealBuildTool.Tests/PreprocessorStateTests.cs b/Engine/Source/Programs/Shared/UnrealBuildTool.Tests/PreprocessorStateTests.cs
--- a/Engine/Source/Programs/Shared/UnrealBuildTool.Tests/PreprocessorStateTests.cs
+++ b/Engine/Source/Programs/Shared/UnrealBuildTool.Tests/PreprocessorStateTests.cs
@@ -133,14 +133,9 @@ namespace UnrealBuildTool.Tests
 				tokens.Add(reader.Current);
 			}
 
-			if (parameters == null)
-			{
-				return new PreprocessorMacro(Identifier.FindOrAdd(name), null, tokens);
-			}
-			else
-			{
-				return new PreprocessorMacro(Identifier.FindOrAdd(name), parameters.Select(x => Identifier.FindOrAdd(x)).ToList(), tokens);
-			}
+			return parameters == null
+				? new PreprocessorMacro(Identifier.FindOrAdd(name), null, tokens)
+				: new PreprocessorMacro(Identifier.FindOrAdd(name), [.. parameters.Select(Identifier.FindOrAdd)], tokens);
 		}
 	}
 }
diff --git a/Engine/Source/Programs/Shared/UnrealBuildTool.Tests/TargetRulesTests.cs b/Engine/Source/Programs/Shared/UnrealBuildTool.Tests/TargetRulesTests.cs
--- a/Engine/Source/Programs/Shared/UnrealBuildTool.Tests/TargetRulesTests.cs
+++ b/Engine/Source/Programs/Shared/UnrealBuildTool.Tests/TargetRulesTests.cs
@@ -1,11 +1,10 @@
 // Copyright Epic Games, Inc. All Rights Reserved.
 
-using System.Collections.Generic;
 using System.IO;
 using EpicGames.Core;
 using Microsoft.VisualStudio.TestTools.UnitTesting;
-using UnrealBuildTool.Tests.TestUtilities;
 using UnrealBuildBase;
+using UnrealBuildTool.Tests.TestUtilities;
 
 namespace UnrealBuildTool.Tests
 {
@@ -33,7 +32,7 @@ namespace UnrealBuildTool.Tests
 
 				// Do not mutate the object
 				{
-					bool result = testTargetRules.RequiresUniqueEnvironment(generatedTestRulesAssembly, null, new Dictionary<string, (string?, string?)>(), out baseTarget);
+					bool result = testTargetRules.RequiresUniqueEnvironment(generatedTestRulesAssembly, null, [], out baseTarget);
 
 					Assert.IsFalse(result);
 				}
@@ -42,7 +41,7 @@ namespace UnrealBuildTool.Tests
 				{
 					WarningLevel preWarningLevel = testTargetRules.CppCompileWarningSettings.ShadowVariableWarningLevel;
 					testTargetRules.CppCompileWarningSettings.ShadowVariableWarningLevel = WarningLevel.Warning;
-					bool result = testTargetRules.RequiresUniqueEnvironment(generatedTestRulesAssembly, null, new Dictionary<string, (string?, string?)>(), out baseTarget);
+					bool result = testTargetRules.RequiresUniqueEnvironment(generatedTestRulesAssembly, null, [], out baseTarget);
 
 					Assert.IsTrue(result);
 
@@ -53,7 +52,7 @@ namespace UnrealBuildTool.Tests
 				{
 					WarningLevel preWarningLevel = testTargetRules.CppCompileWarningSettings.NonTrivialMemAccessWarningLevel;
 					testTargetRules.CppCompileWarningSettings.NonTrivialMemAccessWarningLevel = WarningLevel.Warning;
-					bool result = testTargetRules.RequiresUniqueEnvironment(generatedTestRulesAssembly, null, new Dictionary<string, (string?, string?)>(), out baseTarget);
+					bool result = testTargetRules.RequiresUniqueEnvironment(generatedTestRulesAssembly, null, [], out baseTarget);
 
 					Assert.IsTrue(result);
 
@@ -64,12 +63,12 @@ namespace UnrealBuildTool.Tests
 				{
 					WarningLevel preWarningLevel = testTargetRules.CppCompileWarningSettings.DeprecationWarningLevel;
 					testTargetRules.CppCompileWarningSettings.DeprecationWarningLevel = WarningLevel.Error;
-					bool result = testTargetRules.RequiresUniqueEnvironment(generatedTestRulesAssembly, null, new Dictionary<string, (string?, string?)>(), out baseTarget);
+					bool result = testTargetRules.RequiresUniqueEnvironment(generatedTestRulesAssembly, null, [], out baseTarget);
 
 					Assert.IsTrue(result);
 
 					testTargetRules.CppCompileWarningSettings.DeprecationWarningLevel = WarningLevel.Warning;
-					result = testTargetRules.RequiresUniqueEnvironment(generatedTestRulesAssembly, null, new Dictionary<string, (string?, string?)>(), out baseTarget);
+					result = testTargetRules.RequiresUniqueEnvironment(generatedTestRulesAssembly, null, [], out baseTarget);
 
 					Assert.IsFalse(result);
 
@@ -82,7 +81,7 @@ namespace UnrealBuildTool.Tests
 					testTargetRules.CppCompileWarningSettings.ShadowVariableWarningLevel = WarningLevel.Warning;
 					testTargetRules.bOverrideBuildEnvironment = true;
 
-					bool result = testTargetRules.RequiresUniqueEnvironment(generatedTestRulesAssembly, null, new Dictionary<string, (string?, string?)>(), out baseTarget);
+					bool result = testTargetRules.RequiresUniqueEnvironment(generatedTestRulesAssembly, null, [], out baseTarget);
 
 					Assert.IsFalse(result);
 
diff --git a/Engine/Source/Programs/Shared/UnrealBuildTool.Tests/TokenReaderTests.cs b/Engine/Source/Programs/Shared/UnrealBuildTool.Tests/TokenReaderTests.cs
--- a/Engine/Source/Programs/Shared/UnrealBuildTool.Tests/TokenReaderTests.cs
+++ b/Engine/Source/Programs/Shared/UnrealBuildTool.Tests/TokenReaderTests.cs
@@ -141,7 +141,7 @@ namespace UnrealBuildTool.Tests
 				tokens.Add(reader.Current);
 			}
 
-			string outputText = String.Join(", ", tokens.Select(x => FormatToken(x)));
+			string outputText = String.Join(", ", tokens.Select(FormatToken));
 			Assert.AreEqual(expectedOutputText, outputText);
 		}
 
diff --git a/Engine/Source/Programs/Shared/UnrealBuildTool.Tests/WarningLevelDefaultAttributeTests.cs b/Engine/Source/Programs/Shared/UnrealBuildTool.Tests/WarningLevelDefaultAttributeTests.cs
--- a/Engine/Source/Programs/Shared/UnrealBuildTool.Tests/WarningLevelDefaultAttributeTests.cs
+++ b/Engine/Source/Programs/Shared/UnrealBuildTool.Tests/WarningLevelDefaultAttributeTests.cs
@@ -22,12 +22,12 @@ namespace UnrealBuildTool.Tests
 		[TestMethod]
 		public void TestWindowsCompletelyDefined()
 		{
-			List<VersionWarningLevelDefaultAttribute> attributes = new List<VersionWarningLevelDefaultAttribute>
-			{
+			List<VersionWarningLevelDefaultAttribute> attributes =
+			[
 				new VersionWarningLevelDefaultAttribute( WarningLevel.Off, BuildSettingsVersion.V1, BuildSettingsVersion.V2),
 				new VersionWarningLevelDefaultAttribute(WarningLevel.Warning, BuildSettingsVersion.V3, BuildSettingsVersion.V4),
 				new VersionWarningLevelDefaultAttribute(WarningLevel.Error, BuildSettingsVersion.V5, BuildSettingsVersion.Latest)
-			};
+			];
 
 			List<VersionWarningLevelDefaultAttribute> merged = VersionWarningLevelDefaultAttribute.EnsureWarningLevelDefaultBounds(attributes);
 
@@ -48,11 +48,11 @@ namespace UnrealBuildTool.Tests
 		[TestMethod]
 		public void TestGapResolveToLowerBound()
 		{
-			List<VersionWarningLevelDefaultAttribute> attributes = new List<VersionWarningLevelDefaultAttribute>
-			{
+			List<VersionWarningLevelDefaultAttribute> attributes =
+			[
 				new VersionWarningLevelDefaultAttribute(WarningLevel.Warning, BuildSettingsVersion.V1, BuildSettingsVersion.V2),
 				new VersionWarningLevelDefaultAttribute(WarningLevel.Error, BuildSettingsVersion.V4, BuildSettingsVersion.V5),
-			};
+			];
 
 			List<VersionWarningLevelDefaultAttribute> merged = VersionWarningLevelDefaultAttribute.EnsureWarningLevelDefaultBounds(attributes);
 
@@ -69,11 +69,11 @@ namespace UnrealBuildTool.Tests
 		[TestMethod]
 		public void TestLowerUpperOverlapResizeBottomUpward()
 		{
-			List<VersionWarningLevelDefaultAttribute> attributes = new List<VersionWarningLevelDefaultAttribute>
-			{
+			List<VersionWarningLevelDefaultAttribute> attributes =
+			[
 				new VersionWarningLevelDefaultAttribute(WarningLevel.Warning, BuildSettingsVersion.V1, BuildSettingsVersion.V3),
 				new VersionWarningLevelDefaultAttribute(WarningLevel.Error, BuildSettingsVersion.V2, BuildSettingsVersion.V4)
-			};
+			];
 
 			List<VersionWarningLevelDefaultAttribute> merged = VersionWarningLevelDefaultAttribute.EnsureWarningLevelDefaultBounds(attributes);
 
@@ -90,11 +90,11 @@ namespace UnrealBuildTool.Tests
 		[TestMethod]
 		public void TestAllRangeWide()
 		{
-			List<VersionWarningLevelDefaultAttribute> attributes = new List<VersionWarningLevelDefaultAttribute>
-			{
+			List<VersionWarningLevelDefaultAttribute> attributes =
+			[
 				new VersionWarningLevelDefaultAttribute( WarningLevel.Warning, BuildSettingsVersion.V1, BuildSettingsVersion.V2),
 				new VersionWarningLevelDefaultAttribute(WarningLevel.Error, BuildSettingsVersion.V1, BuildSettingsVersion.V4)
-			};
+			];
 
 			List<VersionWarningLevelDefaultAttribute> merged = VersionWarningLevelDefaultAttribute.EnsureWarningLevelDefaultBounds(attributes);
 
@@ -112,11 +112,11 @@ namespace UnrealBuildTool.Tests
 		[TestMethod]
 		public void TestResolveLowerBound()
 		{
-			List<WarningLevelDefaultAttribute> attributes = new List<WarningLevelDefaultAttribute>
-			{
+			List<WarningLevelDefaultAttribute> attributes =
+			[
 				new VersionWarningLevelDefaultAttribute(WarningLevel.Warning, BuildSettingsVersion.V1, BuildSettingsVersion.V2),
 				new VersionWarningLevelDefaultAttribute(WarningLevel.Error, BuildSettingsVersion.V3, BuildSettingsVersion.V4)
-			};
+			];
 			TestBuildSettingsProvider testBuildSettingsProvider = new TestBuildSettingsProvider();
 			BuildSystemContext newBuildSystemContext = new BuildSystemContext(testBuildSettingsProvider);
 
