EpicGames.Core: Remove unneeded ReadOnlyList ListExtension

#rnx

#Synced-CL 50778211

diff --git a/Engine/Source/Programs/Shared/EpicGames.Core/ListExtensions.cs b/Engine/Source/Programs/Shared/EpicGames.Core/ListExtensions.cs
--- a/Engine/Source/Programs/Shared/EpicGames.Core/ListExtensions.cs
+++ b/Engine/Source/Programs/Shared/EpicGames.Core/ListExtensions.cs
@@ -1,7 +1,6 @@
 // Copyright Epic Games, Inc. All Rights Reserved.
 
 using System;
-using System.Collections;
 using System.Collections.Generic;
 using System.Linq;
 using System.Threading.Tasks;
@@ -13,59 +12,6 @@ namespace EpicGames.Core
 	/// </summary>s
 	public static class ListExtensions
 	{
-		/// <summary>
-		/// Wrapper around a list to implement an immutable <see cref="IList{T}"/> and <see cref="IReadOnlyList{T}"/>
-		/// </summary>
-		/// <typeparam name="T"></typeparam>
-		class ReadOnlyList<T> : IList<T>, IReadOnlyList<T>
-		{
-			readonly IList<T> _inner;
-
-			private ReadOnlyList(IList<T> inner) => _inner = inner;
-
-			public static ReadOnlyList<T> Wrap(IList<T> inner) => new ReadOnlyList<T>(inner);
-
-			static NotSupportedException CreateNotSupportedException()
-				=> new NotSupportedException("Collection is read-only.");
-
-			public T this[int index]
-			{
-				get => _inner[index];
-				set => throw CreateNotSupportedException();
-			}
-
-			public int Count => _inner.Count;
-			public bool IsReadOnly => _inner.IsReadOnly;
-
-			public void Add(T item) => throw CreateNotSupportedException();
-			public void Clear() => throw CreateNotSupportedException();
-			public bool Contains(T item) => _inner.Contains(item);
-			public void CopyTo(T[] array, int arrayIndex) => _inner.CopyTo(array, arrayIndex);
-			public IEnumerator<T> GetEnumerator() => _inner.GetEnumerator();
-			public int IndexOf(T item) => _inner.IndexOf(item);
-			public void Insert(int index, T item) => throw CreateNotSupportedException();
-			public bool Remove(T item) => throw CreateNotSupportedException();
-			public void RemoveAt(int index) => throw CreateNotSupportedException();
-
-			IEnumerator IEnumerable.GetEnumerator() => ((IEnumerable)_inner).GetEnumerator();
-		}
-
-		/// <summary>
-		/// Create a read-only wrapper around a list
-		/// </summary>
-		/// <typeparam name="T"></typeparam>
-		/// <param name="list"></param>
-		/// <returns></returns>
-		public static IList<T> AsReadOnly<T>(this IList<T> list) => ReadOnlyList<T>.Wrap(list);
-
-		/// <summary>
-		/// Create a read-only wrapper around a list
-		/// </summary>
-		/// <typeparam name="T"></typeparam>
-		/// <param name="list"></param>
-		/// <returns></returns>
-		public static IReadOnlyList<T> AsReadOnlyList<T>(this IList<T> list) => ReadOnlyList<T>.Wrap(list);
-
 		/// <summary>
 		/// Sorts a list by a particular field
 		/// </summary>
diff --git a/Engine/Source/Programs/UnrealGameSync/UnrealGameSyncShared/ArchiveInfo.cs b/Engine/Source/Programs/UnrealGameSync/UnrealGameSyncShared/ArchiveInfo.cs
--- a/Engine/Source/Programs/UnrealGameSync/UnrealGameSyncShared/ArchiveInfo.cs
+++ b/Engine/Source/Programs/UnrealGameSync/UnrealGameSyncShared/ArchiveInfo.cs
@@ -73,7 +73,7 @@ namespace UnrealGameSync
 
 		public IArchive? TryGetArchiveForChangeNumber(int changeNumber, int maxChangeNumber)
 		{
-			int idx = ChangeNumberToArchive.Keys.AsReadOnlyList().BinarySearch(changeNumber);
+			int idx = ChangeNumberToArchive.Keys.AsReadOnly().BinarySearch(changeNumber);
 			if (idx >= 0)
 			{
 				return ChangeNumberToArchive.Values[idx];
