[UHT] Fix false positive for plain C++ interface with same suffix as UClass.

Prior to this change, `class IFoo { ... }; UCLASS() class UFoo : UObject  { ... };` would result in an error claiming that IFoo is missing GENERATED_BODY().

#Synced-CL 50722436

diff --git a/Engine/Source/Programs/Shared/EpicGames.UHT/Types/UhtClass.cs b/Engine/Source/Programs/Shared/EpicGames.UHT/Types/UhtClass.cs
--- a/Engine/Source/Programs/Shared/EpicGames.UHT/Types/UhtClass.cs
+++ b/Engine/Source/Programs/Shared/EpicGames.UHT/Types/UhtClass.cs
@@ -1061,7 +1061,8 @@ namespace EpicGames.UHT.Types
 							{
 								string interfaceName = GetInterfaceName();
 								UhtClass? interfaceObj = (UhtClass?)Session.FindType(null, UhtFindOptions.SourceName | UhtFindOptions.Class, interfaceName);
-								if (interfaceObj == null)
+								// Only consider actual UINTERFACE classes, not regular UObject classes that happen to have a matching name
+								if (interfaceObj?.ClassType != UhtClassType.Interface)
 								{
 									if (HasGeneratedBody || Children.Count != 0)
 									{
