Shared: Move default TargetFramework to shared property sheet

#rnx
#Synced-CL 50778211

diff --git a/Engine/Source/Programs/Shared/EpicGames.Analytics.PerformanceTrends/EpicGames.Analytics.PerformanceTrends.csproj b/Engine/Source/Programs/Shared/EpicGames.Analytics.PerformanceTrends/EpicGames.Analytics.PerformanceTrends.csproj
--- a/Engine/Source/Programs/Shared/EpicGames.Analytics.PerformanceTrends/EpicGames.Analytics.PerformanceTrends.csproj
+++ b/Engine/Source/Programs/Shared/EpicGames.Analytics.PerformanceTrends/EpicGames.Analytics.PerformanceTrends.csproj
@@ -1,6 +1,5 @@
 ﻿<Project Sdk="Microsoft.NET.Sdk">
   <PropertyGroup>
-    <TargetFramework>net8.0</TargetFramework>
     <Deterministic>true</Deterministic>
     <RootNamespace>EpicGames.Analytics</RootNamespace>
     <GenerateDocumentationFile>True</GenerateDocumentationFile>
diff --git a/Engine/Source/Programs/Shared/EpicGames.Analytics.Telemetry/EpicGames.Analytics.Telemetry.csproj b/Engine/Source/Programs/Shared/EpicGames.Analytics.Telemetry/EpicGames.Analytics.Telemetry.csproj
--- a/Engine/Source/Programs/Shared/EpicGames.Analytics.Telemetry/EpicGames.Analytics.Telemetry.csproj
+++ b/Engine/Source/Programs/Shared/EpicGames.Analytics.Telemetry/EpicGames.Analytics.Telemetry.csproj
@@ -1,6 +1,5 @@
 ﻿<Project Sdk="Microsoft.NET.Sdk">
   <PropertyGroup>
-    <TargetFramework>net8.0</TargetFramework>
     <RootNamespace>EpicGames.Analytics</RootNamespace>
   </PropertyGroup>
 </Project>
diff --git a/Engine/Source/Programs/Shared/EpicGames.Benchmarks/EpicGames.Benchmarks.csproj b/Engine/Source/Programs/Shared/EpicGames.Benchmarks/EpicGames.Benchmarks.csproj
--- a/Engine/Source/Programs/Shared/EpicGames.Benchmarks/EpicGames.Benchmarks.csproj
+++ b/Engine/Source/Programs/Shared/EpicGames.Benchmarks/EpicGames.Benchmarks.csproj
@@ -1,6 +1,5 @@
 ﻿<Project Sdk="Microsoft.NET.Sdk">
   <PropertyGroup>
-    <TargetFramework>net8.0</TargetFramework>
     <OutputType>Exe</OutputType>
     <NoWarn>$(NoWarn);1591</NoWarn>
   </PropertyGroup>
diff --git a/Engine/Source/Programs/Shared/EpicGames.Build.Tests/EpicGames.Build.Tests.csproj b/Engine/Source/Programs/Shared/EpicGames.Build.Tests/EpicGames.Build.Tests.csproj
--- a/Engine/Source/Programs/Shared/EpicGames.Build.Tests/EpicGames.Build.Tests.csproj
+++ b/Engine/Source/Programs/Shared/EpicGames.Build.Tests/EpicGames.Build.Tests.csproj
@@ -1,8 +1,4 @@
-﻿<Project Sdk="Microsoft.NET.Sdk">  
-  <PropertyGroup>
-    <TargetFramework>net8.0</TargetFramework>
-  </PropertyGroup>
-
+﻿<Project Sdk="Microsoft.NET.Sdk">
   <ItemGroup>
     <ProjectReference Include="..\EpicGames.Build\EpicGames.Build.csproj" />
   </ItemGroup>
diff --git a/Engine/Source/Programs/Shared/EpicGames.Build/EpicGames.Build.csproj b/Engine/Source/Programs/Shared/EpicGames.Build/EpicGames.Build.csproj
--- a/Engine/Source/Programs/Shared/EpicGames.Build/EpicGames.Build.csproj
+++ b/Engine/Source/Programs/Shared/EpicGames.Build/EpicGames.Build.csproj
@@ -1,6 +1,5 @@
 ﻿<Project Sdk="Microsoft.NET.Sdk">
   <PropertyGroup>
-    <TargetFramework>net8.0</TargetFramework>
     <RootNamespace>UnrealBuildBase</RootNamespace>
     <AllowUnsafeBlocks>true</AllowUnsafeBlocks>
   </PropertyGroup>
@@ -16,8 +15,6 @@
   </ItemGroup>
 
   <ItemGroup>
-    <PackageReference Include="Microsoft.CSharp" Version="4.7.0" />
     <PackageReference Include="Microsoft.Extensions.FileSystemGlobbing" Version="8.0.0" />
-    <PackageReference Include="System.Security.Permissions" Version="8.0.0" />
   </ItemGroup>
 </Project>
diff --git a/Engine/Source/Programs/Shared/EpicGames.BuildGraph.Tests/EpicGames.BuildGraph.Tests.csproj b/Engine/Source/Programs/Shared/EpicGames.BuildGraph.Tests/EpicGames.BuildGraph.Tests.csproj
--- a/Engine/Source/Programs/Shared/EpicGames.BuildGraph.Tests/EpicGames.BuildGraph.Tests.csproj
+++ b/Engine/Source/Programs/Shared/EpicGames.BuildGraph.Tests/EpicGames.BuildGraph.Tests.csproj
@@ -1,8 +1,4 @@
 ﻿<Project Sdk="Microsoft.NET.Sdk">
-  <PropertyGroup>
-    <TargetFramework>net8.0</TargetFramework>
-  </PropertyGroup>
-
   <ItemGroup>
     <PackageReference Include="System.Text.Json" Version="8.0.6" />
   </ItemGroup>
@@ -10,5 +6,4 @@
   <ItemGroup>
     <ProjectReference Include="..\EpicGames.BuildGraph\EpicGames.BuildGraph.csproj" />
   </ItemGroup>
-
 </Project>
diff --git a/Engine/Source/Programs/Shared/EpicGames.Core.Tests/EpicGames.Core.Tests.csproj b/Engine/Source/Programs/Shared/EpicGames.Core.Tests/EpicGames.Core.Tests.csproj
--- a/Engine/Source/Programs/Shared/EpicGames.Core.Tests/EpicGames.Core.Tests.csproj
+++ b/Engine/Source/Programs/Shared/EpicGames.Core.Tests/EpicGames.Core.Tests.csproj
@@ -1,11 +1,9 @@
 <Project Sdk="Microsoft.NET.Sdk">
   <PropertyGroup>
-    <TargetFramework>net8.0</TargetFramework>
     <TreatWarningsAsErrors>False</TreatWarningsAsErrors>
   </PropertyGroup>
 
   <ItemGroup>
     <ProjectReference Include="..\EpicGames.Core\EpicGames.Core.csproj" />
   </ItemGroup>
-
 </Project>
diff --git a/Engine/Source/Programs/Shared/EpicGames.Core/EpicGames.Core.csproj b/Engine/Source/Programs/Shared/EpicGames.Core/EpicGames.Core.csproj
--- a/Engine/Source/Programs/Shared/EpicGames.Core/EpicGames.Core.csproj
+++ b/Engine/Source/Programs/Shared/EpicGames.Core/EpicGames.Core.csproj
@@ -1,6 +1,5 @@
 ﻿<Project Sdk="Microsoft.NET.Sdk">
   <PropertyGroup>
-    <TargetFramework>net8.0</TargetFramework>
     <AllowUnsafeBlocks>true</AllowUnsafeBlocks>
   </PropertyGroup>
 
@@ -11,7 +10,6 @@
     <PackageReference Include="Microsoft.Extensions.Logging.Console" Version="8.0.1" />
     <PackageReference Include="Microsoft.Extensions.ObjectPool" Version="8.0.20" />
     <PackageReference Include="OpenTracing" Version="0.12.1" />
-    <PackageReference Include="System.Memory" Version="4.6.3" />
     <PackageReference Include="System.Text.Json" Version="8.0.6" />
   </ItemGroup>
 </Project>
diff --git a/Engine/Source/Programs/Shared/EpicGames.Horde.Tests/EpicGames.Horde.Tests.csproj b/Engine/Source/Programs/Shared/EpicGames.Horde.Tests/EpicGames.Horde.Tests.csproj
--- a/Engine/Source/Programs/Shared/EpicGames.Horde.Tests/EpicGames.Horde.Tests.csproj
+++ b/Engine/Source/Programs/Shared/EpicGames.Horde.Tests/EpicGames.Horde.Tests.csproj
@@ -1,8 +1,4 @@
 <Project Sdk="Microsoft.NET.Sdk">
-  <PropertyGroup>
-    <TargetFramework>net8.0</TargetFramework>
-  </PropertyGroup>
-
   <ItemGroup>
     <PackageReference Include="System.Linq.Async" Version="6.0.3" />
   </ItemGroup>
@@ -26,5 +22,4 @@
       <LastGenOutput>Resources.Designer.cs</LastGenOutput>
     </EmbeddedResource>
   </ItemGroup>
-
 </Project>
diff --git a/Engine/Source/Programs/Shared/EpicGames.Horde/EpicGames.Horde.csproj b/Engine/Source/Programs/Shared/EpicGames.Horde/EpicGames.Horde.csproj
--- a/Engine/Source/Programs/Shared/EpicGames.Horde/EpicGames.Horde.csproj
+++ b/Engine/Source/Programs/Shared/EpicGames.Horde/EpicGames.Horde.csproj
@@ -1,6 +1,5 @@
 ﻿<Project Sdk="Microsoft.NET.Sdk">
   <PropertyGroup>
-    <TargetFramework>net8.0</TargetFramework>
     <GenerateDocumentationFile>True</GenerateDocumentationFile>
     <AllowUnsafeBlocks>True</AllowUnsafeBlocks>
   </PropertyGroup>
@@ -17,7 +16,6 @@
     <PackageReference Include="K4os.Compression.LZ4" Version="1.2.16" />
     <PackageReference Include="System.IO.Pipelines" Version="8.0.0" />
     <PackageReference Include="Microsoft.Data.Sqlite" Version="8.0.20" />
-    <PackageReference Include="Microsoft.Extensions.Caching.Abstractions" Version="8.0.0" />
     <PackageReference Include="Microsoft.Extensions.Caching.Memory" Version="8.0.1" />
     <PackageReference Include="Microsoft.VisualStudio.Threading.Analyzers" Version="17.14.15">
       <PrivateAssets>all</PrivateAssets>
diff --git a/Engine/Source/Programs/Shared/EpicGames.IoHash.Tests/EpicGames.IoHash.Tests.csproj b/Engine/Source/Programs/Shared/EpicGames.IoHash.Tests/EpicGames.IoHash.Tests.csproj
--- a/Engine/Source/Programs/Shared/EpicGames.IoHash.Tests/EpicGames.IoHash.Tests.csproj
+++ b/Engine/Source/Programs/Shared/EpicGames.IoHash.Tests/EpicGames.IoHash.Tests.csproj
@@ -1,10 +1,5 @@
 <Project Sdk="Microsoft.NET.Sdk">
-  <PropertyGroup>
-    <TargetFramework>net8.0</TargetFramework>
-  </PropertyGroup>
-
   <ItemGroup>
     <ProjectReference Include="..\EpicGames.IoHash\EpicGames.IoHash.csproj" />
   </ItemGroup>
-
 </Project>
diff --git a/Engine/Source/Programs/Shared/EpicGames.IoHash/EpicGames.IoHash.csproj b/Engine/Source/Programs/Shared/EpicGames.IoHash/EpicGames.IoHash.csproj
--- a/Engine/Source/Programs/Shared/EpicGames.IoHash/EpicGames.IoHash.csproj
+++ b/Engine/Source/Programs/Shared/EpicGames.IoHash/EpicGames.IoHash.csproj
@@ -1,6 +1,5 @@
 ﻿<Project Sdk="Microsoft.NET.Sdk">
   <PropertyGroup>
-    <TargetFramework>net8.0</TargetFramework>
     <RootNamespace>EpicGames.Core</RootNamespace>
   </PropertyGroup>
 
diff --git a/Engine/Source/Programs/Shared/EpicGames.LogMatchers.Tests/EpicGames.LogMatchers.Tests.csproj b/Engine/Source/Programs/Shared/EpicGames.LogMatchers.Tests/EpicGames.LogMatchers.Tests.csproj
--- a/Engine/Source/Programs/Shared/EpicGames.LogMatchers.Tests/EpicGames.LogMatchers.Tests.csproj
+++ b/Engine/Source/Programs/Shared/EpicGames.LogMatchers.Tests/EpicGames.LogMatchers.Tests.csproj
@@ -1,8 +1,4 @@
 <Project Sdk="Microsoft.NET.Sdk">
-  <PropertyGroup>
-    <TargetFramework>net8.0</TargetFramework>
-  </PropertyGroup>
-
   <ItemGroup>
     <ProjectReference Include="..\EpicGames.Core\EpicGames.Core.csproj" />
     <ProjectReference Include="..\EpicGames.LogMatchers\EpicGames.LogMatchers.csproj" />
diff --git a/Engine/Source/Programs/Shared/EpicGames.MongoDB/EpicGames.MongoDB.csproj b/Engine/Source/Programs/Shared/EpicGames.MongoDB/EpicGames.MongoDB.csproj
--- a/Engine/Source/Programs/Shared/EpicGames.MongoDB/EpicGames.MongoDB.csproj
+++ b/Engine/Source/Programs/Shared/EpicGames.MongoDB/EpicGames.MongoDB.csproj
@@ -1,8 +1,4 @@
 ﻿<Project Sdk="Microsoft.NET.Sdk">
-  <PropertyGroup>
-    <TargetFramework>net8.0</TargetFramework>
-  </PropertyGroup>
-
   <ItemGroup>
     <PackageReference Include="MongoDB.Driver" Version="2.19.2" />
   </ItemGroup>
diff --git a/Engine/Source/Programs/Shared/EpicGames.MsBuild/EpicGames.MsBuild.csproj b/Engine/Source/Programs/Shared/EpicGames.MsBuild/EpicGames.MsBuild.csproj
--- a/Engine/Source/Programs/Shared/EpicGames.MsBuild/EpicGames.MsBuild.csproj
+++ b/Engine/Source/Programs/Shared/EpicGames.MsBuild/EpicGames.MsBuild.csproj
@@ -1,8 +1,4 @@
 ﻿<Project Sdk="Microsoft.NET.Sdk">
-  <PropertyGroup>
-    <TargetFramework>net8.0</TargetFramework>
-  </PropertyGroup>
-
   <ItemGroup>
     <PackageReference Include="Microsoft.Build" Version="17.11.48" ExcludeAssets="runtime" />
     <PackageReference Include="Microsoft.Build.Locator" Version="1.9.1" />
diff --git a/Engine/Source/Programs/Shared/EpicGames.OIDC/EpicGames.OIDC.csproj b/Engine/Source/Programs/Shared/EpicGames.OIDC/EpicGames.OIDC.csproj
--- a/Engine/Source/Programs/Shared/EpicGames.OIDC/EpicGames.OIDC.csproj
+++ b/Engine/Source/Programs/Shared/EpicGames.OIDC/EpicGames.OIDC.csproj
@@ -1,6 +1,5 @@
 ﻿<Project Sdk="Microsoft.NET.Sdk">
   <PropertyGroup>
-    <TargetFramework>net8.0</TargetFramework>
     <EnableConfigurationBindingGenerator>true</EnableConfigurationBindingGenerator>
   </PropertyGroup>
 
diff --git a/Engine/Source/Programs/Shared/EpicGames.Oodle/EpicGames.Oodle.csproj b/Engine/Source/Programs/Shared/EpicGames.Oodle/EpicGames.Oodle.csproj
--- a/Engine/Source/Programs/Shared/EpicGames.Oodle/EpicGames.Oodle.csproj
+++ b/Engine/Source/Programs/Shared/EpicGames.Oodle/EpicGames.Oodle.csproj
@@ -1,6 +1,5 @@
 ﻿<Project Sdk="Microsoft.NET.Sdk">
   <PropertyGroup>
-    <TargetFramework>net8.0</TargetFramework>
     <RootNamespace>EpicGames.Compression</RootNamespace>
     <GenerateDocumentationFile>True</GenerateDocumentationFile>
     <AllowUnsafeBlocks>True</AllowUnsafeBlocks>
diff --git a/Engine/Source/Programs/Shared/EpicGames.Perforce.Fixture/EpicGames.Perforce.Fixture.csproj b/Engine/Source/Programs/Shared/EpicGames.Perforce.Fixture/EpicGames.Perforce.Fixture.csproj
--- a/Engine/Source/Programs/Shared/EpicGames.Perforce.Fixture/EpicGames.Perforce.Fixture.csproj
+++ b/Engine/Source/Programs/Shared/EpicGames.Perforce.Fixture/EpicGames.Perforce.Fixture.csproj
@@ -1,6 +1,5 @@
 <Project Sdk="Microsoft.NET.Sdk">
   <PropertyGroup>
-    <TargetFramework>net8.0</TargetFramework>
     <StartupObject>EpicGames.Perforce.Fixture.Program</StartupObject>
     <IsPackable>false</IsPackable>
     <NoWarn>$(NoWarn);1591</NoWarn>
@@ -27,5 +26,4 @@
   <ItemGroup>
     <ProjectReference Include="..\EpicGames.Perforce.Managed\EpicGames.Perforce.Managed.csproj" />
   </ItemGroup>
-
 </Project>
diff --git a/Engine/Source/Programs/Shared/EpicGames.Perforce.Managed.Tests/EpicGames.Perforce.Managed.Tests.csproj b/Engine/Source/Programs/Shared/EpicGames.Perforce.Managed.Tests/EpicGames.Perforce.Managed.Tests.csproj
--- a/Engine/Source/Programs/Shared/EpicGames.Perforce.Managed.Tests/EpicGames.Perforce.Managed.Tests.csproj
+++ b/Engine/Source/Programs/Shared/EpicGames.Perforce.Managed.Tests/EpicGames.Perforce.Managed.Tests.csproj
@@ -1,11 +1,6 @@
 <Project Sdk="Microsoft.NET.Sdk">
-  <PropertyGroup>
-    <TargetFramework>net8.0</TargetFramework>
-  </PropertyGroup>
-
   <ItemGroup>
     <ProjectReference Include="..\EpicGames.Perforce.Managed\EpicGames.Perforce.Managed.csproj" />
     <ProjectReference Include="..\EpicGames.Perforce.Fixture\EpicGames.Perforce.Fixture.csproj" />
   </ItemGroup>
-
 </Project>
diff --git a/Engine/Source/Programs/Shared/EpicGames.Perforce.Managed/EpicGames.Perforce.Managed.csproj b/Engine/Source/Programs/Shared/EpicGames.Perforce.Managed/EpicGames.Perforce.Managed.csproj
--- a/Engine/Source/Programs/Shared/EpicGames.Perforce.Managed/EpicGames.Perforce.Managed.csproj
+++ b/Engine/Source/Programs/Shared/EpicGames.Perforce.Managed/EpicGames.Perforce.Managed.csproj
@@ -1,9 +1,4 @@
 ﻿<Project Sdk="Microsoft.NET.Sdk">
-  <PropertyGroup>
-    <TargetFramework>net8.0</TargetFramework>
-    <DebugType>pdbonly</DebugType>
-  </PropertyGroup>
-
   <ItemGroup>
     <PackageReference Include="OpenTelemetry" Version="1.8.1" />
   </ItemGroup>
@@ -14,5 +9,4 @@
     <ProjectReference Include="..\EpicGames.Perforce\EpicGames.Perforce.csproj" />
     <ProjectReference Include="..\EpicGames.Serialization\EpicGames.Serialization.csproj" />
   </ItemGroup>
-
 </Project>
diff --git a/Engine/Source/Programs/Shared/EpicGames.Perforce.Tests/EpicGames.Perforce.Tests.csproj b/Engine/Source/Programs/Shared/EpicGames.Perforce.Tests/EpicGames.Perforce.Tests.csproj
--- a/Engine/Source/Programs/Shared/EpicGames.Perforce.Tests/EpicGames.Perforce.Tests.csproj
+++ b/Engine/Source/Programs/Shared/EpicGames.Perforce.Tests/EpicGames.Perforce.Tests.csproj
@@ -1,10 +1,5 @@
 <Project Sdk="Microsoft.NET.Sdk">
-  <PropertyGroup>
-    <TargetFramework>net8.0</TargetFramework>
-  </PropertyGroup>
-
   <ItemGroup>
     <ProjectReference Include="..\EpicGames.Perforce.Managed\EpicGames.Perforce.Managed.csproj" />
   </ItemGroup>
-
 </Project>
diff --git a/Engine/Source/Programs/Shared/EpicGames.Perforce/EpicGames.Perforce.csproj b/Engine/Source/Programs/Shared/EpicGames.Perforce/EpicGames.Perforce.csproj
--- a/Engine/Source/Programs/Shared/EpicGames.Perforce/EpicGames.Perforce.csproj
+++ b/Engine/Source/Programs/Shared/EpicGames.Perforce/EpicGames.Perforce.csproj
@@ -1,8 +1,4 @@
 ﻿<Project Sdk="Microsoft.NET.Sdk">
-  <PropertyGroup>
-    <TargetFramework>net8.0</TargetFramework>
-  </PropertyGroup>
-
   <ItemGroup>
     <None Remove="EpicGames.Perforce.uatbuildrecord" />
     <None Remove="EpicGames.Perforce.xml" />
@@ -46,5 +42,4 @@
   <ItemGroup>
     <ProjectReference Include="..\EpicGames.Core\EpicGames.Core.csproj" />
   </ItemGroup>
-
 </Project>
diff --git a/Engine/Source/Programs/Shared/EpicGames.ProjectStore/EpicGames.ProjectStore.csproj b/Engine/Source/Programs/Shared/EpicGames.ProjectStore/EpicGames.ProjectStore.csproj
--- a/Engine/Source/Programs/Shared/EpicGames.ProjectStore/EpicGames.ProjectStore.csproj
+++ b/Engine/Source/Programs/Shared/EpicGames.ProjectStore/EpicGames.ProjectStore.csproj
@@ -1,8 +1,4 @@
 ﻿<Project Sdk="Microsoft.NET.Sdk">
-  <PropertyGroup>
-    <TargetFramework>net8.0</TargetFramework>
-  </PropertyGroup>
-
   <ItemGroup>
     <ProjectReference Include="..\EpicGames.Build\EpicGames.Build.csproj" />
     <ProjectReference Include="..\EpicGames.Core\EpicGames.Core.csproj" />
diff --git a/Engine/Source/Programs/Shared/EpicGames.Redis.Tests/EpicGames.Redis.Tests.csproj b/Engine/Source/Programs/Shared/EpicGames.Redis.Tests/EpicGames.Redis.Tests.csproj
--- a/Engine/Source/Programs/Shared/EpicGames.Redis.Tests/EpicGames.Redis.Tests.csproj
+++ b/Engine/Source/Programs/Shared/EpicGames.Redis.Tests/EpicGames.Redis.Tests.csproj
@@ -1,10 +1,5 @@
 <Project Sdk="Microsoft.NET.Sdk">
-  <PropertyGroup>
-    <TargetFramework>net8.0</TargetFramework>
-  </PropertyGroup>
-
   <ItemGroup>
     <ProjectReference Include="..\EpicGames.Redis\EpicGames.Redis.csproj" />
   </ItemGroup>
-
 </Project>
diff --git a/Engine/Source/Programs/Shared/EpicGames.Redis/EpicGames.Redis.csproj b/Engine/Source/Programs/Shared/EpicGames.Redis/EpicGames.Redis.csproj
--- a/Engine/Source/Programs/Shared/EpicGames.Redis/EpicGames.Redis.csproj
+++ b/Engine/Source/Programs/Shared/EpicGames.Redis/EpicGames.Redis.csproj
@@ -1,6 +1,5 @@
 ﻿<Project Sdk="Microsoft.NET.Sdk">
   <PropertyGroup>
-    <TargetFramework>net8.0</TargetFramework>
     <GenerateDocumentationFile>True</GenerateDocumentationFile>
   </PropertyGroup>
 
diff --git a/Engine/Source/Programs/Shared/EpicGames.ScriptBuild/EpicGames.ScriptBuild.csproj b/Engine/Source/Programs/Shared/EpicGames.ScriptBuild/EpicGames.ScriptBuild.csproj
--- a/Engine/Source/Programs/Shared/EpicGames.ScriptBuild/EpicGames.ScriptBuild.csproj
+++ b/Engine/Source/Programs/Shared/EpicGames.ScriptBuild/EpicGames.ScriptBuild.csproj
@@ -1,8 +1,4 @@
 ﻿<Project Sdk="Microsoft.NET.Sdk">
-  <PropertyGroup>
-    <TargetFramework>net8.0</TargetFramework>
-  </PropertyGroup>
-
   <!-- Reference file is generated by AutomationTool. It contains a list of project references that will be compiled with 'dotnet build' -->
   <Import
     Project="$(EngineDir)\Intermediate\Build\EpicGames.ScriptBuild.props"
diff --git a/Engine/Source/Programs/Shared/EpicGames.Serialization.Tests/EpicGames.Serialization.Tests.csproj b/Engine/Source/Programs/Shared/EpicGames.Serialization.Tests/EpicGames.Serialization.Tests.csproj
--- a/Engine/Source/Programs/Shared/EpicGames.Serialization.Tests/EpicGames.Serialization.Tests.csproj
+++ b/Engine/Source/Programs/Shared/EpicGames.Serialization.Tests/EpicGames.Serialization.Tests.csproj
@@ -1,8 +1,4 @@
 <Project Sdk="Microsoft.NET.Sdk">
-  <PropertyGroup>
-    <TargetFramework>net8.0</TargetFramework>
-  </PropertyGroup>
-
   <ItemGroup>
     <ProjectReference Include="..\EpicGames.IoHash\EpicGames.IoHash.csproj" />
     <ProjectReference Include="..\EpicGames.Serialization\EpicGames.Serialization.csproj" />
@@ -19,5 +15,4 @@
       <CopyToOutputDirectory>PreserveNewest</CopyToOutputDirectory>
     </None>
   </ItemGroup>
-
 </Project>
diff --git a/Engine/Source/Programs/Shared/EpicGames.Serialization/EpicGames.Serialization.csproj b/Engine/Source/Programs/Shared/EpicGames.Serialization/EpicGames.Serialization.csproj
--- a/Engine/Source/Programs/Shared/EpicGames.Serialization/EpicGames.Serialization.csproj
+++ b/Engine/Source/Programs/Shared/EpicGames.Serialization/EpicGames.Serialization.csproj
@@ -1,8 +1,4 @@
 ﻿<Project Sdk="Microsoft.NET.Sdk">
-  <PropertyGroup>
-    <TargetFramework>net8.0</TargetFramework>
-  </PropertyGroup>
-
   <ItemGroup>
     <PackageReference Include="System.IO.Hashing" Version="8.0.0" />
   </ItemGroup>
diff --git a/Engine/Source/Programs/Shared/EpicGames.Slack/EpicGames.Slack.csproj b/Engine/Source/Programs/Shared/EpicGames.Slack/EpicGames.Slack.csproj
--- a/Engine/Source/Programs/Shared/EpicGames.Slack/EpicGames.Slack.csproj
+++ b/Engine/Source/Programs/Shared/EpicGames.Slack/EpicGames.Slack.csproj
@@ -1,6 +1,5 @@
 ﻿<Project Sdk="Microsoft.NET.Sdk">
   <PropertyGroup>
-    <TargetFramework>net8.0</TargetFramework>
     <DebugType>pdbonly</DebugType>
     <RunAnalyzersDuringBuild>true</RunAnalyzersDuringBuild>
   </PropertyGroup>
diff --git a/Engine/Source/Programs/Shared/EpicGames.Tracing.Tests/EpicGames.Tracing.Tests.csproj b/Engine/Source/Programs/Shared/EpicGames.Tracing.Tests/EpicGames.Tracing.Tests.csproj
--- a/Engine/Source/Programs/Shared/EpicGames.Tracing.Tests/EpicGames.Tracing.Tests.csproj
+++ b/Engine/Source/Programs/Shared/EpicGames.Tracing.Tests/EpicGames.Tracing.Tests.csproj
@@ -1,8 +1,4 @@
 <Project Sdk="Microsoft.NET.Sdk">
-  <PropertyGroup>
-    <TargetFramework>net8.0</TargetFramework>
-  </PropertyGroup>
-
   <ItemGroup>
     <ProjectReference Include="..\EpicGames.Tracing\EpicGames.Tracing.csproj" />
   </ItemGroup>
@@ -12,5 +8,4 @@
       <CopyToOutputDirectory>PreserveNewest</CopyToOutputDirectory>
     </None>
   </ItemGroup>
-
 </Project>
diff --git a/Engine/Source/Programs/Shared/EpicGames.Tracing/EpicGames.Tracing.csproj b/Engine/Source/Programs/Shared/EpicGames.Tracing/EpicGames.Tracing.csproj
--- a/Engine/Source/Programs/Shared/EpicGames.Tracing/EpicGames.Tracing.csproj
+++ b/Engine/Source/Programs/Shared/EpicGames.Tracing/EpicGames.Tracing.csproj
@@ -1,6 +1,5 @@
 ﻿<Project Sdk="Microsoft.NET.Sdk">
   <PropertyGroup>
-    <TargetFramework>net8.0</TargetFramework>
     <NoWarn>$(NoWarn);1591</NoWarn>
   </PropertyGroup>
 </Project>
diff --git a/Engine/Source/Programs/Shared/EpicGames.UBA/EpicGames.UBA.csproj b/Engine/Source/Programs/Shared/EpicGames.UBA/EpicGames.UBA.csproj
--- a/Engine/Source/Programs/Shared/EpicGames.UBA/EpicGames.UBA.csproj
+++ b/Engine/Source/Programs/Shared/EpicGames.UBA/EpicGames.UBA.csproj
@@ -1,6 +1,5 @@
 ﻿<Project Sdk="Microsoft.NET.Sdk">
   <PropertyGroup>
-    <TargetFramework>net8.0</TargetFramework>
     <GenerateDocumentationFile>True</GenerateDocumentationFile>
     <PackageReadmeFile>README.md</PackageReadmeFile>
   </PropertyGroup>
diff --git a/Engine/Source/Programs/Shared/EpicGames.UHT/EpicGames.UHT.csproj b/Engine/Source/Programs/Shared/EpicGames.UHT/EpicGames.UHT.csproj
--- a/Engine/Source/Programs/Shared/EpicGames.UHT/EpicGames.UHT.csproj
+++ b/Engine/Source/Programs/Shared/EpicGames.UHT/EpicGames.UHT.csproj
@@ -1,6 +1,5 @@
 ﻿<Project Sdk="Microsoft.NET.Sdk">
   <PropertyGroup>
-    <TargetFramework>net8.0</TargetFramework>
     <GenerateDocumentationFile>True</GenerateDocumentationFile>
     <EnforceCodeStyleInBuild>False</EnforceCodeStyleInBuild>
   </PropertyGroup>
diff --git a/Engine/Source/Programs/Shared/UnrealBuildTool.Tests/UnrealBuildTool.Tests.csproj b/Engine/Source/Programs/Shared/UnrealBuildTool.Tests/UnrealBuildTool.Tests.csproj
--- a/Engine/Source/Programs/Shared/UnrealBuildTool.Tests/UnrealBuildTool.Tests.csproj
+++ b/Engine/Source/Programs/Shared/UnrealBuildTool.Tests/UnrealBuildTool.Tests.csproj
@@ -1,8 +1,4 @@
 <Project Sdk="Microsoft.NET.Sdk">
-  <PropertyGroup>
-    <TargetFramework>net8.0</TargetFramework>
-  </PropertyGroup>
-
   <ItemGroup>
     <ProjectReference Include="..\EpicGames.Horde\EpicGames.Horde.csproj" />
     <ProjectReference Include="..\..\UnrealBuildTool\UnrealBuildTool.csproj" />
