[Horde][Analytics] Adding computed stream property to use for synthesized value (branch & stream id).

#Synced-CL 50800167

diff --git a/Engine/Source/Programs/Shared/EpicGames.Analytics.PerformanceTrends/PerformanceTrendTelemetry.cs b/Engine/Source/Programs/Shared/EpicGames.Analytics.PerformanceTrends/PerformanceTrendTelemetry.cs
--- a/Engine/Source/Programs/Shared/EpicGames.Analytics.PerformanceTrends/PerformanceTrendTelemetry.cs
+++ b/Engine/Source/Programs/Shared/EpicGames.Analytics.PerformanceTrends/PerformanceTrendTelemetry.cs
@@ -283,6 +283,30 @@ namespace EpicGames.Analytics.PerformanceTrends
 		[Column("summary_name")]
 		public string SummaryName { get; init; }
 
+		#region -- Computed Properties --
+
+		/// <summary>
+		/// The semantic platform the telemetry was produced on.
+		/// </summary>
+		/// <remarks>
+		/// This represents the coalesced result of <see cref="TestConfigName"/> and <see cref="Platform"/>.
+		/// <see cref="TestConfigName"/> represent the modern way of disambiguating between platforms (and within platforms - ex. Windows-Vulkan and !Vulkan), whereas <see cref="Platform"/> represents a historical model.
+		/// </remarks>
+		[Column("computed_platform")]
+		public string? ComputedPlatform { get; init; }
+
+		/// <summary>
+		/// The stream the underlying build was produced from.
+		/// </summary>
+		/// <remarks>
+		///	This represents the coalesced result of <see cref="HordeContextTelemetryRecord.StreamId"/> and <see cref="Branch"/>.
+		///	Common local user invocations that result in telemetry emissions may not have a stream id, but will always have a branch.
+		/// </remarks>
+		[Column("computed_stream")]
+		public string? ComputedStream { get; init; }
+
+		#endregion -- Computed Properties --
+
 		#endregion -- Public Members --
 
 		/// <summary>
diff --git a/Engine/Source/Programs/Shared/EpicGames.Horde/Telemetry/HordeContextTelemetryRecord.cs b/Engine/Source/Programs/Shared/EpicGames.Horde/Telemetry/HordeContextTelemetryRecord.cs
--- a/Engine/Source/Programs/Shared/EpicGames.Horde/Telemetry/HordeContextTelemetryRecord.cs
+++ b/Engine/Source/Programs/Shared/EpicGames.Horde/Telemetry/HordeContextTelemetryRecord.cs
@@ -27,7 +27,7 @@ namespace EpicGames.Horde.Telemetry
 		/// <summary>
 		/// The Horde URL.
 		/// </summary>
-		[Column("horde_url")]
+		[Column("horde_url_str")]
 #pragma warning disable CA1056 // URI-like properties should not be strings
 		public string? HordeUrlStr { get; init; }
 #pragma warning restore CA1056 // URI-like properties should not be strings
