Shared: Add missing project to solution and fix analyze warnings

#rnx
#Synced-CL 50830039

diff --git a/Engine/Source/Programs/Shared/EpicGames.Analytics.PerformanceTrends/PerformanceTrendTelemetry.cs b/Engine/Source/Programs/Shared/EpicGames.Analytics.PerformanceTrends/PerformanceTrendTelemetry.cs
--- a/Engine/Source/Programs/Shared/EpicGames.Analytics.PerformanceTrends/PerformanceTrendTelemetry.cs
+++ b/Engine/Source/Programs/Shared/EpicGames.Analytics.PerformanceTrends/PerformanceTrendTelemetry.cs
@@ -35,7 +35,6 @@ namespace EpicGames.Analytics.PerformanceTrends
 		internal const string VerbatimRhiName = "verbatimrhiname";
 		internal const string DeviceProfile = "deviceprofile";
 
-
 		#endregion -- Internal Constants --
 
 		#region -- Public Members --
@@ -91,8 +90,8 @@ namespace EpicGames.Analytics.PerformanceTrends
 		/// <summary>
 		/// The standard metadata keys.
 		/// </summary>
-		public static readonly string[] StandardMetadataKeys = new string[]
-		{
+		public static readonly string[] StandardMetadataKeys =
+		[
 				EngineVersion,
 				EngineReleaseVersion,
 				BuildVersion,
@@ -112,7 +111,7 @@ namespace EpicGames.Analytics.PerformanceTrends
 				Cpu,
 				VerbatimRhiName,
 				DeviceProfile
-		};
+		];
 
 		/// <summary>
 		/// Creates a new hashtable only with the required metadata entries retained (if they exist).
@@ -121,7 +120,7 @@ namespace EpicGames.Analytics.PerformanceTrends
 		/// <returns>Hashtable with the required metadata entries.</returns>
 		public static Hashtable CreateFilteredHashtable(Hashtable sourceTable)
 		{
-			Hashtable returnTable = new Hashtable();
+			Hashtable returnTable = [];
 
 			for (int i = 0; i < StandardMetadataKeys.Length; i++)
 			{
