EpicGames.AspNet: Fix propeties from shared property sheet that are being overridden
Clean up docstrings

#rnx

#Synced-CL 50830039

diff --git a/Engine/Source/Programs/Shared/EpicGames.AspNet/CbConvertersAspnet.cs b/Engine/Source/Programs/Shared/EpicGames.AspNet/CbConvertersAspnet.cs
--- a/Engine/Source/Programs/Shared/EpicGames.AspNet/CbConvertersAspnet.cs
+++ b/Engine/Source/Programs/Shared/EpicGames.AspNet/CbConvertersAspnet.cs
@@ -9,12 +9,12 @@ using Microsoft.AspNetCore.Mvc;
 namespace EpicGames.AspNet
 {
 	/// <summary>
-	/// 
+	/// Provides methods for registering ASP.NET-specific compact binary converters
 	/// </summary>
 	public static class CbConvertersAspNet
 	{
 		/// <summary>
-		/// 
+		/// Registers ASP.NET-specific converters for ProblemDetails and ValidationProblemDetails types
 		/// </summary>
 		public static void AddAspnetConverters()
 		{
@@ -59,6 +59,11 @@ namespace EpicGames.AspNet
 			writer.WriteField(name, ToCbObject(problemDetails).AsField());
 		}
 
+		/// <summary>
+		/// Converts a ProblemDetails instance to a CbObject representation
+		/// </summary>
+		/// <param name="problemDetails">The ProblemDetails to convert</param>
+		/// <returns>A CbObject containing the serialized problem details</returns>
 		private static CbObject ToCbObject(ProblemDetails problemDetails)
 		{
 			CbWriter objectWriter = new CbWriter();
@@ -91,6 +96,9 @@ namespace EpicGames.AspNet
 		}
 	}
 
+	/// <summary>
+	/// Converter for asp.net validation problem details type
+	/// </summary>
 	class CbValidationProblemDetailsConverter : CbConverter<ValidationProblemDetails>
 	{
 		/// <inheritdoc/>
@@ -139,6 +147,11 @@ namespace EpicGames.AspNet
 			writer.WriteField(name, ToCbObject(validationProblemDetails).AsField());
 		}
 
+		/// <summary>
+		/// Converts a ValidationProblemDetails instance to a CbObject representation
+		/// </summary>
+		/// <param name="validationProblemDetails">The ValidationProblemDetails to convert</param>
+		/// <returns>A CbObject containing the serialized validation problem details</returns>
 		private static CbObject ToCbObject(ValidationProblemDetails validationProblemDetails)
 		{
 			CbWriter objectWriter = new CbWriter();
diff --git a/Engine/Source/Programs/Shared/EpicGames.AspNet/ClientSendSlowException.cs b/Engine/Source/Programs/Shared/EpicGames.AspNet/ClientSendSlowException.cs
--- a/Engine/Source/Programs/Shared/EpicGames.AspNet/ClientSendSlowException.cs
+++ b/Engine/Source/Programs/Shared/EpicGames.AspNet/ClientSendSlowException.cs
@@ -5,8 +5,16 @@ using Microsoft.AspNetCore.Http;
 
 namespace EpicGames.AspNet
 {
+	/// <summary>
+	/// Utility methods for detecting and handling slow client send exceptions
+	/// </summary>
 	public static class ClientSendSlowExceptionUtil
 	{
+		/// <summary>
+		/// Checks if a BadHttpRequestException indicates a slow client send and throws ClientSendSlowException if so
+		/// </summary>
+		/// <param name="e">The exception to check</param>
+		/// <returns>True if the exception was not a slow send exception</returns>
 		public static bool MaybeThrowSlowSendException(BadHttpRequestException e)
 		{
 			// if a user is sending data to slowly it will be terminated by asp.net and a bad http request exception is thrown
@@ -29,8 +37,15 @@ namespace EpicGames.AspNet
 		}
 	}
 
+	/// <summary>
+	/// Exception thrown when a client is sending data too slowly
+	/// </summary>
 	public class ClientSendSlowException : Exception
 	{
+		/// <summary>
+		/// Initializes a new instance of the ClientSendSlowException class
+		/// </summary>
+		/// <param name="innerException">The inner exception that caused this exception</param>
 		public ClientSendSlowException(Exception? innerException) : base("Client was sending data to slowly", innerException)
 		{
 		}
diff --git a/Engine/Source/Programs/Shared/EpicGames.AspNet/CompactBinaryFormatter.cs b/Engine/Source/Programs/Shared/EpicGames.AspNet/CompactBinaryFormatter.cs
--- a/Engine/Source/Programs/Shared/EpicGames.AspNet/CompactBinaryFormatter.cs
+++ b/Engine/Source/Programs/Shared/EpicGames.AspNet/CompactBinaryFormatter.cs
@@ -30,7 +30,7 @@ namespace EpicGames.AspNet
 		public const string UnrealCompressedBuffer = "application/x-ue-comp";
 
 		/// <summary>
-		/// 
+		/// Media type for Jupiter inline payload
 		/// </summary>
 		public const string JupiterInlinedPayload = "application/x-jupiter-inline";
 
@@ -98,6 +98,12 @@ namespace EpicGames.AspNet
 			return await InputFormatterResult.SuccessAsync(CbSerializer.Deserialize(new CbField(data), context.ModelType)!);
 		}
 
+		/// <summary>
+		/// Formats a byte array as a hexadecimal dump with line breaks every 16 bytes
+		/// </summary>
+		/// <param name="data">The byte array to format</param>
+		/// <param name="maxLength">Maximum number of bytes to include in the dump</param>
+		/// <returns>Formatted hexadecimal string with whitespace and line breaks</returns>
 		static string FormatHexDump(byte[] data, int maxLength)
 		{
 			ReadOnlySpan<byte> span = data.AsSpan(0, Math.Min(data.Length, maxLength));
diff --git a/Engine/Source/Programs/Shared/EpicGames.AspNet/EpicGames.AspNet.csproj b/Engine/Source/Programs/Shared/EpicGames.AspNet/EpicGames.AspNet.csproj
--- a/Engine/Source/Programs/Shared/EpicGames.AspNet/EpicGames.AspNet.csproj
+++ b/Engine/Source/Programs/Shared/EpicGames.AspNet/EpicGames.AspNet.csproj
@@ -1,10 +1,24 @@
 ﻿<Project Sdk="Microsoft.NET.Sdk.Web">
+  <!-- Web SDK import order overrides OutputType from Directory.Build.props so it needs to be set back to Library -->
   <PropertyGroup>
-    <TargetFramework>net8.0</TargetFramework>
-    <!-- This needs to be explictly declared as a library again to build. -->
     <OutputType>Library</OutputType>
-    <DebugType>pdbonly</DebugType>
-    <NoWarn>$(NoWarn);1591</NoWarn>
+  </PropertyGroup>
+
+  <!-- Web SDK import order overrides configuration properties from UnrealEngine.csproj.props -->
+  <PropertyGroup Condition="'$(Configuration)|$(Platform)' == 'Release|AnyCPU'">
+    <Optimize>true</Optimize>
+  </PropertyGroup>
+
+  <PropertyGroup Condition="'$(Configuration)|$(Platform)' == 'Development|AnyCPU'">
+    <Optimize>true</Optimize>
+  </PropertyGroup>
+
+  <PropertyGroup Condition="'$(Configuration)|$(Platform)' == 'Analyze|AnyCPU'">
+    <GenerateDocumentationFile>true</GenerateDocumentationFile>
+    <RunAnalyzersDuringBuild>true</RunAnalyzersDuringBuild>
+    <EnforceCodeStyleInBuild>true</EnforceCodeStyleInBuild>
+    <NuGetAudit>true</NuGetAudit>
+    <NuGetAuditLevel>critical</NuGetAuditLevel>
   </PropertyGroup>
 
   <ItemGroup>
@@ -13,7 +27,6 @@
   </ItemGroup>
 
   <ItemGroup>
-    <PackageReference Include="Microsoft.CSharp" Version="4.7.0" />
     <PackageReference Include="System.Text.Json" Version="8.0.6" />
   </ItemGroup>
 </Project>
diff --git a/Engine/Source/Programs/Shared/EpicGames.AspNet/ServerTimings.cs b/Engine/Source/Programs/Shared/EpicGames.AspNet/ServerTimings.cs
--- a/Engine/Source/Programs/Shared/EpicGames.AspNet/ServerTimings.cs
+++ b/Engine/Source/Programs/Shared/EpicGames.AspNet/ServerTimings.cs
@@ -12,13 +12,37 @@ using Microsoft.Extensions.DependencyInjection;
 
 namespace EpicGames.AspNet
 {
+	/// <summary>
+	/// Represents a server timing metric for performance monitoring
+	/// </summary>
 	public struct ServerTimingMetric
 	{
+		/// <summary>
+		/// The name of the metric
+		/// </summary>
 		private readonly string _metricName;
+
+		/// <summary>
+		/// The duration in milliseconds
+		/// </summary>
 		private readonly double? _duration;
+
+		/// <summary>
+		/// Optional description of the metric
+		/// </summary>
 		private readonly string? _description;
+
+		/// <summary>
+		/// Cached string representation of the server timing metric
+		/// </summary>
 		private string? _serverTimingMetric;
 
+		/// <summary>
+		/// Initializes a new instance of the ServerTimingMetric struct
+		/// </summary>
+		/// <param name="metricName">The name of the metric</param>
+		/// <param name="duration">The duration in milliseconds</param>
+		/// <param name="description">Optional description of the metric</param>
 		public ServerTimingMetric(string metricName, double? duration, string? description)
 		{
 			_metricName = metricName;
@@ -28,6 +52,10 @@ namespace EpicGames.AspNet
 			_serverTimingMetric = null;
 		}
 
+		/// <summary>
+		/// Converts the metric to its string representation
+		/// </summary>
+		/// <returns>Formatted server timing metric string</returns>
 		public override string ToString()
 		{
 			if (_serverTimingMetric != null)
@@ -55,21 +83,64 @@ namespace EpicGames.AspNet
 		}
 	}
 
+	/// <summary>
+	/// Interface for managing server timing metrics
+	/// </summary>
 	public interface IServerTiming
 	{
+		/// <summary>
+		/// Adds a server timing metric
+		/// </summary>
+		/// <param name="metricName">The name of the metric</param>
+		/// <param name="duration">The duration in milliseconds</param>
+		/// <param name="description">Optional description</param>
 		public void AddServerTimingMetric(string metricName, double? duration, string? description);
+
+		/// <summary>
+		/// Creates a scoped metric that measures duration automatically
+		/// </summary>
+		/// <param name="metricName">The name of the metric</param>
+		/// <param name="description">Optional description</param>
+		/// <returns>A disposable scoped metric</returns>
 		public ServerTimingMetricScoped CreateServerTimingMetricScope(string metricName, string? description);
 
+		/// <summary>
+		/// Gets the collection of all recorded metrics
+		/// </summary>
 		public IReadOnlyCollection<ServerTimingMetric> Metrics { get; }
 	}
 
+	/// <summary>
+	/// Scoped server timing metric that automatically records duration when disposed
+	/// </summary>
 	public sealed class ServerTimingMetricScoped : IDisposable
 	{
+		/// <summary>
+		/// The timing manager that will receive the metric when disposed
+		/// </summary>
 		private readonly IServerTiming _timingManager;
+
+		/// <summary>
+		/// The name of the metric
+		/// </summary>
 		private readonly string _metricName;
+
+		/// <summary>
+		/// Optional description of the metric
+		/// </summary>
 		private readonly string? _description;
+
+		/// <summary>
+		/// The start time of the metric measurement
+		/// </summary>
 		private readonly DateTime _startTime;
 
+		/// <summary>
+		/// Initializes a new instance of the ServerTimingMetricScoped class
+		/// </summary>
+		/// <param name="timingManager">The timing manager that will receive the metric</param>
+		/// <param name="metricName">The name of the metric</param>
+		/// <param name="description">Optional description of the metric</param>
 		internal ServerTimingMetricScoped(IServerTiming timingManager, string metricName, string? description)
 		{
 			_timingManager = timingManager;
@@ -77,6 +148,10 @@ namespace EpicGames.AspNet
 			_description = description;
 			_startTime = DateTime.Now;
 		}
+
+		/// <summary>
+		/// Disposes the scoped metric and records the duration
+		/// </summary>
 		public void Dispose()
 		{
 			TimeSpan duration = DateTime.Now - _startTime;
@@ -84,32 +159,56 @@ namespace EpicGames.AspNet
 		}
 	}
 
+	/// <summary>
+	/// Default implementation of server timing metrics
+	/// </summary>
 	public class ServerTiming : IServerTiming
 	{
+		/// <summary>
+		/// Thread-safe collection of recorded metrics
+		/// </summary>
 		private readonly ConcurrentBag<ServerTimingMetric> _metrics = [];
 
+		/// <inheritdoc/>
 		public void AddServerTimingMetric(string metricName, double? duration, string? description)
 		{
 			_metrics.Add(new ServerTimingMetric(metricName, duration, description));
 		}
 
+		/// <inheritdoc/>
 		public ServerTimingMetricScoped CreateServerTimingMetricScope(string metricName, string? description)
 		{
 			return new ServerTimingMetricScoped(this, metricName, description);
 		}
 
+		/// <inheritdoc/>
 		public IReadOnlyCollection<ServerTimingMetric> Metrics => _metrics;
 	}
 
+	/// <summary>
+	/// Middleware for adding server timing headers to HTTP responses
+	/// </summary>
 	public class ServerTimingMiddleware
 	{
+		/// <summary>
+		/// The next request delegate in the middleware pipeline
+		/// </summary>
 		private readonly RequestDelegate _next;
 
+		/// <summary>
+		/// Initializes a new instance of the ServerTimingMiddleware class
+		/// </summary>
+		/// <param name="next">The next request delegate in the pipeline</param>
 		public ServerTimingMiddleware(RequestDelegate next)
 		{
 			_next = next ?? throw new ArgumentNullException(nameof(next));
 		}
 
+		/// <summary>
+		/// Invokes the middleware to add server timing headers
+		/// </summary>
+		/// <param name="context">The HTTP context</param>
+		/// <returns>A task representing the asynchronous operation</returns>
 		public async Task InvokeAsync(HttpContext context)
 		{
 			IServerTiming serverTiming = context.RequestServices.GetRequiredService<IServerTiming>();
@@ -123,11 +222,22 @@ namespace EpicGames.AspNet
 			}
 		}
 
+		/// <summary>
+		/// Checks if the request supports trailer headers
+		/// </summary>
+		/// <param name="request">The HTTP request</param>
+		/// <returns>True if trailers are allowed</returns>
 		public static bool AllowsTrailers(HttpRequest request)
 		{
 			return request.Headers.ContainsKey("TE") && request.Headers["TE"].Contains("trailers");
 		}
 
+		/// <summary>
+		/// Handles server timing by adding metrics as HTTP trailer headers
+		/// </summary>
+		/// <param name="context">The HTTP context</param>
+		/// <param name="serverTiming">The server timing instance containing metrics</param>
+		/// <returns>A task representing the asynchronous operation</returns>
 		private async Task HandleServerTimingAsTrailerHeaderAsync(HttpContext context, IServerTiming serverTiming)
 		{
 			context.Response.DeclareTrailer("Server-Timing");
@@ -141,6 +251,12 @@ namespace EpicGames.AspNet
 				serverTimingValue);
 		}
 
+		/// <summary>
+		/// Handles server timing by adding metrics as HTTP response headers
+		/// </summary>
+		/// <param name="context">The HTTP context</param>
+		/// <param name="serverTiming">The server timing instance containing metrics</param>
+		/// <returns>A task representing the asynchronous operation</returns>
 		private Task HandleServerTimingAsResponseHeadersAsync(HttpContext context, IServerTiming serverTiming)
 		{
 			context.Response.OnStarting(() =>
@@ -157,8 +273,17 @@ namespace EpicGames.AspNet
 			return _next(context);
 		}
 	}
+
+	/// <summary>
+	/// Extension methods for adding server timing services
+	/// </summary>
 	public static class ServerTimingServiceCollectionExtensions
 	{
+		/// <summary>
+		/// Adds server timing services to the service collection
+		/// </summary>
+		/// <param name="services">The service collection</param>
+		/// <returns>The service collection for chaining</returns>
 		public static IServiceCollection AddServerTiming(this IServiceCollection services)
 		{
 			services.AddScoped<IServerTiming, ServerTiming>();
