Shared: prefer implicit object creation when type is apparent

#Synced-CL 50888098

diff --git a/Engine/Source/Programs/Shared/EpicGames.Analytics.Generators/AnalyticsSourceGenHelper.cs b/Engine/Source/Programs/Shared/EpicGames.Analytics.Generators/AnalyticsSourceGenHelper.cs
--- a/Engine/Source/Programs/Shared/EpicGames.Analytics.Generators/AnalyticsSourceGenHelper.cs
+++ b/Engine/Source/Programs/Shared/EpicGames.Analytics.Generators/AnalyticsSourceGenHelper.cs
@@ -26,7 +26,7 @@ namespace EpicGames.Analytics.Generators
 		/// <returns>string representing the version.</returns>
 		public static string GenerateAnalyticsTableGenVersionFile()
 		{
-			StringBuilder sb = new StringBuilder();
+			StringBuilder sb = new();
 			sb.Append($$"""
 				namespace {{AnalyticsTableGenNamespace}}
 				{
@@ -54,7 +54,7 @@ namespace EpicGames.Analytics.Generators
 		public static string GenerateAnalyticsTable(AnalyticsTableToGenerate analyticstableToGenerate)
 		{
 			{
-				StringBuilder sb = new StringBuilder();
+				StringBuilder sb = new();
 				sb.Append($$"""
 				namespace {{AnalyticsTableGenNamespace}}
 				{
diff --git a/Engine/Source/Programs/Shared/EpicGames.Analytics.Generators/AnalyticsTableGenerator.cs b/Engine/Source/Programs/Shared/EpicGames.Analytics.Generators/AnalyticsTableGenerator.cs
--- a/Engine/Source/Programs/Shared/EpicGames.Analytics.Generators/AnalyticsTableGenerator.cs
+++ b/Engine/Source/Programs/Shared/EpicGames.Analytics.Generators/AnalyticsTableGenerator.cs
@@ -52,7 +52,7 @@ namespace EpicGames.Analytics.Generators
 	public class AnalyticsTableGenerator : ISourceGenerator
 	{
 #pragma warning disable RS2008 // Enable analyzer release tracking
-		private static readonly DiagnosticDescriptor s_descriptor = new DiagnosticDescriptor(id: "HORDEATG01", title: "Unexpected Error", messageFormat: "Error for object: {0}", category: "Design", defaultSeverity: DiagnosticSeverity.Warning, isEnabledByDefault: true);
+		private static readonly DiagnosticDescriptor s_descriptor = new(id: "HORDEATG01", title: "Unexpected Error", messageFormat: "Error for object: {0}", category: "Design", defaultSeverity: DiagnosticSeverity.Warning, isEnabledByDefault: true);
 #pragma warning restore RS2008 // Enable analyzer release tracking
 
 		/// <inheritdoc/>
diff --git a/Engine/Source/Programs/Shared/EpicGames.Analytics.PerformanceTrends/PerformanceTrendTelemetry.cs b/Engine/Source/Programs/Shared/EpicGames.Analytics.PerformanceTrends/PerformanceTrendTelemetry.cs
--- a/Engine/Source/Programs/Shared/EpicGames.Analytics.PerformanceTrends/PerformanceTrendTelemetry.cs
+++ b/Engine/Source/Programs/Shared/EpicGames.Analytics.PerformanceTrends/PerformanceTrendTelemetry.cs
@@ -51,7 +51,7 @@ namespace EpicGames.Analytics.PerformanceTrends
 				return null;
 			}
 
-			Hashtable sanitized = new Hashtable(source.Count);
+			Hashtable sanitized = new(source.Count);
 
 			foreach (DictionaryEntry entry in source)
 			{
diff --git a/Engine/Source/Programs/Shared/EpicGames.Analytics.Telemetry/AnalyticsTelemetryHelpers.cs b/Engine/Source/Programs/Shared/EpicGames.Analytics.Telemetry/AnalyticsTelemetryHelpers.cs
--- a/Engine/Source/Programs/Shared/EpicGames.Analytics.Telemetry/AnalyticsTelemetryHelpers.cs
+++ b/Engine/Source/Programs/Shared/EpicGames.Analytics.Telemetry/AnalyticsTelemetryHelpers.cs
@@ -18,7 +18,7 @@ namespace EpicGames.Analytics
 		/// <returns>The left hashtable with all the right table's elements merged in.</returns>
 		public static Hashtable MergeLeftCopy(Hashtable left, Hashtable right, bool overrideTarget = true)
 		{
-			Hashtable result = new Hashtable(left);
+			Hashtable result = new(left);
 
 			MergeLeft(result, right, overrideTarget);
 
diff --git a/Engine/Source/Programs/Shared/EpicGames.Analytics.Telemetry/ColumnValidator.cs b/Engine/Source/Programs/Shared/EpicGames.Analytics.Telemetry/ColumnValidator.cs
--- a/Engine/Source/Programs/Shared/EpicGames.Analytics.Telemetry/ColumnValidator.cs
+++ b/Engine/Source/Programs/Shared/EpicGames.Analytics.Telemetry/ColumnValidator.cs
@@ -64,7 +64,7 @@ namespace EpicGames.Analytics
 
 		private static HashSet<string> ReplicateAsHashSet(Hashtable source, bool fuzzyReplication)
 		{
-			HashSet<string> sourceKeys = new HashSet<string>(source.Count, StringComparer.OrdinalIgnoreCase);
+			HashSet<string> sourceKeys = new(source.Count, StringComparer.OrdinalIgnoreCase);
 			foreach (object key in source.Keys)
 			{
 				if (key is string s)
diff --git a/Engine/Source/Programs/Shared/EpicGames.AspNet/CbConvertersAspnet.cs b/Engine/Source/Programs/Shared/EpicGames.AspNet/CbConvertersAspnet.cs
--- a/Engine/Source/Programs/Shared/EpicGames.AspNet/CbConvertersAspnet.cs
+++ b/Engine/Source/Programs/Shared/EpicGames.AspNet/CbConvertersAspnet.cs
@@ -36,7 +36,7 @@ namespace EpicGames.AspNet
 				throw new CbException($"Error converting field \"{field.Name}\" to ProblemDetails. Expected CbObject.");
 			}
 
-			ProblemDetails result = new ProblemDetails
+			ProblemDetails result = new()
 			{
 				Title = field[new Utf8String("title")].AsString(),
 				Detail = field[new Utf8String("detail")].AsString(),
@@ -66,7 +66,7 @@ namespace EpicGames.AspNet
 		/// <returns>A CbObject containing the serialized problem details</returns>
 		private static CbObject ToCbObject(ProblemDetails problemDetails)
 		{
-			CbWriter objectWriter = new CbWriter();
+			CbWriter objectWriter = new();
 			objectWriter.BeginObject();
 
 			objectWriter.WriteString(new Utf8String("title"), problemDetails.Title);
@@ -123,7 +123,7 @@ namespace EpicGames.AspNet
 				errors[error.Name.ToString()] = [.. errorValuesList];
 			}
 
-			ValidationProblemDetails result = new ValidationProblemDetails
+			ValidationProblemDetails result = new()
 			{
 				Title = field[new Utf8String("title")].AsString(),
 				Detail = field[new Utf8String("detail")].AsString(),
@@ -154,7 +154,7 @@ namespace EpicGames.AspNet
 		/// <returns>A CbObject containing the serialized validation problem details</returns>
 		private static CbObject ToCbObject(ValidationProblemDetails validationProblemDetails)
 		{
-			CbWriter objectWriter = new CbWriter();
+			CbWriter objectWriter = new();
 			objectWriter.BeginObject();
 
 			objectWriter.WriteString(new Utf8String("title"), validationProblemDetails.Title);
diff --git a/Engine/Source/Programs/Shared/EpicGames.AspNet/CompactBinaryFormatter.cs b/Engine/Source/Programs/Shared/EpicGames.AspNet/CompactBinaryFormatter.cs
--- a/Engine/Source/Programs/Shared/EpicGames.AspNet/CompactBinaryFormatter.cs
+++ b/Engine/Source/Programs/Shared/EpicGames.AspNet/CompactBinaryFormatter.cs
@@ -66,7 +66,7 @@ namespace EpicGames.AspNet
 			byte[] data;
 			try
 			{
-				using MemoryStream stream = new MemoryStream();
+				using MemoryStream stream = new();
 				await context.HttpContext.Request.Body.CopyToAsync(stream);
 				data = stream.ToArray();
 			}
diff --git a/Engine/Source/Programs/Shared/EpicGames.AspNet/ServerTimings.cs b/Engine/Source/Programs/Shared/EpicGames.AspNet/ServerTimings.cs
--- a/Engine/Source/Programs/Shared/EpicGames.AspNet/ServerTimings.cs
+++ b/Engine/Source/Programs/Shared/EpicGames.AspNet/ServerTimings.cs
@@ -63,7 +63,7 @@ namespace EpicGames.AspNet
 				return _serverTimingMetric;
 			}
 
-			StringBuilder sb = new StringBuilder(_metricName);
+			StringBuilder sb = new(_metricName);
 			if (_duration != null)
 			{
 				sb.Append(";dur=");
diff --git a/Engine/Source/Programs/Shared/EpicGames.Build.Tests/DotnetProjectInfoTests.cs b/Engine/Source/Programs/Shared/EpicGames.Build.Tests/DotnetProjectInfoTests.cs
--- a/Engine/Source/Programs/Shared/EpicGames.Build.Tests/DotnetProjectInfoTests.cs
+++ b/Engine/Source/Programs/Shared/EpicGames.Build.Tests/DotnetProjectInfoTests.cs
@@ -52,7 +52,7 @@ public class DotnetProjectInfoTests
 
 	private static FileReference GetOwnProjectPath()
 	{
-		DirectoryReference? searchDirectory = new DirectoryReference(Assembly.GetExecutingAssembly().Location);
+		DirectoryReference? searchDirectory = new(Assembly.GetExecutingAssembly().Location);
 		FileReference? maybeOwnProjectPath = null;
 		while (searchDirectory is not null)
 		{
diff --git a/Engine/Source/Programs/Shared/EpicGames.Build/Automation/CommandLine.cs b/Engine/Source/Programs/Shared/EpicGames.Build/Automation/CommandLine.cs
--- a/Engine/Source/Programs/Shared/EpicGames.Build/Automation/CommandLine.cs
+++ b/Engine/Source/Programs/Shared/EpicGames.Build/Automation/CommandLine.cs
@@ -48,7 +48,7 @@ namespace UnrealBuildBase
 			private set;
 		} = new HashSet<string>(IgnoredGlobalParameters, StringComparer.InvariantCultureIgnoreCase);
 
-		readonly HashSet<string> GlobalParameterValues = new HashSet<string>(StringComparer.InvariantCultureIgnoreCase);
+		readonly HashSet<string> GlobalParameterValues = new(StringComparer.InvariantCultureIgnoreCase);
 
 		public void SetGlobal(string Parameter)
 		{
@@ -96,7 +96,7 @@ namespace UnrealBuildBase
 			return GlobalParameterValues.Contains(Parameter);
 		}
 
-		Dictionary<string, object?> UncheckedParameters = new Dictionary<string, object?>(StringComparer.InvariantCultureIgnoreCase);
+		Dictionary<string, object?> UncheckedParameters = new(StringComparer.InvariantCultureIgnoreCase);
 		public void SetUnchecked(string Name, object? Value = null)
 		{
 			UncheckedParameters[Name] = Value;
diff --git a/Engine/Source/Programs/Shared/EpicGames.Build/Automation/FormatCommandLine.cs b/Engine/Source/Programs/Shared/EpicGames.Build/Automation/FormatCommandLine.cs
--- a/Engine/Source/Programs/Shared/EpicGames.Build/Automation/FormatCommandLine.cs
+++ b/Engine/Source/Programs/Shared/EpicGames.Build/Automation/FormatCommandLine.cs
@@ -15,7 +15,7 @@ namespace UnrealBuildBase
 		/// <returns>Single string containing all arguments separated with a space.</returns>
 		public static string FormatCommandLine(IEnumerable<string> arguments)
 		{
-			StringBuilder result = new StringBuilder();
+			StringBuilder result = new();
 			foreach (string argument in arguments)
 			{
 				if (result.Length > 0)
diff --git a/Engine/Source/Programs/Shared/EpicGames.Build/Automation/ProcessSingleton.cs b/Engine/Source/Programs/Shared/EpicGames.Build/Automation/ProcessSingleton.cs
--- a/Engine/Source/Programs/Shared/EpicGames.Build/Automation/ProcessSingleton.cs
+++ b/Engine/Source/Programs/Shared/EpicGames.Build/Automation/ProcessSingleton.cs
@@ -27,8 +27,8 @@ namespace UnrealBuildBase
 		{
 			// Need to execute this logic on a background thread, since mutex ownership on Linux has thread affinity (ie. mutexes be released on the
 			// same thread that acquires it, which is not guaranteed through an async continuation)
-			TaskCompletionSource<ExitCode> result = new TaskCompletionSource<ExitCode>();
-			Thread thread = new Thread(() => RunSingleInstanceThread(main, result, bWaitForUATMutex, logger));
+			TaskCompletionSource<ExitCode> result = new();
+			Thread thread = new(() => RunSingleInstanceThread(main, result, bWaitForUATMutex, logger));
 			thread.Start();
 			return await result.Task;
 		}
@@ -43,7 +43,7 @@ namespace UnrealBuildBase
 					?? Unreal.DotnetPath;
 
 				string mutexName = GlobalSingleInstanceMutex.GetUniqueMutexForPath(entryAssemblyLocation.GetFileNameWithoutExtension(), entryAssemblyLocation);
-				using (Mutex singleInstanceMutex = new Mutex(true, mutexName, out bool bCreatedMutex))
+				using (Mutex singleInstanceMutex = new(true, mutexName, out bool bCreatedMutex))
 				{
 					IsSoleInstance = bCreatedMutex;
 
diff --git a/Engine/Source/Programs/Shared/EpicGames.Build/DotnetProjectInfo.cs b/Engine/Source/Programs/Shared/EpicGames.Build/DotnetProjectInfo.cs
--- a/Engine/Source/Programs/Shared/EpicGames.Build/DotnetProjectInfo.cs
+++ b/Engine/Source/Programs/Shared/EpicGames.Build/DotnetProjectInfo.cs
@@ -207,7 +207,7 @@ public sealed record class DotnetProjectInfo : IDotnetProjectInfo
 
 	private class MsBuildRegistration : IMsBuildRegistration
 	{
-		internal static readonly MsBuildRegistration Instance = new MsBuildRegistration();
+		internal static readonly MsBuildRegistration Instance = new();
 
 		public DirectoryReference DotnetDirectory => Unreal.DotnetDirectory;
 
diff --git a/Engine/Source/Programs/Shared/EpicGames.Build/Exceptions/BuildExceptionExtensions.cs b/Engine/Source/Programs/Shared/EpicGames.Build/Exceptions/BuildExceptionExtensions.cs
--- a/Engine/Source/Programs/Shared/EpicGames.Build/Exceptions/BuildExceptionExtensions.cs
+++ b/Engine/Source/Programs/Shared/EpicGames.Build/Exceptions/BuildExceptionExtensions.cs
@@ -25,7 +25,7 @@ namespace UnrealBuildBase
 			}
 			else if (ex is JsonException jsonException)
 			{
-				FileReference source = new FileReference(jsonException.Path ?? jsonException.Source ?? "unknown");
+				FileReference source = new(jsonException.Path ?? jsonException.Source ?? "unknown");
 				LogValue fileValue = LogValue.SourceFile(source, source.GetFileName());
 				long line = jsonException.LineNumber ?? 0;
 				logger.LogError(KnownLogEvents.Compiler, "{File}({Line}): error: {Message}", fileValue, line, ExceptionUtils.FormatExceptionDetails(ex));
diff --git a/Engine/Source/Programs/Shared/EpicGames.Build/Platform/UEBuildPlatformSDK.cs b/Engine/Source/Programs/Shared/EpicGames.Build/Platform/UEBuildPlatformSDK.cs
--- a/Engine/Source/Programs/Shared/EpicGames.Build/Platform/UEBuildPlatformSDK.cs
+++ b/Engine/Source/Programs/Shared/EpicGames.Build/Platform/UEBuildPlatformSDK.cs
@@ -160,7 +160,7 @@ namespace EpicGames.Core
 
 		public SDKCollection GetAllSDKInfo()
 		{
-			SDKCollection allSdks = new SDKCollection(this);
+			SDKCollection allSdks = new(this);
 
 			if (bIsSdkAllowedOnHost)
 			{
@@ -203,7 +203,7 @@ namespace EpicGames.Core
 
 		public SDKCollection GetAllSoftwareInfo(string? deviceType = null, string? current = null)
 		{
-			SDKCollection allSoftwares = new SDKCollection(this);
+			SDKCollection allSoftwares = new(this);
 
 			// walk over the valid software ranges, potentially setting the current verison if one was passed in
 			// if there is only one known SDK named "Software" we allow any DeviceType
@@ -696,7 +696,7 @@ namespace EpicGames.Core
 				{
 					SDKInfoValidity = SDKStatus.Invalid;
 
-					StringBuilder msg = new StringBuilder();
+					StringBuilder msg = new();
 					msg.AppendLine($"Unable to find valid SDK(s) for {PlatformName}:");
 
 					foreach (SDKDescriptor desc in sdkInfo.Sdks)
@@ -1181,7 +1181,7 @@ namespace EpicGames.Core
 				string versionFilename = Path.Combine(platformSDKRoot, CurrentlyInstalledSDKStringManifest);
 				if (File.Exists(versionFilename))
 				{
-					using (StreamReader reader = new StreamReader(versionFilename))
+					using (StreamReader reader = new(versionFilename))
 					{
 						string? version = reader.ReadLine();
 						string? type = reader.ReadLine();
@@ -1223,7 +1223,7 @@ namespace EpicGames.Core
 				string versionFilename = Path.Combine(platformSDKRoot, LastRunScriptVersionManifest);
 				if (File.Exists(versionFilename))
 				{
-					using (StreamReader reader = new StreamReader(versionFilename))
+					using (StreamReader reader = new(versionFilename))
 					{
 						string? version = reader.ReadLine();
 						if (version != null)
@@ -1599,7 +1599,7 @@ namespace EpicGames.Core
 			string envVarFile = Path.Combine(platformSDKRoot, SDKEnvironmentVarsFile);
 			if (File.Exists(envVarFile))
 			{
-				using (StreamReader reader = new StreamReader(envVarFile))
+				using (StreamReader reader = new(envVarFile))
 				{
 					List<string> pathAdds = [];
 					List<string> pathRemoves = [];
diff --git a/Engine/Source/Programs/Shared/EpicGames.Build/System/CompileScriptModules.cs b/Engine/Source/Programs/Shared/EpicGames.Build/System/CompileScriptModules.cs
--- a/Engine/Source/Programs/Shared/EpicGames.Build/System/CompileScriptModules.cs
+++ b/Engine/Source/Programs/Shared/EpicGames.Build/System/CompileScriptModules.cs
@@ -336,7 +336,7 @@ namespace UnrealBuildBase
 					default:
 						throw new Exception("Unsupported rules file type");
 				}
-				Dictionary<FileReference, CsProjBuildRecordEntry> OutRecords = new Dictionary<FileReference, CsProjBuildRecordEntry>(BuildRecords.Count);
+				Dictionary<FileReference, CsProjBuildRecordEntry> OutRecords = new(BuildRecords.Count);
 				foreach (KeyValuePair<FileReference, CsProjBuildRecordEntry> Record in BuildRecords.Where(x => x.Value.BuildRecordFile.HasExtension(FilterExtension)))
 				{
 					FileReference TargetPath = FileReference.Combine(Record.Key.Directory, Record.Value.BuildRecord.TargetPath!);
@@ -441,7 +441,7 @@ namespace UnrealBuildBase
 
 					if (BuildRecord != null && BuildRecord.ProjectPath != null)
 					{
-						CsProjBuildRecordEntry Entry = new CsProjBuildRecordEntry(
+						CsProjBuildRecordEntry Entry = new(
 							new(Path.GetFullPath(BuildRecord.ProjectPath, JsonFile.Directory.FullName)),
 							JsonFile,
 							BuildRecord);
@@ -476,7 +476,7 @@ namespace UnrealBuildBase
 					Files.Add(Glob.ItemType!, TypedFiles);
 				}
 
-				Matcher matcher = new Matcher();
+				Matcher matcher = new();
 				if (Glob.Include?.Count > 0)
 				{
 					matcher.AddIncludePatterns(Glob.Include);
@@ -501,7 +501,7 @@ namespace UnrealBuildBase
 				{
 					// Matcher.Match() doesn't handle inconsistent path separators correctly - which is why globs
 					// are normalized when they are added to CsProjBuildRecord
-					Matcher removeMatcher = new Matcher();
+					Matcher removeMatcher = new();
 					removeMatcher.AddIncludePatterns(Glob.Remove);
 					TypedFiles.RemoveWhere(F => removeMatcher.Match(F).HasMatches);
 				}
@@ -678,7 +678,7 @@ namespace UnrealBuildBase
 			}
 			else
 			{
-				DirectoryReference ScriptsDir = new DirectoryReference(Path.GetDirectoryName(ScriptsForProjectFileName)!);
+				DirectoryReference ScriptsDir = new(Path.GetDirectoryName(ScriptsForProjectFileName)!);
 				ScriptsDir = DirectoryReference.FindCorrectCase(ScriptsDir);
 				GameDirectories.Add(ScriptsDir);
 			}
diff --git a/Engine/Source/Programs/Shared/EpicGames.Build/System/DirectoryEnumerator.cs b/Engine/Source/Programs/Shared/EpicGames.Build/System/DirectoryEnumerator.cs
--- a/Engine/Source/Programs/Shared/EpicGames.Build/System/DirectoryEnumerator.cs
+++ b/Engine/Source/Programs/Shared/EpicGames.Build/System/DirectoryEnumerator.cs
@@ -212,7 +212,7 @@ namespace UnrealBuildBase
 
 					if (!(nameChars == 1 && namePtr[0] == '.') && !(nameChars == 2 && namePtr[0] == '.' && namePtr[1] == '.'))
 					{
-						string name = new string(namePtr, 0, nameChars);
+						string name = new(namePtr, 0, nameChars);
 						entries.Add(new Entry(name, info->EndOfFile, (FileAttributes)info->FileAttributes, DateTimeOffset.FromFileTime(info->CreationTime).UtcDateTime, DateTimeOffset.FromFileTime(info->LastWriteTime).UtcDateTime));
 					}
 
diff --git a/Engine/Source/Programs/Shared/EpicGames.Build/System/FileHasher.cs b/Engine/Source/Programs/Shared/EpicGames.Build/System/FileHasher.cs
--- a/Engine/Source/Programs/Shared/EpicGames.Build/System/FileHasher.cs
+++ b/Engine/Source/Programs/Shared/EpicGames.Build/System/FileHasher.cs
@@ -116,7 +116,7 @@ namespace UnrealBuildBase
 		/// <param name="hash">The updated hash</param>
 		public void SetDigest(FileItem item, IoHash hash)
 		{
-			CachedDigest digest = new CachedDigest(item, hash);
+			CachedDigest digest = new(item, hash);
 			CachedDigests.AddOrUpdate(item.FullName, digest, (key, oldValue) => digest);
 		}
 
@@ -135,7 +135,7 @@ namespace UnrealBuildBase
 		static async Task<CachedDigest> ComputeDigest(FileItem item, ILogger? logger, CancellationToken CancellationToken = default)
 		{
 			using FileStream stream = FileReference.Open(item.Location, FileMode.Open, FileAccess.Read, FileShare.Read);
-			CachedDigest digest = new CachedDigest(item, await IoHash.ComputeAsync(stream, item.Length, CancellationToken));
+			CachedDigest digest = new(item, await IoHash.ComputeAsync(stream, item.Length, CancellationToken));
 			logger?.LogDebug("Computed IoHash {Digest} File {Location} Size {Size} LastWrite {LastWrite}", digest.Digest, item.FullName, item.Length, item.LastWriteTimeUtc);
 			return digest;
 		}
diff --git a/Engine/Source/Programs/Shared/EpicGames.Build/System/NativeProjectsBase.cs b/Engine/Source/Programs/Shared/EpicGames.Build/System/NativeProjectsBase.cs
--- a/Engine/Source/Programs/Shared/EpicGames.Build/System/NativeProjectsBase.cs
+++ b/Engine/Source/Programs/Shared/EpicGames.Build/System/NativeProjectsBase.cs
@@ -17,7 +17,7 @@ namespace UnrealBuildBase
 		/// <summary>
 		/// Object used for synchronizing access to static fields
 		/// </summary>
-		static protected object LockObject = new object();
+		static protected object LockObject = new();
 
 		/// <summary>
 		/// The native project base directories
diff --git a/Engine/Source/Programs/Shared/EpicGames.Build/System/Rules.cs b/Engine/Source/Programs/Shared/EpicGames.Build/System/Rules.cs
--- a/Engine/Source/Programs/Shared/EpicGames.Build/System/Rules.cs
+++ b/Engine/Source/Programs/Shared/EpicGames.Build/System/Rules.cs
@@ -127,7 +127,7 @@ namespace UnrealBuildBase
 		/// <param name="DirectoryPath">Directory to invalidate</param>
 		public static void InvalidateRulesFileCache(string DirectoryPath)
 		{
-			DirectoryReference dirRef = new DirectoryReference(DirectoryPath);
+			DirectoryReference dirRef = new(DirectoryPath);
 			lock (Root)
 			{
 				Directory? dir = FindDirectory(dirRef, false);
diff --git a/Engine/Source/Programs/Shared/EpicGames.Build/System/Timeline.cs b/Engine/Source/Programs/Shared/EpicGames.Build/System/Timeline.cs
--- a/Engine/Source/Programs/Shared/EpicGames.Build/System/Timeline.cs
+++ b/Engine/Source/Programs/Shared/EpicGames.Build/System/Timeline.cs
@@ -121,7 +121,7 @@ namespace UnrealBuildBase
 		/// <summary>
 		/// The stopwatch used for timing
 		/// </summary>
-		static Stopwatch Stopwatch = new Stopwatch();
+		static Stopwatch Stopwatch = new();
 
 		/// <summary>
 		/// The recorded events
@@ -169,7 +169,7 @@ namespace UnrealBuildBase
 		/// <returns>Event to track the length of the event</returns>
 		public static ITimelineEvent ScopeEvent(string Name)
 		{
-			Event Event = new Event(Name, Stopwatch.Elapsed, null);
+			Event Event = new(Name, Stopwatch.Elapsed, null);
 			lock (Events)
 			{
 				Events.Add(Event);
@@ -187,7 +187,7 @@ namespace UnrealBuildBase
 		/// <returns>Event to track the length of the event</returns>
 		public static ITimelineEvent ScopeEvent(string Name, ILogger Logger, LogLevel Verbosity = LogLevel.Information)
 		{
-			Event Event = new Event(Name, Stopwatch.Elapsed, null, Logger, Verbosity);
+			Event Event = new(Name, Stopwatch.Elapsed, null, Logger, Verbosity);
 			lock (Events)
 			{
 				Events.Add(Event);
@@ -309,7 +309,7 @@ namespace UnrealBuildBase
 		/// <param name="Logger">Logger for output</param>
 		static void Print(TimeSpan StartTime, TimeSpan? FinishTime, TimeSpan? MinKnownTime, string Label, List<Event> OuterEvents, LogLevel Verbosity, ILogger Logger)
 		{
-			StringBuilder Prefix = new StringBuilder();
+			StringBuilder Prefix = new();
 
 			for (int Idx = 0; Idx < OuterEvents.Count - 1; Idx++)
 			{
diff --git a/Engine/Source/Programs/Shared/EpicGames.Build/Unreal.cs b/Engine/Source/Programs/Shared/EpicGames.Build/Unreal.cs
--- a/Engine/Source/Programs/Shared/EpicGames.Build/Unreal.cs
+++ b/Engine/Source/Programs/Shared/EpicGames.Build/Unreal.cs
@@ -248,7 +248,7 @@ namespace UnrealBuildBase
 		/// </summary>
 		/// <returns>True if running on a build machine</returns>
 		public static bool IsBuildMachine() => s_isBuildMachine.Value;
-		private static readonly Lazy<bool> s_isBuildMachine = new Lazy<bool>(() => String.Equals("1", Environment.GetEnvironmentVariable("IsBuildMachine")?.Trim(), StringComparison.Ordinal));
+		private static readonly Lazy<bool> s_isBuildMachine = new(() => String.Equals("1", Environment.GetEnvironmentVariable("IsBuildMachine")?.Trim(), StringComparison.Ordinal));
 
 		/// <summary>
 		/// Returns true if the application is running using installed Engine components
@@ -356,16 +356,16 @@ namespace UnrealBuildBase
 
 		// A subset of the functionality in DataDrivenPlatformInfo.GetAllPlatformInfos() - finds the DataDrivenPlatformInfo.ini files and records their existence, but does not parse them
 		// (perhaps DataDrivenPlatformInfo.GetAllPlatformInfos() could be modified to use this data to avoid an additional search through the filesystem)
-		private static readonly Lazy<HashSet<string>> s_iniPresentForPlatform = new Lazy<HashSet<string>>(() =>
+		private static readonly Lazy<HashSet<string>> s_iniPresentForPlatform = new(() =>
 		{
-			HashSet<string> set = new HashSet<string>(StringComparer.OrdinalIgnoreCase);
+			HashSet<string> set = new(StringComparer.OrdinalIgnoreCase);
 			// find all platform directories (skipping NFL/NoRedist)
 			foreach (DirectoryReference engineConfigDir in GetExtensionDirs(EngineDirectory, "Config", bIncludeRestrictedDirectories: false))
 			{
 				// look through all config dirs looking for the data driven ini file
 				foreach (string filePath in Directory.EnumerateFiles(engineConfigDir.FullName, "DataDrivenPlatformInfo.ini", SearchOption.AllDirectories))
 				{
-					FileReference fileRef = new FileReference(filePath);
+					FileReference fileRef = new(filePath);
 
 					// get the platform name from the path
 					string iniPlatformName;
diff --git a/Engine/Source/Programs/Shared/EpicGames.BuildGraph.Tests/GraphTests.cs b/Engine/Source/Programs/Shared/EpicGames.BuildGraph.Tests/GraphTests.cs
--- a/Engine/Source/Programs/Shared/EpicGames.BuildGraph.Tests/GraphTests.cs
+++ b/Engine/Source/Programs/Shared/EpicGames.BuildGraph.Tests/GraphTests.cs
@@ -18,14 +18,14 @@ namespace EpicGames.BuildGraph.Tests
 		static object Evaluate(BgExpr expr)
 		{
 			(byte[] data, BgThunkDef[] methods) = BgCompiler.Compile(expr);
-			BgInterpreter interpreter = new BgInterpreter(data, methods, new Dictionary<string, string>());
+			BgInterpreter interpreter = new(data, methods, new Dictionary<string, string>());
 			return interpreter.Evaluate();
 		}
 
 		[TestMethod]
 		public void NodeTest()
 		{
-			BgAgent agent = new BgAgent("name", "type");
+			BgAgent agent = new("name", "type");
 			BgNode nodeSpec1 = agent.AddNode(x => UpdateVersionFiles());
 
 			BgNode nodeSpec2 = agent.AddNode(x => CompileShooterGameWin64()).Requires(nodeSpec1);
@@ -51,7 +51,7 @@ namespace EpicGames.BuildGraph.Tests
 		[TestMethod]
 		public void AgentTest()
 		{
-			BgAgent expr = new BgAgent("TestAgent", BgList<BgString>.Create("win64", "incremental"));
+			BgAgent expr = new("TestAgent", BgList<BgString>.Create("win64", "incremental"));
 
 			BgAgentDef agent = ((BgObjectDef)Evaluate(expr)).Deserialize<BgAgentDef>();
 			Assert.AreEqual("TestAgent", agent.Name);
@@ -70,13 +70,13 @@ namespace EpicGames.BuildGraph.Tests
 		[TestMethod]
 		public void AggregateTest()
 		{
-			BgAgent agent = new BgAgent("test", "test");
+			BgAgent agent = new("test", "test");
 
 			BgNode nodeSpec1 = agent.AddNode(x => UpdateVersionFiles());
 			BgNode nodeSpec2 = agent.AddNode(x => CompileShooterGameWin64());
 
-			BgLabel labelSpec = new BgLabel("name", "category");
-			BgAggregate aggregateSpec = new BgAggregate("All nodes", BgList.Create(nodeSpec1, nodeSpec2), labelSpec);
+			BgLabel labelSpec = new("name", "category");
+			BgAggregate aggregateSpec = new("All nodes", BgList.Create(nodeSpec1, nodeSpec2), labelSpec);
 
 			BgAggregateExpressionDef aggregateExpr = ((BgObjectDef)Evaluate(aggregateSpec)).Deserialize<BgAggregateExpressionDef>();
 			BgAggregateDef aggregate = aggregateExpr.ToAggregateDef();
@@ -97,7 +97,7 @@ namespace EpicGames.BuildGraph.Tests
 		[TestMethod]
 		public void LabelTest()
 		{
-			BgLabel labelSpec = new BgLabel("name", "category", "ugsBadge", "ugsProject", BgLabelChange.Code);
+			BgLabel labelSpec = new("name", "category", "ugsBadge", "ugsProject", BgLabelChange.Code);
 
 			BgLabelDef label = ((BgObjectDef)Evaluate(labelSpec)).Deserialize<BgLabelDef>();
 			Assert.AreEqual("name", label.DashboardName);
@@ -110,13 +110,13 @@ namespace EpicGames.BuildGraph.Tests
 		[TestMethod]
 		public void GraphTest()
 		{
-			BgAgent agent = new BgAgent("test", "test");
+			BgAgent agent = new("test", "test");
 
 			BgNode nodeSpec1 = agent.AddNode(x => UpdateVersionFiles());
 			BgNode nodeSpec2 = agent.AddNode(x => CompileShooterGameWin64());
 
-			BgAggregate aggregateSpec = new BgAggregate("All nodes", BgList.Create(nodeSpec1, nodeSpec2), new BgLabel("name", "category"));
-			BgGraph graphSpec = new BgGraph(BgList.Create(nodeSpec1, nodeSpec2), BgList.Create(aggregateSpec));
+			BgAggregate aggregateSpec = new("All nodes", BgList.Create(nodeSpec1, nodeSpec2), new BgLabel("name", "category"));
+			BgGraph graphSpec = new(BgList.Create(nodeSpec1, nodeSpec2), BgList.Create(aggregateSpec));
 
 			BgGraphDef graph = ((BgObjectDef)Evaluate(graphSpec)).Deserialize<BgGraphExpressionDef>().ToGraphDef();
 			Assert.AreEqual(2, graph.NameToNode.Count);
diff --git a/Engine/Source/Programs/Shared/EpicGames.BuildGraph.Tests/OpcodeTests.cs b/Engine/Source/Programs/Shared/EpicGames.BuildGraph.Tests/OpcodeTests.cs
--- a/Engine/Source/Programs/Shared/EpicGames.BuildGraph.Tests/OpcodeTests.cs
+++ b/Engine/Source/Programs/Shared/EpicGames.BuildGraph.Tests/OpcodeTests.cs
@@ -16,14 +16,14 @@ namespace EpicGames.BuildGraph.Tests
 		static object Evaluate(BgExpr expr)
 		{
 			(byte[] data, BgThunkDef[] methods) = BgCompiler.Compile(expr);
-			BgInterpreter interpreter = new BgInterpreter(data, methods, new Dictionary<string, string>());
+			BgInterpreter interpreter = new(data, methods, new Dictionary<string, string>());
 			return interpreter.Evaluate();
 		}
 
 		static BgOptionDef EvaluateOption(BgExpr expr)
 		{
 			(byte[] data, BgThunkDef[] methods) = BgCompiler.Compile(expr);
-			BgInterpreter interpreter = new BgInterpreter(data, methods, new Dictionary<string, string>());
+			BgInterpreter interpreter = new(data, methods, new Dictionary<string, string>());
 			interpreter.Evaluate();
 			return interpreter.OptionDefs[0];
 		}
@@ -580,20 +580,20 @@ namespace EpicGames.BuildGraph.Tests
 		[TestMethod]
 		public void Call()
 		{
-			BgFunc<BgBool> func = new BgFunc<BgBool>(() => true);
+			BgFunc<BgBool> func = new(() => true);
 			object result1 = Evaluate(func.Call());
 			Assert.AreEqual(true, result1);
 
-			BgFunc<BgInt, BgInt> func2 = new BgFunc<BgInt, BgInt>(x => (x * 3) + 2);
+			BgFunc<BgInt, BgInt> func2 = new(x => (x * 3) + 2);
 			object result2 = Evaluate(func2.Call(7));
 			Assert.AreEqual(23, result2);
 
-			BgFunc<BgInt, BgInt> func3a = new BgFunc<BgInt, BgInt>(x => x * 3);
-			BgFunc<BgInt, BgInt> func3b = new BgFunc<BgInt, BgInt>(x => func3a.Call(x) + 2);
+			BgFunc<BgInt, BgInt> func3a = new(x => x * 3);
+			BgFunc<BgInt, BgInt> func3b = new(x => func3a.Call(x) + 2);
 			object result3 = Evaluate(func3b.Call(7));
 			Assert.AreEqual(23, result3);
 
-			BgFunc<BgInt, BgInt, BgInt> func4 = new BgFunc<BgInt, BgInt, BgInt>((x, y) => (x + y));
+			BgFunc<BgInt, BgInt, BgInt> func4 = new((x, y) => (x + y));
 			object result4 = Evaluate(func4.Call(7, 9));
 			Assert.AreEqual(16, result4);
 		}
diff --git a/Engine/Source/Programs/Shared/EpicGames.BuildGraph/BgAgentDef.cs b/Engine/Source/Programs/Shared/EpicGames.BuildGraph/BgAgentDef.cs
--- a/Engine/Source/Programs/Shared/EpicGames.BuildGraph/BgAgentDef.cs
+++ b/Engine/Source/Programs/Shared/EpicGames.BuildGraph/BgAgentDef.cs
@@ -48,7 +48,7 @@ namespace EpicGames.BuildGraph
 		/// <inheritdoc/>
 		public override BgAgentDef Deserialize(BgObjectDef<BgAgentDef> obj)
 		{
-			BgAgentDef agent = new BgAgentDef(obj.Get(x => x.Name, ""));
+			BgAgentDef agent = new(obj.Get(x => x.Name, ""));
 			obj.CopyTo(agent);
 			return agent;
 		}
diff --git a/Engine/Source/Programs/Shared/EpicGames.BuildGraph/BgAggregateDef.cs b/Engine/Source/Programs/Shared/EpicGames.BuildGraph/BgAggregateDef.cs
--- a/Engine/Source/Programs/Shared/EpicGames.BuildGraph/BgAggregateDef.cs
+++ b/Engine/Source/Programs/Shared/EpicGames.BuildGraph/BgAggregateDef.cs
@@ -69,7 +69,7 @@ namespace EpicGames.BuildGraph
 		/// <returns></returns>
 		public BgAggregateDef ToAggregateDef()
 		{
-			BgAggregateDef aggregate = new BgAggregateDef(Name);
+			BgAggregateDef aggregate = new(Name);
 			aggregate.RequiredNodes.UnionWith(RequiredNodes);
 			return aggregate;
 		}
@@ -80,7 +80,7 @@ namespace EpicGames.BuildGraph
 		/// <inheritdoc/>
 		public override BgAggregateExpressionDef Deserialize(BgObjectDef<BgAggregateExpressionDef> obj)
 		{
-			BgAggregateExpressionDef aggregate = new BgAggregateExpressionDef(obj.Get(x => x.Name, ""));
+			BgAggregateExpressionDef aggregate = new(obj.Get(x => x.Name, ""));
 			obj.CopyTo(aggregate);
 			return aggregate;
 		}
diff --git a/Engine/Source/Programs/Shared/EpicGames.BuildGraph/BgCompiler.cs b/Engine/Source/Programs/Shared/EpicGames.BuildGraph/BgCompiler.cs
--- a/Engine/Source/Programs/Shared/EpicGames.BuildGraph/BgCompiler.cs
+++ b/Engine/Source/Programs/Shared/EpicGames.BuildGraph/BgCompiler.cs
@@ -25,7 +25,7 @@ namespace EpicGames.BuildGraph
 	{
 		class ReferenceEqualityComparer : IEqualityComparer<object>
 		{
-			public static readonly ReferenceEqualityComparer Instance = new ReferenceEqualityComparer();
+			public static readonly ReferenceEqualityComparer Instance = new();
 
 			public new bool Equals([AllowNull] object x, [AllowNull] object y) => ReferenceEquals(x, y);
 
@@ -35,7 +35,7 @@ namespace EpicGames.BuildGraph
 
 		class FragmentCollector : BgBytecodeWriter
 		{
-			readonly HashSet<BgExpr> _uniqueExprs = new HashSet<BgExpr>(ReferenceEqualityComparer.Instance);
+			readonly HashSet<BgExpr> _uniqueExprs = new(ReferenceEqualityComparer.Instance);
 			readonly List<BgExpr> _fragments;
 			readonly Dictionary<BgExpr, int> _exprToFragmentIndex;
 
@@ -200,20 +200,20 @@ namespace EpicGames.BuildGraph
 		public static (byte[], BgThunkDef[]) Compile(BgExpr expr)
 		{
 			List<BgExpr> fragments = [];
-			Dictionary<BgExpr, int> exprToFragmentIndex = new Dictionary<BgExpr, int>(ReferenceEqualityComparer.Instance);
+			Dictionary<BgExpr, int> exprToFragmentIndex = new(ReferenceEqualityComparer.Instance);
 
 			// Figure out which expressions need to be compiled into separate fragments
-			FragmentCollector collector = new FragmentCollector(fragments, exprToFragmentIndex);
+			FragmentCollector collector = new(fragments, exprToFragmentIndex);
 			collector.WriteExprAsFragment(expr);
 
 			// Write all the fragments
-			ByteArrayBuilder code = new ByteArrayBuilder();
-			List<int> fragmentLengths = new List<int>(fragments.Count);
+			ByteArrayBuilder code = new();
+			List<int> fragmentLengths = new(fragments.Count);
 
 			List<BgThunkDef> thunks = [];
 			List<string> names = [];
 
-			ForwardWriter writer = new ForwardWriter(code, exprToFragmentIndex, names, thunks);
+			ForwardWriter writer = new(code, exprToFragmentIndex, names, thunks);
 			for (int idx = 0; idx < fragments.Count; idx++)
 			{
 				int fragmentOffset = code.Length;
@@ -222,7 +222,7 @@ namespace EpicGames.BuildGraph
 			}
 
 			// Create the header
-			ByteArrayBuilder header = new ByteArrayBuilder();
+			ByteArrayBuilder header = new();
 			header.WriteUnsignedVarInt((int)BgBytecodeVersion.Current);
 			header.WriteVariableLengthArray(names, x => header.WriteString(x));
 			header.WriteVariableLengthArray(fragmentLengths, x => header.WriteUnsignedVarInt(x));
diff --git a/Engine/Source/Programs/Shared/EpicGames.BuildGraph/BgGraphDef.cs b/Engine/Source/Programs/Shared/EpicGames.BuildGraph/BgGraphDef.cs
--- a/Engine/Source/Programs/Shared/EpicGames.BuildGraph/BgGraphDef.cs
+++ b/Engine/Source/Programs/Shared/EpicGames.BuildGraph/BgGraphDef.cs
@@ -278,7 +278,7 @@ namespace EpicGames.BuildGraph
 			}
 
 			// Create a new list of aggregates for everything that's left
-			Dictionary<string, BgAggregateDef> newNameToAggregate = new Dictionary<string, BgAggregateDef>(NameToAggregate.Comparer);
+			Dictionary<string, BgAggregateDef> newNameToAggregate = new(NameToAggregate.Comparer);
 			foreach (BgAggregateDef aggregate in NameToAggregate.Values)
 			{
 				if (aggregate.RequiredNodes.All(x => retainNodes.Contains(x)))
@@ -300,7 +300,7 @@ namespace EpicGames.BuildGraph
 			Badges.RemoveAll(x => x.Nodes.Any(y => !retainNodes.Contains(y)));
 
 			// Rebuild the tag name to output dictionary
-			Dictionary<string, BgNodeOutput> newTagNameToNodeOutput = new Dictionary<string, BgNodeOutput>(TagNameToNodeOutput.Comparer);
+			Dictionary<string, BgNodeOutput> newTagNameToNodeOutput = new(TagNameToNodeOutput.Comparer);
 			foreach (BgNodeOutput output in Agents.SelectMany(x => x.Nodes).SelectMany(x => x.Outputs))
 			{
 				newTagNameToNodeOutput.Add(output.TagName, output);
@@ -333,7 +333,7 @@ namespace EpicGames.BuildGraph
 			}
 
 			// Open the output file
-			using (JsonWriter jsonWriter = new JsonWriter(file.FullName))
+			using (JsonWriter jsonWriter = new(file.FullName))
 			{
 				jsonWriter.WriteObjectStart();
 
@@ -440,7 +440,7 @@ namespace EpicGames.BuildGraph
 		public void ExportForHorde(FileReference file)
 		{
 			DirectoryReference.CreateDirectory(file.Directory);
-			using (JsonWriter jsonWriter = new JsonWriter(file.FullName))
+			using (JsonWriter jsonWriter = new(file.FullName))
 			{
 				jsonWriter.WriteObjectStart();
 
@@ -662,7 +662,7 @@ namespace EpicGames.BuildGraph
 				{
 					string name = String.Format("-set:{0}=...", option.Name);
 
-					StringBuilder description = new StringBuilder(option.Description);
+					StringBuilder description = new(option.Description);
 
 					string? defaultValue = null;
 					if (option is BgBoolOptionDef boolOption)
@@ -790,7 +790,7 @@ namespace EpicGames.BuildGraph
 		{
 			List<BgNodeExpressionDef> nodes = [.. Nodes];
 
-			BgGraphDef graph = new BgGraphDef();
+			BgGraphDef graph = new();
 			foreach (BgAggregateExpressionDef aggregate in Aggregates)
 			{
 				graph.NameToAggregate[aggregate.Name] = aggregate.ToAggregateDef();
diff --git a/Engine/Source/Programs/Shared/EpicGames.BuildGraph/BgInterpreter.cs b/Engine/Source/Programs/Shared/EpicGames.BuildGraph/BgInterpreter.cs
--- a/Engine/Source/Programs/Shared/EpicGames.BuildGraph/BgInterpreter.cs
+++ b/Engine/Source/Programs/Shared/EpicGames.BuildGraph/BgInterpreter.cs
@@ -209,7 +209,7 @@ namespace EpicGames.BuildGraph
 			_thunks = thunks;
 			_options = options;
 
-			MemoryReader reader = new MemoryReader(data);
+			MemoryReader reader = new(data);
 
 			_version = (BgBytecodeVersion)reader.ReadUnsignedVarInt();
 			_names = reader.ReadVariableLengthArray(() => reader.ReadString());
diff --git a/Engine/Source/Programs/Shared/EpicGames.BuildGraph/BgNodeDef.cs b/Engine/Source/Programs/Shared/EpicGames.BuildGraph/BgNodeDef.cs
--- a/Engine/Source/Programs/Shared/EpicGames.BuildGraph/BgNodeDef.cs
+++ b/Engine/Source/Programs/Shared/EpicGames.BuildGraph/BgNodeDef.cs
@@ -304,7 +304,7 @@ namespace EpicGames.BuildGraph
 		/// <inheritdoc/>
 		public override BgNodeExpressionDef Deserialize(BgObjectDef<BgNodeExpressionDef> obj)
 		{
-			BgNodeExpressionDef node = new BgNodeExpressionDef(obj.Get(x => x.Agent, null!), obj.Get(x => x.Name, ""), obj.Get(x => x.Thunk!, null!), obj.Get(x => x.OutputCount, 0));
+			BgNodeExpressionDef node = new(obj.Get(x => x.Agent, null!), obj.Get(x => x.Name, ""), obj.Get(x => x.Thunk!, null!), obj.Get(x => x.OutputCount, 0));
 			obj.CopyTo(node);
 
 			node.Inputs.AddRange(node.InputExprs.SelectMany(x => x.Flatten()));
diff --git a/Engine/Source/Programs/Shared/EpicGames.BuildGraph/BgObjectDef.cs b/Engine/Source/Programs/Shared/EpicGames.BuildGraph/BgObjectDef.cs
--- a/Engine/Source/Programs/Shared/EpicGames.BuildGraph/BgObjectDef.cs
+++ b/Engine/Source/Programs/Shared/EpicGames.BuildGraph/BgObjectDef.cs
@@ -73,12 +73,12 @@ namespace EpicGames.BuildGraph
 		/// </summary>
 		/// <typeparam name="T">Type of the object</typeparam>
 		/// <returns></returns>
-		public BgObjectDef<T> WithType<T>() => new BgObjectDef<T>(Properties);
+		public BgObjectDef<T> WithType<T>() => new(Properties);
 
 		/// <summary>
 		/// Cache of serializer instances
 		/// </summary>
-		static readonly ConcurrentDictionary<Type, BgObjectSerializer> s_typeToSerializer = new ConcurrentDictionary<Type, BgObjectSerializer>();
+		static readonly ConcurrentDictionary<Type, BgObjectSerializer> s_typeToSerializer = new();
 
 		/// <summary>
 		/// Deserialize an object of the given type, using the default serializer
@@ -237,7 +237,7 @@ namespace EpicGames.BuildGraph
 		/// <returns></returns>
 		static IReadOnlyDictionary<string, PropertySetter> GetPropertyMap()
 		{
-			Dictionary<string, PropertySetter> nameToProperty = new Dictionary<string, PropertySetter>(StringComparer.Ordinal);
+			Dictionary<string, PropertySetter> nameToProperty = new(StringComparer.Ordinal);
 
 			PropertyInfo[] properties = typeof(T).GetProperties(BindingFlags.Public | BindingFlags.Instance);
 			foreach (PropertyInfo property in properties)
@@ -352,7 +352,7 @@ namespace EpicGames.BuildGraph
 		/// <inheritdoc/>
 		public override T Deserialize(BgObjectDef<T> obj)
 		{
-			T result = new T();
+			T result = new();
 			obj.CopyTo(result);
 			return result;
 		}
diff --git a/Engine/Source/Programs/Shared/EpicGames.BuildGraph/BgOptionDef.cs b/Engine/Source/Programs/Shared/EpicGames.BuildGraph/BgOptionDef.cs
--- a/Engine/Source/Programs/Shared/EpicGames.BuildGraph/BgOptionDef.cs
+++ b/Engine/Source/Programs/Shared/EpicGames.BuildGraph/BgOptionDef.cs
@@ -82,7 +82,7 @@ namespace EpicGames.BuildGraph
 		/// <inheritdoc/>
 		public override BgBoolOptionDef Deserialize(BgObjectDef<BgBoolOptionDef> obj)
 		{
-			BgBoolOptionDef option = new BgBoolOptionDef(obj.Get(x => x.Name, null!));
+			BgBoolOptionDef option = new(obj.Get(x => x.Name, null!));
 			obj.CopyTo(option);
 			return option;
 		}
@@ -130,7 +130,7 @@ namespace EpicGames.BuildGraph
 		/// <inheritdoc/>
 		public override BgIntOptionDef Deserialize(BgObjectDef<BgIntOptionDef> obj)
 		{
-			BgIntOptionDef option = new BgIntOptionDef(obj.Get(x => x.Name, null!));
+			BgIntOptionDef option = new(obj.Get(x => x.Name, null!));
 			obj.CopyTo(option);
 			return option;
 		}
@@ -209,7 +209,7 @@ namespace EpicGames.BuildGraph
 		/// <inheritdoc/>
 		public override BgStringOptionDef Deserialize(BgObjectDef<BgStringOptionDef> obj)
 		{
-			BgStringOptionDef option = new BgStringOptionDef(obj.Get(x => x.Name, null!));
+			BgStringOptionDef option = new(obj.Get(x => x.Name, null!));
 			obj.CopyTo(option);
 			return option;
 		}
@@ -278,7 +278,7 @@ namespace EpicGames.BuildGraph
 		/// <inheritdoc/>
 		public override BgListOptionDef Deserialize(BgObjectDef<BgListOptionDef> obj)
 		{
-			BgListOptionDef option = new BgListOptionDef(obj.Get(x => x.Name, null!));
+			BgListOptionDef option = new(obj.Get(x => x.Name, null!));
 			obj.CopyTo(option);
 			return option;
 		}
diff --git a/Engine/Source/Programs/Shared/EpicGames.BuildGraph/Expressions/BgFunc.cs b/Engine/Source/Programs/Shared/EpicGames.BuildGraph/Expressions/BgFunc.cs
--- a/Engine/Source/Programs/Shared/EpicGames.BuildGraph/Expressions/BgFunc.cs
+++ b/Engine/Source/Programs/Shared/EpicGames.BuildGraph/Expressions/BgFunc.cs
@@ -45,7 +45,7 @@ namespace EpicGames.BuildGraph.Expressions
 		/// <summary>
 		/// Implicit conversion from a regular C# function
 		/// </summary>
-		public static implicit operator BgFunc<TResult>(Func<TResult> func) => new BgFunc<TResult>(func);
+		public static implicit operator BgFunc<TResult>(Func<TResult> func) => new(func);
 
 		/// <summary>
 		/// Call the function with the given arguments
@@ -75,7 +75,7 @@ namespace EpicGames.BuildGraph.Expressions
 		/// <summary>
 		/// Implicit conversion from a regular C# function
 		/// </summary>
-		public static implicit operator BgFunc<TArg, TResult>(Func<TArg, TResult> func) => new BgFunc<TArg, TResult>(func);
+		public static implicit operator BgFunc<TArg, TResult>(Func<TArg, TResult> func) => new(func);
 
 		/// <summary>
 		/// Call the function with the given arguments
@@ -108,7 +108,7 @@ namespace EpicGames.BuildGraph.Expressions
 		/// <summary>
 		/// Implicit conversion from a regular C# function
 		/// </summary>
-		public static implicit operator BgFunc<TArg1, TArg2, TResult>(Func<TArg1, TArg2, TResult> func) => new BgFunc<TArg1, TArg2, TResult>(func);
+		public static implicit operator BgFunc<TArg1, TArg2, TResult>(Func<TArg1, TArg2, TResult> func) => new(func);
 
 		/// <summary>
 		/// Call the function with the given arguments
diff --git a/Engine/Source/Programs/Shared/EpicGames.BuildGraph/Expressions/BgNode.cs b/Engine/Source/Programs/Shared/EpicGames.BuildGraph/Expressions/BgNode.cs
--- a/Engine/Source/Programs/Shared/EpicGames.BuildGraph/Expressions/BgNode.cs
+++ b/Engine/Source/Programs/Shared/EpicGames.BuildGraph/Expressions/BgNode.cs
@@ -174,7 +174,7 @@ namespace EpicGames.BuildGraph.Expressions
 		/// Clone this node
 		/// </summary>
 		/// <returns>Clone of this node</returns>
-		protected virtual BgNode Clone() => new BgNode(this);
+		protected virtual BgNode Clone() => new(this);
 
 		/// <summary>
 		/// Gets the default tag name for the numbered output index
@@ -227,7 +227,7 @@ namespace EpicGames.BuildGraph.Expressions
 							throw new BgNodeException($"Unterminated parameter expression for {nameof(BgNodeNameAttribute)} in {template}");
 						}
 
-						StringView paramName = new StringView(template, nextIdx, endIdx - nextIdx);
+						StringView paramName = new(template, nextIdx, endIdx - nextIdx);
 
 						int paramIdx = Array.FindIndex(parameters, x => x.Name != null && paramName.Equals(x.Name, StringComparison.Ordinal));
 						if (paramIdx == -1)
@@ -274,7 +274,7 @@ namespace EpicGames.BuildGraph.Expressions
 		/// </summary>
 		public static string GetNodeNameFromMethodName(string methodName)
 		{
-			StringBuilder name = new StringBuilder();
+			StringBuilder name = new();
 			name.Append(methodName[0]);
 
 			int length = methodName.Length;
diff --git a/Engine/Source/Programs/Shared/EpicGames.BuildGraph/Expressions/BgType.cs b/Engine/Source/Programs/Shared/EpicGames.BuildGraph/Expressions/BgType.cs
--- a/Engine/Source/Programs/Shared/EpicGames.BuildGraph/Expressions/BgType.cs
+++ b/Engine/Source/Programs/Shared/EpicGames.BuildGraph/Expressions/BgType.cs
@@ -69,7 +69,7 @@ namespace EpicGames.BuildGraph.Expressions
 			public static BgType<T> Instance { get; } = (BgType<T>)CreateInstance(typeof(T));
 		}
 
-		static readonly ConcurrentDictionary<Type, IBgType> s_typeCache = new ConcurrentDictionary<Type, IBgType>();
+		static readonly ConcurrentDictionary<Type, IBgType> s_typeCache = new();
 
 		/// <summary>
 		/// Create a converter instance for the given type
diff --git a/Engine/Source/Programs/Shared/EpicGames.Core.Tests/BuzHashTests.cs b/Engine/Source/Programs/Shared/EpicGames.Core.Tests/BuzHashTests.cs
--- a/Engine/Source/Programs/Shared/EpicGames.Core.Tests/BuzHashTests.cs
+++ b/Engine/Source/Programs/Shared/EpicGames.Core.Tests/BuzHashTests.cs
@@ -52,7 +52,7 @@ namespace EpicGames.Core.Tests
 		[TestMethod]
 		public void Candidates()
 		{
-			Random random = new Random(0);
+			Random random = new(0);
 
 			byte[] bytes = new byte[16 * 1024 * 1024];
 			random.NextBytes(bytes);
diff --git a/Engine/Source/Programs/Shared/EpicGames.Core.Tests/JsonObjectTests.cs b/Engine/Source/Programs/Shared/EpicGames.Core.Tests/JsonObjectTests.cs
--- a/Engine/Source/Programs/Shared/EpicGames.Core.Tests/JsonObjectTests.cs
+++ b/Engine/Source/Programs/Shared/EpicGames.Core.Tests/JsonObjectTests.cs
@@ -419,7 +419,7 @@ namespace EpicGames.Core.Tests
 		{
 			string fieldName = "StrongestAttack";
 			Assert.IsTrue(_validTestObject.ContainsField(fieldName));
-			EnemyAttack setValue = new EnemyAttack("Ultima", 300, StatusEffect.None);
+			EnemyAttack setValue = new("Ultima", 300, StatusEffect.None);
 			JsonObject setValueObject = setValue.ToJsonObject();
 			_validTestObject.AddOrSetFieldValue(fieldName, setValue.ToJsonObject());
 			Assert.AreEqual(setValueObject, _validTestObject.GetObjectField(fieldName));
@@ -528,16 +528,16 @@ namespace EpicGames.Core.Tests
 		[TestMethod]
 		public void Constructor()
 		{
-			JsonObject obj = new JsonObject();
+			JsonObject obj = new();
 			Assert.AreEqual(0, obj.KeyNames.Count());
 		}
 
 		[TestMethod]
 		public void ToJsonString()
 		{
-			string stringToMatch = new string(_testJsonText.ToCharArray());
+			string stringToMatch = new(_testJsonText.ToCharArray());
 			string[] lines = stringToMatch.Split([Environment.NewLine], StringSplitOptions.None);
-			StringBuilder jsonStringBuilder = new StringBuilder();
+			StringBuilder jsonStringBuilder = new();
 			foreach (string line in lines)
 			{
 				// Utf8JsonWriter uses 2 spaces for indents, we replace them with tabs here 
@@ -558,12 +558,12 @@ namespace EpicGames.Core.Tests
 			string escapeCharacterFieldName = "Escape";
 			// .NET Json serialization escapes certian characters like +,<,>,~,& for security reasons. We need to make sure those come through 
 			string escapeString = "++MyTest+Test<>`&";
-			StringBuilder builder = new StringBuilder();
+			StringBuilder builder = new();
 			builder.AppendLine("{");
 			builder.AppendLine($"\t\"{escapeCharacterFieldName}\": \"{escapeString}\"");
 			builder.AppendLine("}");
 			string correctString = builder.ToString();
-			JsonObject jsonObject = new JsonObject();
+			JsonObject jsonObject = new();
 			jsonObject.AddOrSetFieldValue(escapeCharacterFieldName, escapeString);
 			string jsonString = jsonObject.ToJsonString();
 			Assert.AreEqual(correctString, jsonString);
@@ -587,7 +587,7 @@ namespace EpicGames.Core.Tests
 
 		private static Enemy CreateValidJsonDataSource()
 		{
-			Enemy enemy1 = new Enemy();
+			Enemy enemy1 = new();
 			enemy1.Name = "Monster";
 			enemy1.Hp = 100;
 			enemy1.CriticalHitRate = 3.0;
@@ -619,7 +619,7 @@ namespace EpicGames.Core.Tests
 		private static FileReference CreateTempJsonFile(string fileName, string? fileContent)
 		{
 			string inputFile = Path.Join(s_tempDirectory.FullName, fileName);
-			FileReference inputFileReference = new FileReference(inputFile);
+			FileReference inputFileReference = new(inputFile);
 			File.WriteAllText(inputFile, fileContent);
 			return inputFileReference;
 		}
@@ -663,7 +663,7 @@ namespace EpicGames.Core.Tests
 	{
 		public JsonObject ToJsonObject()
 		{
-			JsonObject obj = new JsonObject();
+			JsonObject obj = new();
 			obj.AddOrSetFieldValue("Name", Name);
 			obj.AddOrSetFieldValue("Damage", Damage);
 			obj.AddOrSetFieldValue<StatusEffect>("Effect", Effect);
diff --git a/Engine/Source/Programs/Shared/EpicGames.Core.Tests/JsonTracerTests.cs b/Engine/Source/Programs/Shared/EpicGames.Core.Tests/JsonTracerTests.cs
--- a/Engine/Source/Programs/Shared/EpicGames.Core.Tests/JsonTracerTests.cs
+++ b/Engine/Source/Programs/Shared/EpicGames.Core.Tests/JsonTracerTests.cs
@@ -18,7 +18,7 @@ namespace EpicGames.Core.Tests
 		[TestMethod]
 		public void Run()
 		{
-			JsonTracer tracer = new JsonTracer();
+			JsonTracer tracer = new();
 			using (IScope scope1 = tracer.BuildSpan("bogusOp1").StartActive())
 			{
 			}
@@ -32,7 +32,7 @@ namespace EpicGames.Core.Tests
 		[TestMethod]
 		public void SpansWithTags()
 		{
-			JsonTracer tracer = new JsonTracer();
+			JsonTracer tracer = new();
 			using (IScope scope1 = tracer.BuildSpan("bogusOp1").WithTag("age", 94).StartActive())
 			{
 			}
@@ -49,7 +49,7 @@ namespace EpicGames.Core.Tests
 		[TestMethod]
 		public void SerializeToJson()
 		{
-			JsonTracer tracer = new JsonTracer();
+			JsonTracer tracer = new();
 			using (IScope scope1 = tracer.BuildSpan("bogusOp1").WithTag("age", 94).WithTag("foo", "bar").StartActive())
 			{
 			}
@@ -57,8 +57,8 @@ namespace EpicGames.Core.Tests
 			{
 			}
 
-			using StringWriter sw = new StringWriter();
-			using JsonWriter jw = new JsonWriter(sw);
+			using StringWriter sw = new();
+			using JsonWriter jw = new(sw);
 			tracer.GetFinishedSpansAsJson(jw);
 			byte[] data = Encoding.ASCII.GetBytes(sw.ToString());
 
@@ -79,7 +79,7 @@ namespace EpicGames.Core.Tests
 		[TestMethod]
 		public void SpansWithParallelFor()
 		{
-			JsonTracer tracer = new JsonTracer();
+			JsonTracer tracer = new();
 			using (IScope parentScope = tracer.BuildSpan("Parent").StartActive())
 			{
 				Parallel.For(0, 100, index =>
@@ -122,7 +122,7 @@ namespace EpicGames.Core.Tests
 		[TestMethod]
 		public void SpansWithAsync()
 		{
-			JsonTracer tracer = new JsonTracer();
+			JsonTracer tracer = new();
 			using (IScope parentScope = tracer.BuildSpan("Parent").StartActive())
 			{
 				Assert.AreEqual(tracer.ActiveSpan, parentScope.Span);
@@ -151,7 +151,7 @@ namespace EpicGames.Core.Tests
 		[TestMethod]
 		public void SpansWithParallelForEachAsync()
 		{
-			JsonTracer tracer = new JsonTracer();
+			JsonTracer tracer = new();
 			using (IScope parentScope = tracer.BuildSpan("Parent").StartActive())
 			{
 				Assert.AreEqual(tracer.ActiveSpan, parentScope.Span);
diff --git a/Engine/Source/Programs/Shared/EpicGames.Core.Tests/LogEventTests.cs b/Engine/Source/Programs/Shared/EpicGames.Core.Tests/LogEventTests.cs
--- a/Engine/Source/Programs/Shared/EpicGames.Core.Tests/LogEventTests.cs
+++ b/Engine/Source/Programs/Shared/EpicGames.Core.Tests/LogEventTests.cs
@@ -22,7 +22,7 @@ namespace EpicGames.Core.Tests
 		[TestMethod]
 		public void ArgumentTests()
 		{
-			CaptureLogger logger = new CaptureLogger();
+			CaptureLogger logger = new();
 			logger.LogInformation("Test {Value} {@Value}", new TestClass { Foo = 123, Bar = "456" }, new TestClass { Foo = 789, Bar = "hello" });
 
 			LogEvent logEvent = logger.Events[0];
@@ -34,8 +34,8 @@ namespace EpicGames.Core.Tests
 		[TestMethod]
 		public void RoundTripTests()
 		{
-			DateTime time = new DateTime(2024, 1, 1);
-			LogEvent ev = new LogEvent(time, LogLevel.Information, default, "hello", null, [KeyValuePair.Create<string, object?>("foo", new LogValue(new Utf8String("type"), "text"))], null);
+			DateTime time = new(2024, 1, 1);
+			LogEvent ev = new(time, LogLevel.Information, default, "hello", null, [KeyValuePair.Create<string, object?>("foo", new LogValue(new Utf8String("type"), "text"))], null);
 
 			JsonLogEvent jev = JsonLogEvent.Parse(ev.ToJsonBytes());
 			string jevs = jev.ToString();
diff --git a/Engine/Source/Programs/Shared/EpicGames.Core.Tests/LogValueTests.cs b/Engine/Source/Programs/Shared/EpicGames.Core.Tests/LogValueTests.cs
--- a/Engine/Source/Programs/Shared/EpicGames.Core.Tests/LogValueTests.cs
+++ b/Engine/Source/Programs/Shared/EpicGames.Core.Tests/LogValueTests.cs
@@ -19,8 +19,8 @@ namespace EpicGames.Core.Tests
 
 		static string RenderLogValue(object value)
 		{
-			ArrayBufferWriter<byte> buffer = new ArrayBufferWriter<byte>();
-			using (Utf8JsonWriter writer = new Utf8JsonWriter(buffer))
+			ArrayBufferWriter<byte> buffer = new();
+			using (Utf8JsonWriter writer = new(buffer))
 			{
 				LogValueFormatter.Format(value, writer);
 			}
@@ -44,7 +44,7 @@ namespace EpicGames.Core.Tests
 		[TestMethod]
 		public void SourceFile()
 		{
-			FileReference file = new FileReference("foo.txt");
+			FileReference file = new("foo.txt");
 			string value = RenderLogValue(LogValue.SourceFile(file, "foo"));
 			Assert.AreEqual(value, "{\"$type\":\"SourceFile\",\"$text\":\"foo\",\"file\":\"" + file.FullName.Replace("\\", "\\\\", StringComparison.Ordinal) + "\"}");
 		}
diff --git a/Engine/Source/Programs/Shared/EpicGames.Core.Tests/MemoryAllocatorTests.cs b/Engine/Source/Programs/Shared/EpicGames.Core.Tests/MemoryAllocatorTests.cs
--- a/Engine/Source/Programs/Shared/EpicGames.Core.Tests/MemoryAllocatorTests.cs
+++ b/Engine/Source/Programs/Shared/EpicGames.Core.Tests/MemoryAllocatorTests.cs
@@ -11,7 +11,7 @@ namespace EpicGames.Core.Tests
 		[TestMethod]
 		public void TestGlobalAllocator()
 		{
-			GlobalHeapAllocator allocator = new GlobalHeapAllocator();
+			GlobalHeapAllocator allocator = new();
 
 			using IMemoryOwner<byte> handle = allocator.Alloc(1024, null);
 			handle.Memory.Span[0] = 123;
@@ -20,7 +20,7 @@ namespace EpicGames.Core.Tests
 		[TestMethod]
 		public void TestVirtualAllocator()
 		{
-			VirtualMemoryAllocator allocator = new VirtualMemoryAllocator();
+			VirtualMemoryAllocator allocator = new();
 
 			using IMemoryOwner<byte> handle = allocator.Alloc(1024, null);
 			handle.Memory.Span[0] = 123;
diff --git a/Engine/Source/Programs/Shared/EpicGames.Core.Tests/MemorySerializerTests.cs b/Engine/Source/Programs/Shared/EpicGames.Core.Tests/MemorySerializerTests.cs
--- a/Engine/Source/Programs/Shared/EpicGames.Core.Tests/MemorySerializerTests.cs
+++ b/Engine/Source/Programs/Shared/EpicGames.Core.Tests/MemorySerializerTests.cs
@@ -14,7 +14,7 @@ namespace EpicGames.Core.Tests
 		{
 			byte[] data;
 			{
-				ByteArrayBuilder writer = new ByteArrayBuilder();
+				ByteArrayBuilder writer = new();
 				writer.WriteInt32(12345);
 				writer.WriteInt32(-12345);
 				writer.WriteInt8(SByte.MinValue);
@@ -28,7 +28,7 @@ namespace EpicGames.Core.Tests
 				data = writer.ToByteArray();
 			}
 			{
-				MemoryReader reader = new MemoryReader(data);
+				MemoryReader reader = new(data);
 				Assert.AreEqual(12345, reader.ReadInt32());
 				Assert.AreEqual(-12345, reader.ReadInt32());
 				Assert.AreEqual(SByte.MinValue, reader.ReadInt8());
@@ -48,7 +48,7 @@ namespace EpicGames.Core.Tests
 		{
 			byte[] data;
 			{
-				ByteArrayBuilder writer = new ByteArrayBuilder();
+				ByteArrayBuilder writer = new();
 				writer.WriteUInt32(12345);
 				writer.WriteUInt8(Byte.MaxValue);
 				writer.WriteUInt16(UInt16.MaxValue);
@@ -57,7 +57,7 @@ namespace EpicGames.Core.Tests
 				data = writer.ToByteArray();
 			}
 			{
-				MemoryReader reader = new MemoryReader(data);
+				MemoryReader reader = new(data);
 				Assert.AreEqual(12345U, reader.ReadUInt32());
 				Assert.AreEqual(Byte.MaxValue, reader.ReadUInt8());
 				Assert.AreEqual(UInt16.MaxValue, reader.ReadUInt16());
@@ -72,13 +72,13 @@ namespace EpicGames.Core.Tests
 		{
 			byte[] data;
 			{
-				ByteArrayBuilder writer = new ByteArrayBuilder();
+				ByteArrayBuilder writer = new();
 				writer.WriteVariableLengthBytes([1, 2, 3]);
 				writer.WriteVariableLengthBytes([4, 5, 6]);
 				data = writer.ToByteArray();
 			}
 			{
-				MemoryReader reader = new MemoryReader(data);
+				MemoryReader reader = new(data);
 				Assert.IsTrue(reader.ReadVariableLengthBytes().Span.SequenceEqual(new byte[] { 1, 2, 3 }));
 				Assert.IsTrue(reader.ReadVariableLengthBytes().Span.SequenceEqual(new byte[] { 4, 5, 6 }));
 				reader.CheckEmpty();
@@ -90,14 +90,14 @@ namespace EpicGames.Core.Tests
 		{
 			byte[] data;
 			{
-				ByteArrayBuilder writer = new ByteArrayBuilder();
+				ByteArrayBuilder writer = new();
 				writer.WriteString("hello", Encoding.UTF8);
 				writer.WriteString("there", Encoding.Unicode);
 				writer.WriteString("world", Encoding.UTF32);
 				data = writer.ToByteArray();
 			}
 			{
-				MemoryReader reader = new MemoryReader(data);
+				MemoryReader reader = new(data);
 				Assert.AreEqual("hello", reader.ReadString(Encoding.UTF8));
 				Assert.AreEqual("there", reader.ReadString(Encoding.Unicode));
 				Assert.AreEqual("world", reader.ReadString(Encoding.UTF32));
diff --git a/Engine/Source/Programs/Shared/EpicGames.Core.Tests/MemoryWriterTests.cs b/Engine/Source/Programs/Shared/EpicGames.Core.Tests/MemoryWriterTests.cs
--- a/Engine/Source/Programs/Shared/EpicGames.Core.Tests/MemoryWriterTests.cs
+++ b/Engine/Source/Programs/Shared/EpicGames.Core.Tests/MemoryWriterTests.cs
@@ -47,11 +47,11 @@ namespace EpicGames.Core.Tests
 		[TestMethod]
 		public void RefCountedMemoryWriter()
 		{
-			TestAllocator allocator = new TestAllocator();
+			TestAllocator allocator = new();
 			Assert.AreEqual(0, allocator.AllocatedSize);
 			Assert.AreEqual(0, allocator.AllocatedCount);
 
-			using (RefCountedMemoryWriter writer = new RefCountedMemoryWriter(allocator, 4096))
+			using (RefCountedMemoryWriter writer = new(allocator, 4096))
 			{
 				writer.WriteString("Hello world");
 				Assert.IsTrue(allocator.AllocatedSize > 0);
@@ -65,14 +65,14 @@ namespace EpicGames.Core.Tests
 		[TestMethod]
 		public void RefCountedMemoryWriterSequence()
 		{
-			TestAllocator allocator = new TestAllocator();
+			TestAllocator allocator = new();
 			Assert.AreEqual(0, allocator.AllocatedSize);
 			Assert.AreEqual(0, allocator.AllocatedCount);
 
 			IRefCountedHandle<ReadOnlyMemory<byte>> memory;
 			IRefCountedHandle<ReadOnlySequence<byte>> sequence;
 
-			using (RefCountedMemoryWriter writer = new RefCountedMemoryWriter(allocator, 10))
+			using (RefCountedMemoryWriter writer = new(allocator, 10))
 			{
 				writer.WriteString("This should be written to a single buffer and the original buffer should be freed");
 				int allocatedSize1 = allocator.AllocatedSize;
@@ -106,14 +106,14 @@ namespace EpicGames.Core.Tests
 		[TestMethod]
 		public void RefCountedMemoryWriterSequence2()
 		{
-			TestAllocator allocator = new TestAllocator();
+			TestAllocator allocator = new();
 			Assert.AreEqual(0, allocator.AllocatedSize);
 			Assert.AreEqual(0, allocator.AllocatedCount);
 
 			int allocatedSize1;
 			IRefCountedHandle<ReadOnlySequence<byte>> sequence;
 
-			using (RefCountedMemoryWriter writer = new RefCountedMemoryWriter(allocator, 10))
+			using (RefCountedMemoryWriter writer = new(allocator, 10))
 			{
 				writer.WriteString("This should be written to a single buffer and the original buffer should be freed");
 				allocatedSize1 = allocator.AllocatedSize;
diff --git a/Engine/Source/Programs/Shared/EpicGames.Core/AssemblyUtils.cs b/Engine/Source/Programs/Shared/EpicGames.Core/AssemblyUtils.cs
--- a/Engine/Source/Programs/Shared/EpicGames.Core/AssemblyUtils.cs
+++ b/Engine/Source/Programs/Shared/EpicGames.Core/AssemblyUtils.cs
@@ -153,9 +153,9 @@ namespace EpicGames.Core
 		}
 
 		// Map of assembly name to path on disk
-		private static readonly Dictionary<string, string> s_assemblyLocationCache = new Dictionary<string, string>(StringComparer.OrdinalIgnoreCase);
+		private static readonly Dictionary<string, string> s_assemblyLocationCache = new(StringComparer.OrdinalIgnoreCase);
 		// Track last modified date of each assembly, so we can ensure we always reference the latest one in the case of stale assemblies on disk.
-		private static readonly Dictionary<string, DateTime> s_assemblyWriteTimes = new Dictionary<string, DateTime>(StringComparer.OrdinalIgnoreCase);
+		private static readonly Dictionary<string, DateTime> s_assemblyWriteTimes = new(StringComparer.OrdinalIgnoreCase);
 		// Flag used to make sure we don't redundantly add resolvers
 		private static bool s_addedToAssemblyResolver = false;
 
diff --git a/Engine/Source/Programs/Shared/EpicGames.Core/AsyncEvent.cs b/Engine/Source/Programs/Shared/EpicGames.Core/AsyncEvent.cs
--- a/Engine/Source/Programs/Shared/EpicGames.Core/AsyncEvent.cs
+++ b/Engine/Source/Programs/Shared/EpicGames.Core/AsyncEvent.cs
@@ -13,7 +13,7 @@ namespace EpicGames.Core
 		/// <summary>
 		/// Completion source to wait on
 		/// </summary>
-		TaskCompletionSource<bool> _source = new TaskCompletionSource<bool>(TaskCreationOptions.RunContinuationsAsynchronously);
+		TaskCompletionSource<bool> _source = new(TaskCreationOptions.RunContinuationsAsynchronously);
 
 		/// <summary>
 		/// Legacy name for <see cref="Pulse"/>.
diff --git a/Engine/Source/Programs/Shared/EpicGames.Core/AsyncQueue.cs b/Engine/Source/Programs/Shared/EpicGames.Core/AsyncQueue.cs
--- a/Engine/Source/Programs/Shared/EpicGames.Core/AsyncQueue.cs
+++ b/Engine/Source/Programs/Shared/EpicGames.Core/AsyncQueue.cs
@@ -15,7 +15,7 @@ namespace EpicGames.Core
 	{
 		readonly List<Task> _workerTasks = [];
 		readonly Channel<Func<CancellationToken, Task>> _queuedTasks = Channel.CreateUnbounded<Func<CancellationToken, Task>>();
-		readonly CancellationTokenSource _cancellationSource = new CancellationTokenSource();
+		readonly CancellationTokenSource _cancellationSource = new();
 
 		/// <summary>
 		/// Constructor
diff --git a/Engine/Source/Programs/Shared/EpicGames.Core/AsyncUtils.cs b/Engine/Source/Programs/Shared/EpicGames.Core/AsyncUtils.cs
--- a/Engine/Source/Programs/Shared/EpicGames.Core/AsyncUtils.cs
+++ b/Engine/Source/Programs/Shared/EpicGames.Core/AsyncUtils.cs
@@ -230,7 +230,7 @@ namespace EpicGames.Core
 		/// <param name="cancellationToken">Cancellation token for the operation</param>
 		public static async Task WaitOneAsync(this WaitHandle handle, int timeoutMs, CancellationToken cancellationToken = default)
 		{
-			TaskCompletionSource<bool> completionSource = new TaskCompletionSource<bool>();
+			TaskCompletionSource<bool> completionSource = new();
 
 			RegisteredWaitHandle waitHandle = ThreadPool.RegisterWaitForSingleObject(handle, (state, timedOut) => ((TaskCompletionSource<bool>)state!).TrySetResult(!timedOut), completionSource, timeoutMs, true);
 			try
diff --git a/Engine/Source/Programs/Shared/EpicGames.Core/BackgroundTask.cs b/Engine/Source/Programs/Shared/EpicGames.Core/BackgroundTask.cs
--- a/Engine/Source/Programs/Shared/EpicGames.Core/BackgroundTask.cs
+++ b/Engine/Source/Programs/Shared/EpicGames.Core/BackgroundTask.cs
@@ -13,7 +13,7 @@ namespace EpicGames.Core
 	{
 		readonly Func<CancellationToken, Task> _runTask;
 		Task? _task;
-		readonly CancellationTokenSource _cancellationTokenSource = new CancellationTokenSource();
+		readonly CancellationTokenSource _cancellationTokenSource = new();
 
 		/// <summary>
 		/// Accessor for the underlying task
@@ -43,7 +43,7 @@ namespace EpicGames.Core
 		/// <returns></returns>
 		public static BackgroundTask StartNew(Func<CancellationToken, Task> runTask)
 		{
-			BackgroundTask task = new BackgroundTask(runTask);
+			BackgroundTask task = new(runTask);
 			task.Start();
 			return task;
 		}
@@ -55,7 +55,7 @@ namespace EpicGames.Core
 		/// <returns>New background task instance</returns>
 		public static BackgroundTask<T> StartNew<T>(Func<CancellationToken, Task<T>> runTask)
 		{
-			BackgroundTask<T> task = new BackgroundTask<T>(runTask);
+			BackgroundTask<T> task = new(runTask);
 			task.Start();
 			return task;
 		}
@@ -112,7 +112,7 @@ namespace EpicGames.Core
 	{
 		readonly Func<CancellationToken, Task<T>> _runTask;
 		Task<T>? _task;
-		readonly CancellationTokenSource _cancellationTokenSource = new CancellationTokenSource();
+		readonly CancellationTokenSource _cancellationTokenSource = new();
 
 		/// <summary>
 		/// Accessor for the underlying task
diff --git a/Engine/Source/Programs/Shared/EpicGames.Core/BinaryArchive.cs b/Engine/Source/Programs/Shared/EpicGames.Core/BinaryArchive.cs
--- a/Engine/Source/Programs/Shared/EpicGames.Core/BinaryArchive.cs
+++ b/Engine/Source/Programs/Shared/EpicGames.Core/BinaryArchive.cs
@@ -16,7 +16,7 @@ namespace EpicGames.Core
 	{
 		class Instantiator<T> where T : class, new()
 		{
-			public static T Instance = new T();
+			public static T Instance = new();
 		}
 
 		sealed class TypeSerializer
@@ -38,7 +38,7 @@ namespace EpicGames.Core
 
 			static Func<BinaryArchiveReader, object> CreateBoxedReadMethod(Type type, MethodInfo methodInfo)
 			{
-				DynamicMethod readerMethod = new DynamicMethod($"Boxed_{methodInfo.Name}", typeof(object), [typeof(BinaryArchiveReader)]);
+				DynamicMethod readerMethod = new($"Boxed_{methodInfo.Name}", typeof(object), [typeof(BinaryArchiveReader)]);
 
 				ILGenerator generator = readerMethod.GetILGenerator(64);
 				generator.Emit(OpCodes.Ldarg_0);
@@ -54,7 +54,7 @@ namespace EpicGames.Core
 
 			static Action<BinaryArchiveWriter, object> CreateBoxedWriteMethod(Type type, MethodInfo methodInfo)
 			{
-				DynamicMethod writerMethod = new DynamicMethod($"Boxed_{methodInfo.Name}", typeof(void), [typeof(BinaryArchiveWriter), typeof(object)]);
+				DynamicMethod writerMethod = new($"Boxed_{methodInfo.Name}", typeof(void), [typeof(BinaryArchiveWriter), typeof(object)]);
 
 				ILGenerator generator = writerMethod.GetILGenerator(64);
 				generator.Emit(OpCodes.Ldarg_0);
@@ -74,7 +74,7 @@ namespace EpicGames.Core
 			}
 		}
 
-		static readonly ConcurrentDictionary<Type, TypeSerializer> s_typeToSerializerInfo = new ConcurrentDictionary<Type, TypeSerializer>()
+		static readonly ConcurrentDictionary<Type, TypeSerializer> s_typeToSerializerInfo = new()
 		{
 			[typeof(byte)] = TypeSerializer.Create(reader => reader.ReadByte(), (writer, value) => writer.WriteByte(value)),
 			[typeof(sbyte)] = TypeSerializer.Create(reader => reader.ReadSignedByte(), (writer, value) => writer.WriteSignedByte(value)),
@@ -87,8 +87,8 @@ namespace EpicGames.Core
 			[typeof(string)] = TypeSerializer.Create(reader => reader.ReadString(), (writer, value) => writer.WriteString(value))
 		};
 
-		static readonly ConcurrentDictionary<Type, ContentHash> s_typeToDigest = new ConcurrentDictionary<Type, ContentHash>();
-		static readonly ConcurrentDictionary<ContentHash, Type> s_digestToType = new ConcurrentDictionary<ContentHash, Type>();
+		static readonly ConcurrentDictionary<Type, ContentHash> s_typeToDigest = new();
+		static readonly ConcurrentDictionary<ContentHash, Type> s_digestToType = new();
 
 		static MethodInfo GetMethodInfo(Expression<Action> expr)
 		{
@@ -262,12 +262,12 @@ namespace EpicGames.Core
 				InterfaceMapping interfaceMapping = converterType.GetInterfaceMap(interfaceType);
 
 				// Create a reader method
-				DynamicMethod readerMethod = new DynamicMethod($"BinaryArchiveReader_Dynamic_{type.Name}_With_{converterType.Name}", type, [typeof(BinaryArchiveReader)]);
+				DynamicMethod readerMethod = new($"BinaryArchiveReader_Dynamic_{type.Name}_With_{converterType.Name}", type, [typeof(BinaryArchiveReader)]);
 				int readerIdx = Array.FindIndex(interfaceMapping.InterfaceMethods, x => x.Name.Equals(nameof(IBinaryConverter<object>.Read), StringComparison.Ordinal));
 				GenerateConverterReaderForwardingMethod(readerMethod, converterField, interfaceMapping.TargetMethods[readerIdx]);
 
 				// Create a writer method
-				DynamicMethod writerMethod = new DynamicMethod($"BinaryArchiveWriter_Dynamic_{type.Name}_With_{converterType.Name}", typeof(void), [typeof(BinaryArchiveWriter), typeof(object)]);
+				DynamicMethod writerMethod = new($"BinaryArchiveWriter_Dynamic_{type.Name}_With_{converterType.Name}", typeof(void), [typeof(BinaryArchiveWriter), typeof(object)]);
 				int writerIdx = Array.FindIndex(interfaceMapping.InterfaceMethods, x => x.Name.Equals(nameof(IBinaryConverter<object>.Write), StringComparison.Ordinal));
 				GenerateConverterWriterForwardingMethod(writerMethod, converterField, interfaceMapping.TargetMethods[writerIdx]);
 
@@ -311,10 +311,10 @@ namespace EpicGames.Core
 					?? throw new NotImplementedException($"Type '{type.Name}' does not have a parameterless constructor");
 
 				// Create the methods
-				DynamicMethod writerMethod = new DynamicMethod($"BinaryArchiveWriter_Dynamic_{type.Name}", typeof(void), [typeof(BinaryArchiveWriter), type]);
+				DynamicMethod writerMethod = new($"BinaryArchiveWriter_Dynamic_{type.Name}", typeof(void), [typeof(BinaryArchiveWriter), type]);
 				GenerateClassWriterMethod(writerMethod, properties);
 
-				DynamicMethod readerMethod = new DynamicMethod($"BinaryArchiveReader_Dynamic_{type.Name}", type, [typeof(BinaryArchiveReader)]);
+				DynamicMethod readerMethod = new($"BinaryArchiveReader_Dynamic_{type.Name}", type, [typeof(BinaryArchiveReader)]);
 				GenerateClassReaderMethod(readerMethod, constructorInfo, properties);
 
 				return new TypeSerializer(readerMethod, writerMethod);
diff --git a/Engine/Source/Programs/Shared/EpicGames.Core/BinaryArchiveReader.cs b/Engine/Source/Programs/Shared/EpicGames.Core/BinaryArchiveReader.cs
--- a/Engine/Source/Programs/Shared/EpicGames.Core/BinaryArchiveReader.cs
+++ b/Engine/Source/Programs/Shared/EpicGames.Core/BinaryArchiveReader.cs
@@ -341,7 +341,7 @@ namespace EpicGames.Core
 			}
 			else
 			{
-				List<T> result = new List<T>(count);
+				List<T> result = new(count);
 				for (int idx = 0; idx < count; idx++)
 				{
 					result.Add(readElement());
@@ -390,7 +390,7 @@ namespace EpicGames.Core
 			}
 			else
 			{
-				HashSet<T> result = new HashSet<T>(comparer);
+				HashSet<T> result = new(comparer);
 				for (int idx = 0; idx < count; idx++)
 				{
 					result.Add(readElement());
@@ -439,7 +439,7 @@ namespace EpicGames.Core
 			}
 			else
 			{
-				SortedSet<T> result = new SortedSet<T>(comparer);
+				SortedSet<T> result = new(comparer);
 				for (int idx = 0; idx < count; idx++)
 				{
 					result.Add(readElement());
@@ -465,7 +465,7 @@ namespace EpicGames.Core
 			}
 			else
 			{
-				Dictionary<TK, TV> result = new Dictionary<TK, TV>(count);
+				Dictionary<TK, TV> result = new(count);
 				for (int idx = 0; idx < count; idx++)
 				{
 					result.Add(readKey(), readValue());
@@ -492,7 +492,7 @@ namespace EpicGames.Core
 			}
 			else
 			{
-				Dictionary<TK, TV> result = new Dictionary<TK, TV>(count, comparer);
+				Dictionary<TK, TV> result = new(count, comparer);
 				for (int idx = 0; idx < count; idx++)
 				{
 					result.Add(readKey(), readValue());
diff --git a/Engine/Source/Programs/Shared/EpicGames.Core/BinaryArchiveWriter.cs b/Engine/Source/Programs/Shared/EpicGames.Core/BinaryArchiveWriter.cs
--- a/Engine/Source/Programs/Shared/EpicGames.Core/BinaryArchiveWriter.cs
+++ b/Engine/Source/Programs/Shared/EpicGames.Core/BinaryArchiveWriter.cs
@@ -32,7 +32,7 @@ namespace EpicGames.Core
 		/// <summary>
 		/// Instance of the ReferenceComparer class which can be shared by all archive writers
 		/// </summary>
-		static readonly ReferenceComparer s_referenceComparerInstance = new ReferenceComparer();
+		static readonly ReferenceComparer s_referenceComparerInstance = new();
 
 		/// <summary>
 		/// The output stream being written to
@@ -52,7 +52,7 @@ namespace EpicGames.Core
 		/// <summary>
 		/// Map of object instance to unique id
 		/// </summary>
-		readonly Dictionary<object, int> _objectToUniqueId = new Dictionary<object, int>(s_referenceComparerInstance);
+		readonly Dictionary<object, int> _objectToUniqueId = new(s_referenceComparerInstance);
 
 		/// <summary>
 		/// Constructor
diff --git a/Engine/Source/Programs/Shared/EpicGames.Core/BinaryConverter.cs b/Engine/Source/Programs/Shared/EpicGames.Core/BinaryConverter.cs
--- a/Engine/Source/Programs/Shared/EpicGames.Core/BinaryConverter.cs
+++ b/Engine/Source/Programs/Shared/EpicGames.Core/BinaryConverter.cs
@@ -46,7 +46,7 @@ namespace EpicGames.Core
 		/// <summary>
 		/// Map from type to the converter type
 		/// </summary>
-		static readonly ConcurrentDictionary<Type, Type> s_typeToConverterType = new ConcurrentDictionary<Type, Type>();
+		static readonly ConcurrentDictionary<Type, Type> s_typeToConverterType = new();
 
 		/// <summary>
 		/// Explicitly register the converter for a type. If Type is a generic type, the converter should also be a generic type with the same type arguments.
diff --git a/Engine/Source/Programs/Shared/EpicGames.Core/BinaryFormatterUtils.cs b/Engine/Source/Programs/Shared/EpicGames.Core/BinaryFormatterUtils.cs
--- a/Engine/Source/Programs/Shared/EpicGames.Core/BinaryFormatterUtils.cs
+++ b/Engine/Source/Programs/Shared/EpicGames.Core/BinaryFormatterUtils.cs
@@ -17,8 +17,8 @@ namespace EpicGames.Core
 		/// <returns>Instance of the object that was read from disk</returns>
 		public static T Load<T>(FileReference location)
 		{
-			using FileStream stream = new FileStream(location.FullName, FileMode.Open, FileAccess.Read);
-			DataContractSerializer serializer = new DataContractSerializer(typeof(T));
+			using FileStream stream = new(location.FullName, FileMode.Open, FileAccess.Read);
+			DataContractSerializer serializer = new(typeof(T));
 			return (T?)serializer.ReadObject(stream) ?? throw new SerializationException($"Unable to deserialize {location} as {typeof(T)}");
 		}
 
@@ -30,8 +30,8 @@ namespace EpicGames.Core
 		public static void Save(FileReference location, object obj)
 		{
 			DirectoryReference.CreateDirectory(location.Directory);
-			using FileStream stream = new FileStream(location.FullName, FileMode.Create, FileAccess.Write);
-			DataContractSerializer serializer = new DataContractSerializer(obj.GetType());
+			using FileStream stream = new(location.FullName, FileMode.Create, FileAccess.Write);
+			DataContractSerializer serializer = new(obj.GetType());
 			serializer.WriteObject(stream, obj);
 		}
 
@@ -43,9 +43,9 @@ namespace EpicGames.Core
 		public static void SaveIfDifferent(FileReference location, object obj)
 		{
 			byte[] contents;
-			using (MemoryStream stream = new MemoryStream())
+			using (MemoryStream stream = new())
 			{
-				DataContractSerializer serializer = new DataContractSerializer(obj.GetType());
+				DataContractSerializer serializer = new(obj.GetType());
 				serializer.WriteObject(stream, obj);
 				contents = stream.ToArray();
 			}
diff --git a/Engine/Source/Programs/Shared/EpicGames.Core/BinaryReaderExtensions.cs b/Engine/Source/Programs/Shared/EpicGames.Core/BinaryReaderExtensions.cs
--- a/Engine/Source/Programs/Shared/EpicGames.Core/BinaryReaderExtensions.cs
+++ b/Engine/Source/Programs/Shared/EpicGames.Core/BinaryReaderExtensions.cs
@@ -102,7 +102,7 @@ namespace EpicGames.Core
 				return null;
 			}
 
-			List<T> items = new List<T>(numItems);
+			List<T> items = new(numItems);
 			for (int idx = 0; idx < numItems; idx++)
 			{
 				items.Add(readElement());
@@ -127,7 +127,7 @@ namespace EpicGames.Core
 				return null;
 			}
 
-			Dictionary<TKey, TValue> items = new Dictionary<TKey, TValue>(numItems);
+			Dictionary<TKey, TValue> items = new(numItems);
 			for (int idx = 0; idx < numItems; idx++)
 			{
 				TKey key = readKey();
diff --git a/Engine/Source/Programs/Shared/EpicGames.Core/ByteArrayBuilder.cs b/Engine/Source/Programs/Shared/EpicGames.Core/ByteArrayBuilder.cs
--- a/Engine/Source/Programs/Shared/EpicGames.Core/ByteArrayBuilder.cs
+++ b/Engine/Source/Programs/Shared/EpicGames.Core/ByteArrayBuilder.cs
@@ -129,7 +129,7 @@ namespace EpicGames.Core
 		/// <returns>Sequence of bytes</returns>
 		public ReadOnlySequence<byte> AsSequence()
 		{
-			ReadOnlySequenceBuilder<byte> builder = new ReadOnlySequenceBuilder<byte>();
+			ReadOnlySequenceBuilder<byte> builder = new();
 			AppendTo(builder);
 			return builder.Construct();
 		}
diff --git a/Engine/Source/Programs/Shared/EpicGames.Core/ByteString.cs b/Engine/Source/Programs/Shared/EpicGames.Core/ByteString.cs
--- a/Engine/Source/Programs/Shared/EpicGames.Core/ByteString.cs
+++ b/Engine/Source/Programs/Shared/EpicGames.Core/ByteString.cs
@@ -39,7 +39,7 @@ namespace EpicGames.Core
 		/// </summary>
 		/// <param name="text">Text to parse</param>
 		/// <returns>New byte string instance</returns>
-		public static ByteString Parse(ReadOnlySpan<byte> text) => new ByteString(StringUtils.ParseHexString(text));
+		public static ByteString Parse(ReadOnlySpan<byte> text) => new(StringUtils.ParseHexString(text));
 
 		/// <summary>
 		/// Parses a byte string from a hexadecimal string
@@ -53,7 +53,7 @@ namespace EpicGames.Core
 		/// </summary>
 		/// <param name="text">Text to parse</param>
 		/// <returns>New byte string instance</returns>
-		public static ByteString Parse(string text) => new ByteString(StringUtils.ParseHexString(text));
+		public static ByteString Parse(string text) => new(StringUtils.ParseHexString(text));
 
 		/// <summary>
 		/// Attempts to parse a hexadecimal string as a byte string
@@ -117,7 +117,7 @@ namespace EpicGames.Core
 		/// <inheritdoc/>
 		public override int GetHashCode()
 		{
-			HashCode hashCode = new HashCode();
+			HashCode hashCode = new();
 			hashCode.AddBytes(Data.Span);
 			return hashCode.ToHashCode();
 		}
diff --git a/Engine/Source/Programs/Shared/EpicGames.Core/CertificateUtils.cs b/Engine/Source/Programs/Shared/EpicGames.Core/CertificateUtils.cs
--- a/Engine/Source/Programs/Shared/EpicGames.Core/CertificateUtils.cs
+++ b/Engine/Source/Programs/Shared/EpicGames.Core/CertificateUtils.cs
@@ -26,14 +26,14 @@ namespace EpicGames.Core
 
 			using (RSA algorithm = RSA.Create(2048))
 			{
-				X500DistinguishedName distinguishedName = new X500DistinguishedName($"CN={dnsName}");
-				CertificateRequest request = new CertificateRequest(distinguishedName, algorithm, HashAlgorithmName.SHA256, RSASignaturePadding.Pkcs1);
+				X500DistinguishedName distinguishedName = new($"CN={dnsName}");
+				CertificateRequest request = new(distinguishedName, algorithm, HashAlgorithmName.SHA256, RSASignaturePadding.Pkcs1);
 
 				request.CertificateExtensions.Add(new X509BasicConstraintsExtension(false, false, 0, true));
 				request.CertificateExtensions.Add(new X509KeyUsageExtension(X509KeyUsageFlags.KeyEncipherment | X509KeyUsageFlags.DigitalSignature, true));
 				request.CertificateExtensions.Add(new X509EnhancedKeyUsageExtension([new Oid("1.3.6.1.5.5.7.3.1")], true));
 
-				SubjectAlternativeNameBuilder alternativeNameBuilder = new SubjectAlternativeNameBuilder();
+				SubjectAlternativeNameBuilder alternativeNameBuilder = new();
 				alternativeNameBuilder.AddDnsName(dnsName);
 				request.CertificateExtensions.Add(alternativeNameBuilder.Build(true));
 
diff --git a/Engine/Source/Programs/Shared/EpicGames.Core/ChannelWithRefCount.cs b/Engine/Source/Programs/Shared/EpicGames.Core/ChannelWithRefCount.cs
--- a/Engine/Source/Programs/Shared/EpicGames.Core/ChannelWithRefCount.cs
+++ b/Engine/Source/Programs/Shared/EpicGames.Core/ChannelWithRefCount.cs
@@ -58,6 +58,6 @@ namespace EpicGames.Core
 		/// Creates a channel whose completion state is controlled by a reference count
 		/// </summary>
 		public static ChannelWithRefCount<T> WithRefCount<T>(this Channel<T> channel)
-			=> new ChannelWithRefCount<T>(channel);
+			=> new(channel);
 	}
 }
diff --git a/Engine/Source/Programs/Shared/EpicGames.Core/Clock.cs b/Engine/Source/Programs/Shared/EpicGames.Core/Clock.cs
--- a/Engine/Source/Programs/Shared/EpicGames.Core/Clock.cs
+++ b/Engine/Source/Programs/Shared/EpicGames.Core/Clock.cs
@@ -67,7 +67,7 @@ public interface IClock
 public sealed class NullTicker : ITicker
 {
 	/// <inheritdoc/>
-	public ValueTask DisposeAsync() => new ValueTask();
+	public ValueTask DisposeAsync() => new();
 
 	/// <inheritdoc/>
 	public Task StartAsync() => Task.CompletedTask;
diff --git a/Engine/Source/Programs/Shared/EpicGames.Core/CommandLineArguments.cs b/Engine/Source/Programs/Shared/EpicGames.Core/CommandLineArguments.cs
--- a/Engine/Source/Programs/Shared/EpicGames.Core/CommandLineArguments.cs
+++ b/Engine/Source/Programs/Shared/EpicGames.Core/CommandLineArguments.cs
@@ -853,7 +853,7 @@ namespace EpicGames.Core
 		/// <param name="logger">Sink for error/warning messages</param>
 		public T ApplyTo<T>(ILogger logger) where T : new()
 		{
-			T obj = new T();
+			T obj = new();
 			ApplyTo(obj, logger);
 			return obj;
 		}
@@ -914,7 +914,7 @@ namespace EpicGames.Core
 			List<ArgumentTarget> targets = GetArgumentTargetsForType(type);
 			foreach (ArgumentTarget target in targets)
 			{
-				StringBuilder descriptionBuilder = new StringBuilder();
+				StringBuilder descriptionBuilder = new();
 				foreach (DescriptionAttribute attribute in target.Member.GetCustomAttributes<DescriptionAttribute>())
 				{
 					if (descriptionBuilder.Length > 0)
@@ -999,7 +999,7 @@ namespace EpicGames.Core
 		/// <returns>Joined command line</returns>
 		public static string Join(IEnumerable<string> arguments)
 		{
-			StringBuilder result = new StringBuilder();
+			StringBuilder result = new();
 			foreach (string argument in arguments)
 			{
 				if (result.Length > 0)
@@ -1018,7 +1018,7 @@ namespace EpicGames.Core
 		/// <returns>Array of arguments</returns>
 		public static string[] Split(string commandLine)
 		{
-			StringBuilder argument = new StringBuilder();
+			StringBuilder argument = new();
 
 			List<string> arguments = [];
 			// First do a pass leaving all quotes in the arguments, they will be removed later
@@ -1063,7 +1063,7 @@ namespace EpicGames.Core
 		/// <returns>New argument list</returns>
 		public CommandLineArguments Append(IEnumerable<string> appendArguments)
 		{
-			CommandLineArguments newArguments = new CommandLineArguments([.. Enumerable.Concat(_arguments, appendArguments)]);
+			CommandLineArguments newArguments = new([.. Enumerable.Concat(_arguments, appendArguments)]);
 			for (int idx = 0; idx < _arguments.Length; idx++)
 			{
 				if (HasBeenUsed(idx))
@@ -1086,7 +1086,7 @@ namespace EpicGames.Core
 
 			// Split the arguments into the values array and an array of new arguments
 			int[] newArgumentIndex = new int[_arguments.Length];
-			List<string> newArgumentList = new List<string>(_arguments.Length);
+			List<string> newArgumentList = new(_arguments.Length);
 			for (int idx = 0; idx < _arguments.Length; idx++)
 			{
 				string argument = _arguments[idx];
@@ -1103,7 +1103,7 @@ namespace EpicGames.Core
 			}
 
 			// Create the new argument list, and mark the same arguments as used
-			CommandLineArguments newArguments = new CommandLineArguments([.. newArgumentList]);
+			CommandLineArguments newArguments = new([.. newArgumentList]);
 			for (int idx = 0; idx < _arguments.Length; idx++)
 			{
 				if (HasBeenUsed(idx) && newArgumentIndex[idx] != -1)
@@ -1441,7 +1441,7 @@ namespace EpicGames.Core
 		/// <returns></returns>
 		public override string ToString()
 		{
-			StringBuilder result = new StringBuilder();
+			StringBuilder result = new();
 			foreach (string argument in _arguments)
 			{
 				Append(result, argument);
diff --git a/Engine/Source/Programs/Shared/EpicGames.Core/DefaultConsoleLogger.cs b/Engine/Source/Programs/Shared/EpicGames.Core/DefaultConsoleLogger.cs
--- a/Engine/Source/Programs/Shared/EpicGames.Core/DefaultConsoleLogger.cs
+++ b/Engine/Source/Programs/Shared/EpicGames.Core/DefaultConsoleLogger.cs
@@ -53,7 +53,7 @@ namespace EpicGames.Core
 		/// <inheritdoc/>
 		public void Log<TState>(LogLevel logLevel, EventId eventId, TState state, Exception? exception, Func<TState, Exception?, string> formatter)
 		{
-			using StringWriter writer = new StringWriter();
+			using StringWriter writer = new();
 			Write<TState>(logLevel, state, exception, formatter, writer);
 			Console.WriteLine(writer.ToString());
 		}
diff --git a/Engine/Source/Programs/Shared/EpicGames.Core/Digest.cs b/Engine/Source/Programs/Shared/EpicGames.Core/Digest.cs
--- a/Engine/Source/Programs/Shared/EpicGames.Core/Digest.cs
+++ b/Engine/Source/Programs/Shared/EpicGames.Core/Digest.cs
@@ -257,7 +257,7 @@ namespace EpicGames.Core
 		/// <summary>
 		/// Zero digest value
 		/// </summary>
-		public static Digest<T> Zero => new Digest<T>(new byte[Traits.Length]);
+		public static Digest<T> Zero => new(new byte[Traits.Length]);
 
 		/// <summary>
 		/// Memory storing the digest data
diff --git a/Engine/Source/Programs/Shared/EpicGames.Core/DirectoryReference.cs b/Engine/Source/Programs/Shared/EpicGames.Core/DirectoryReference.cs
--- a/Engine/Source/Programs/Shared/EpicGames.Core/DirectoryReference.cs
+++ b/Engine/Source/Programs/Shared/EpicGames.Core/DirectoryReference.cs
@@ -305,7 +305,7 @@ namespace EpicGames.Core
 		/// Constructs a DirectoryInfo object from this reference
 		/// </summary>
 		/// <returns>New DirectoryInfo object</returns>
-		public DirectoryInfo ToDirectoryInfo() => new DirectoryInfo(FullName);
+		public DirectoryInfo ToDirectoryInfo() => new(FullName);
 
 		#region System.IO.Directory Wrapper Methods
 
diff --git a/Engine/Source/Programs/Shared/EpicGames.Core/EnumerableExtensions.cs b/Engine/Source/Programs/Shared/EpicGames.Core/EnumerableExtensions.cs
--- a/Engine/Source/Programs/Shared/EpicGames.Core/EnumerableExtensions.cs
+++ b/Engine/Source/Programs/Shared/EpicGames.Core/EnumerableExtensions.cs
@@ -31,7 +31,7 @@ namespace EpicGames.Core
 		/// <returns>Sequence of batches</returns>
 		public static IEnumerable<IReadOnlyList<TElement>> Batch<TElement>(this IEnumerable<TElement> sequence, int batchSize)
 		{
-			List<TElement> elements = new List<TElement>(batchSize);
+			List<TElement> elements = new(batchSize);
 			foreach (TElement element in sequence)
 			{
 				elements.Add(element);
diff --git a/Engine/Source/Programs/Shared/EpicGames.Core/ExceptionUtils.cs b/Engine/Source/Programs/Shared/EpicGames.Core/ExceptionUtils.cs
--- a/Engine/Source/Programs/Shared/EpicGames.Core/ExceptionUtils.cs
+++ b/Engine/Source/Programs/Shared/EpicGames.Core/ExceptionUtils.cs
@@ -70,7 +70,7 @@ namespace EpicGames.Core
 		/// <returns>String containing the exception information. May be multiple lines.</returns>
 		public static string FormatException(Exception ex)
 		{
-			StringBuilder errorMessage = new StringBuilder();
+			StringBuilder errorMessage = new();
 			if (ex is AggregateException exAgg)
 			{
 				Exception? innerException = exAgg.InnerException;
@@ -107,7 +107,7 @@ namespace EpicGames.Core
 				exceptionStack.Add(currentEx);
 			}
 
-			StringBuilder message = new StringBuilder();
+			StringBuilder message = new();
 			for (int idx = exceptionStack.Count - 1; idx >= 0; idx--)
 			{
 				Exception currentEx = exceptionStack[idx];
diff --git a/Engine/Source/Programs/Shared/EpicGames.Core/FileFilter.cs b/Engine/Source/Programs/Shared/EpicGames.Core/FileFilter.cs
--- a/Engine/Source/Programs/Shared/EpicGames.Core/FileFilter.cs
+++ b/Engine/Source/Programs/Shared/EpicGames.Core/FileFilter.cs
@@ -538,7 +538,7 @@ namespace EpicGames.Core
 		/// <returns>List of files matching the wildcard</returns>
 		public static List<FileReference> ResolveWildcard(DirectoryReference baseDir, string pattern)
 		{
-			FileFilter filter = new FileFilter(FileFilterType.Exclude);
+			FileFilter filter = new(FileFilterType.Exclude);
 			filter.AddRule(pattern);
 			return filter.ApplyToDirectory(baseDir, true);
 		}
diff --git a/Engine/Source/Programs/Shared/EpicGames.Core/FilePattern.cs b/Engine/Source/Programs/Shared/EpicGames.Core/FilePattern.cs
--- a/Engine/Source/Programs/Shared/EpicGames.Core/FilePattern.cs
+++ b/Engine/Source/Programs/Shared/EpicGames.Core/FilePattern.cs
@@ -91,7 +91,7 @@ namespace EpicGames.Core
 		public FilePattern(DirectoryReference rootDirectory, string pattern)
 		{
 			// Normalize the path separators
-			StringBuilder text = new StringBuilder(pattern);
+			StringBuilder text = new(pattern);
 			if (Path.DirectorySeparatorChar != '\\')
 			{
 				text.Replace('\\', Path.DirectorySeparatorChar);
@@ -178,7 +178,7 @@ namespace EpicGames.Core
 			}
 			else
 			{
-				StringBuilder pattern = new StringBuilder();
+				StringBuilder pattern = new();
 				foreach (string token in Tokens)
 				{
 					pattern.Append(token);
@@ -258,7 +258,7 @@ namespace EpicGames.Core
 		/// <returns>The regex pattern</returns>
 		public string GetRegexPattern()
 		{
-			StringBuilder pattern = new StringBuilder("^");
+			StringBuilder pattern = new("^");
 			pattern.Append(Regex.Escape(Tokens[0]));
 			for (int idx = 1; idx < Tokens.Count; idx += 2)
 			{
@@ -289,7 +289,7 @@ namespace EpicGames.Core
 		/// <returns>String representing the regex replacement pattern</returns>
 		public string GetRegexReplacementPattern()
 		{
-			StringBuilder pattern = new StringBuilder();
+			StringBuilder pattern = new();
 			for (int idx = 0; ; idx += 2)
 			{
 				// Append the escaped replacement character
@@ -328,7 +328,7 @@ namespace EpicGames.Core
 			// If we have multiple potential source files, but no wildcards in the output pattern, assume it's a directory and append the pattern from the source.
 			if (sourcePattern.ContainsWildcards() && !targetPattern.ContainsWildcards())
 			{
-				StringBuilder newPattern = new StringBuilder();
+				StringBuilder newPattern = new();
 				foreach (string token in targetPattern.Tokens)
 				{
 					newPattern.Append(token);
@@ -361,7 +361,7 @@ namespace EpicGames.Core
 				}
 
 				// Create a filter to match the source files
-				FileFilter filter = new FileFilter(FileFilterType.Exclude);
+				FileFilter filter = new(FileFilterType.Exclude);
 				filter.Include(String.Join("", sourcePattern.Tokens));
 
 				// Apply it to the source directory
diff --git a/Engine/Source/Programs/Shared/EpicGames.Core/FileSet.cs b/Engine/Source/Programs/Shared/EpicGames.Core/FileSet.cs
--- a/Engine/Source/Programs/Shared/EpicGames.Core/FileSet.cs
+++ b/Engine/Source/Programs/Shared/EpicGames.Core/FileSet.cs
@@ -208,7 +208,7 @@ namespace EpicGames.Core
 		/// <returns></returns>
 		public Dictionary<string, FileReference> Flatten()
 		{
-			Dictionary<string, FileReference> pathToSourceFile = new Dictionary<string, FileReference>(StringComparer.OrdinalIgnoreCase);
+			Dictionary<string, FileReference> pathToSourceFile = new(StringComparer.OrdinalIgnoreCase);
 			FlattenInternal(String.Empty, pathToSourceFile);
 			return pathToSourceFile;
 		}
@@ -353,7 +353,7 @@ namespace EpicGames.Core
 		/// <inheritdoc/>
 		public override IEnumerable<KeyValuePair<string, FileReference>> EnumerateFiles()
 		{
-			Dictionary<string, FileReference> files = new Dictionary<string, FileReference>(_lhs.EnumerateFiles(), StringComparer.OrdinalIgnoreCase);
+			Dictionary<string, FileReference> files = new(_lhs.EnumerateFiles(), StringComparer.OrdinalIgnoreCase);
 			foreach ((string name, FileReference file) in _rhs.EnumerateFiles())
 			{
 				FileReference? existingFile;
@@ -372,7 +372,7 @@ namespace EpicGames.Core
 		/// <inheritdoc/>
 		public override IEnumerable<KeyValuePair<string, FileSet>> EnumerateDirectories()
 		{
-			Dictionary<string, FileSet> nameToSubTree = new Dictionary<string, FileSet>(_lhs.EnumerateDirectories(), StringComparer.OrdinalIgnoreCase);
+			Dictionary<string, FileSet> nameToSubTree = new(_lhs.EnumerateDirectories(), StringComparer.OrdinalIgnoreCase);
 			foreach ((string name, FileSet subTree) in _rhs.EnumerateDirectories())
 			{
 				FileSet? existingSubTree;
@@ -412,14 +412,14 @@ namespace EpicGames.Core
 		/// <inheritdoc/>
 		public override IEnumerable<KeyValuePair<string, FileReference>> EnumerateFiles()
 		{
-			HashSet<string> rhsFiles = new HashSet<string>(_rhs.EnumerateFiles().Select(x => x.Key), StringComparer.OrdinalIgnoreCase);
+			HashSet<string> rhsFiles = new(_rhs.EnumerateFiles().Select(x => x.Key), StringComparer.OrdinalIgnoreCase);
 			return _lhs.EnumerateFiles().Where(x => !rhsFiles.Contains(x.Key));
 		}
 
 		/// <inheritdoc/>
 		public override IEnumerable<KeyValuePair<string, FileSet>> EnumerateDirectories()
 		{
-			Dictionary<string, FileSet> rhsDirs = new Dictionary<string, FileSet>(_rhs.EnumerateDirectories(), StringComparer.OrdinalIgnoreCase);
+			Dictionary<string, FileSet> rhsDirs = new(_rhs.EnumerateDirectories(), StringComparer.OrdinalIgnoreCase);
 			foreach ((string name, FileSet lhsSet) in _lhs.EnumerateDirectories())
 			{
 				FileSet? rhsSet;
@@ -478,7 +478,7 @@ namespace EpicGames.Core
 				FileFilter[] possibleFilters = [.. _filters.Where(x => x.PossiblyMatches(filterName))];
 				if (possibleFilters.Length > 0)
 				{
-					FileSetFromFilter subTreeFilter = new FileSetFromFilter(item.Value, possibleFilters);
+					FileSetFromFilter subTreeFilter = new(item.Value, possibleFilters);
 					yield return new KeyValuePair<string, FileSet>(item.Key, subTreeFilter);
 				}
 			}
diff --git a/Engine/Source/Programs/Shared/EpicGames.Core/FileUtils.cs b/Engine/Source/Programs/Shared/EpicGames.Core/FileUtils.cs
--- a/Engine/Source/Programs/Shared/EpicGames.Core/FileUtils.cs
+++ b/Engine/Source/Programs/Shared/EpicGames.Core/FileUtils.cs
@@ -220,7 +220,7 @@ namespace EpicGames.Core
 
 				if (lockInfoList != null && lockInfoList.Count > 0)
 				{
-					StringBuilder message = new StringBuilder("Processes with open handles to file:");
+					StringBuilder message = new("Processes with open handles to file:");
 					foreach (FileLockInfoWin32 lockInfo in lockInfoList)
 					{
 						message.Append($"\n  {lockInfo}");
@@ -556,7 +556,7 @@ namespace EpicGames.Core
 		/// <returns></returns>
 		public static int GetFileMode_Mac(string fileName)
 		{
-			stat64_t stat = new stat64_t();
+			stat64_t stat = new();
 			int result = stat64(fileName, stat);
 			return (result >= 0) ? stat.st_mode : -1;
 		}
@@ -578,7 +578,7 @@ namespace EpicGames.Core
 		/// <returns></returns>
 		public static int GetFileMode_Linux(string fileName)
 		{
-			stat64_linux_t stat = new stat64_linux_t();
+			stat64_linux_t stat = new();
 			int result = stat64_linux(1, fileName, stat);
 			return (result >= 0) ? (int)stat.st_mode : -1;
 		}
@@ -640,7 +640,7 @@ namespace EpicGames.Core
 		[return: MarshalAs(UnmanagedType.Bool)]
 		static extern bool SetFileAttributesW(string lpFileName, uint dwFileAttributes);
 
-		static readonly IntPtr INVALID_HANDLE_VALUE = new IntPtr(-1);
+		static readonly IntPtr INVALID_HANDLE_VALUE = new(-1);
 
 		const int ERROR_FILE_NOT_FOUND = 2;
 		const int ERROR_PATH_NOT_FOUND = 3;
@@ -650,7 +650,7 @@ namespace EpicGames.Core
 		[SupportedOSPlatform("windows")]
 		private static void ForceDeleteLongDirectoryContentsWin32(string dirName)
 		{
-			WIN32_FIND_DATA findData = new WIN32_FIND_DATA();
+			WIN32_FIND_DATA findData = new();
 
 			const string RawPathPrefix = "\\\\?\\";
 			if (!dirName.StartsWith(RawPathPrefix, StringComparison.Ordinal))
@@ -731,7 +731,7 @@ namespace EpicGames.Core
 			}
 
 			// Capture the exception for failing to delete the file
-			Win32Exception ex = new Win32Exception();
+			Win32Exception ex = new();
 
 			// Check the file exists and is not readonly
 			int attributes = GetFileAttributesW(fileName);
@@ -755,7 +755,7 @@ namespace EpicGames.Core
 			}
 
 			// Get a useful error message about why the delete failed
-			StringBuilder message = new StringBuilder($"Unable to delete {fileName} - {ex.Message}");
+			StringBuilder message = new($"Unable to delete {fileName} - {ex.Message}");
 			if (ex.NativeErrorCode == ERROR_ACCESS_DENIED)
 			{
 				string? lockInfo = GetFileLockInfo(fileName);
@@ -857,7 +857,7 @@ namespace EpicGames.Core
 			uint sessionHandle = 0;
 			try
 			{
-				StringBuilder sessionKey = new StringBuilder(CCH_RM_SESSION_KEY + 1);
+				StringBuilder sessionKey = new(CCH_RM_SESSION_KEY + 1);
 
 				int result = RmStartSession(out sessionHandle, 0, sessionKey);
 				if (result != 0)
diff --git a/Engine/Source/Programs/Shared/EpicGames.Core/ForwardingLogger.cs b/Engine/Source/Programs/Shared/EpicGames.Core/ForwardingLogger.cs
--- a/Engine/Source/Programs/Shared/EpicGames.Core/ForwardingLogger.cs
+++ b/Engine/Source/Programs/Shared/EpicGames.Core/ForwardingLogger.cs
@@ -46,7 +46,7 @@ namespace EpicGames.Core
 		/// <inheritdoc/>
 		public IDisposable? BeginScope<TState>(TState state) where TState : notnull
 		{
-			DisposableList scope = new DisposableList();
+			DisposableList scope = new();
 			foreach (ILogger logger in Loggers)
 			{
 				IDisposable? next = logger.BeginScope(state);
diff --git a/Engine/Source/Programs/Shared/EpicGames.Core/HelpUtils.cs b/Engine/Source/Programs/Shared/EpicGames.Core/HelpUtils.cs
--- a/Engine/Source/Programs/Shared/EpicGames.Core/HelpUtils.cs
+++ b/Engine/Source/Programs/Shared/EpicGames.Core/HelpUtils.cs
@@ -73,7 +73,7 @@ namespace EpicGames.Core
 		/// <returns>The description text</returns>
 		public static string GetDescription(Type type)
 		{
-			StringBuilder descriptionText = new StringBuilder();
+			StringBuilder descriptionText = new();
 			foreach (DescriptionAttribute attribute in type.GetCustomAttributes(typeof(DescriptionAttribute), false).OfType<DescriptionAttribute>())
 			{
 				if (descriptionText.Length > 0)
@@ -165,7 +165,7 @@ namespace EpicGames.Core
 			if (items.Count > 0)
 			{
 				// string used to intent the param
-				string indentString = new string(' ', indent);
+				string indentString = new(' ', indent);
 
 				// default the padding value
 				int rightPadding = Math.Max(minFirstColumnWidth, items.Max(x => x.Key.Length + 2));
