Shared: prefer implicit object creation when type is apparent

diff --git a/Engine/Source/Programs/Shared/EpicGames.Core/ILoggerProperty.cs b/Engine/Source/Programs/Shared/EpicGames.Core/ILoggerProperty.cs
--- a/Engine/Source/Programs/Shared/EpicGames.Core/ILoggerProperty.cs
+++ b/Engine/Source/Programs/Shared/EpicGames.Core/ILoggerProperty.cs
@@ -71,7 +71,7 @@ namespace EpicGames.Core
 		/// </example>
 		public static LogScopedPropertyBuilder WithProperty(this ILogger logger, string key, object value)
 		{
-			LogScopedPropertyBuilder builder = new LogScopedPropertyBuilder(logger);
+			LogScopedPropertyBuilder builder = new(logger);
 			builder.WithProperty(key, value);
 			return builder;
 		}
diff --git a/Engine/Source/Programs/Shared/EpicGames.Core/IMemoryAllocator.cs b/Engine/Source/Programs/Shared/EpicGames.Core/IMemoryAllocator.cs
--- a/Engine/Source/Programs/Shared/EpicGames.Core/IMemoryAllocator.cs
+++ b/Engine/Source/Programs/Shared/EpicGames.Core/IMemoryAllocator.cs
@@ -114,10 +114,10 @@ namespace EpicGames.Core
 			byte* GetPointer() => (byte*)_handle.ToPointer();
 
 			/// <inheritdoc/>
-			public override Span<byte> GetSpan() => new Span<byte>(GetPointer(), _length);
+			public override Span<byte> GetSpan() => new(GetPointer(), _length);
 
 			/// <inheritdoc/>
-			public override MemoryHandle Pin(int elementIndex) => new MemoryHandle(GetPointer() + elementIndex);
+			public override MemoryHandle Pin(int elementIndex) => new(GetPointer() + elementIndex);
 
 			/// <inheritdoc/>
 			public override void Unpin() { }
@@ -174,10 +174,10 @@ namespace EpicGames.Core
 			}
 
 			/// <inheritdoc/>
-			public override Span<byte> GetSpan() => new Span<byte>(_pointer, _length);
+			public override Span<byte> GetSpan() => new(_pointer, _length);
 
 			/// <inheritdoc/>
-			public override MemoryHandle Pin(int elementIndex) => new MemoryHandle(_pointer + elementIndex);
+			public override MemoryHandle Pin(int elementIndex) => new(_pointer + elementIndex);
 
 			/// <inheritdoc/>
 			public override void Unpin() { }
@@ -254,10 +254,10 @@ namespace EpicGames.Core
 				=> $"{_tag ?? "???"} ({_allocation?.Memory.Length ?? 0} bytes)";
 		}
 
-		readonly object _lockObject = new object();
+		readonly object _lockObject = new();
 		readonly IMemoryAllocator<byte> _inner;
 		long _allocatedSize;
-		readonly LinkedList<MemoryAllocation> _allocations = new LinkedList<MemoryAllocation>();
+		readonly LinkedList<MemoryAllocation> _allocations = new();
 
 		/// <summary>
 		/// Total size of allocated memory
diff --git a/Engine/Source/Programs/Shared/EpicGames.Core/JsonLogEvent.cs b/Engine/Source/Programs/Shared/EpicGames.Core/JsonLogEvent.cs
--- a/Engine/Source/Programs/Shared/EpicGames.Core/JsonLogEvent.cs
+++ b/Engine/Source/Programs/Shared/EpicGames.Core/JsonLogEvent.cs
@@ -127,7 +127,7 @@ namespace EpicGames.Core
 			int lineIndex = 0;
 			int lineCount = 1;
 
-			Utf8JsonReader reader = new Utf8JsonReader(data.Span);
+			Utf8JsonReader reader = new(data.Span);
 			if (reader.Read() && reader.TokenType == JsonTokenType.StartObject)
 			{
 				while (reader.Read() && reader.TokenType == JsonTokenType.PropertyName)
@@ -220,14 +220,14 @@ namespace EpicGames.Core
 			return (LogLevel)result;
 		}
 
-		static readonly Utf8String s_newlineEscaped = new Utf8String("\\n");
+		static readonly Utf8String s_newlineEscaped = new("\\n");
 
 		/// <summary>
 		/// Gets the rendered message from the event data
 		/// </summary>
 		public Utf8String GetRenderedMessage()
 		{
-			Utf8JsonReader reader = new Utf8JsonReader(Data.Span);
+			Utf8JsonReader reader = new(Data.Span);
 			if (reader.Read() && reader.TokenType == JsonTokenType.StartObject)
 			{
 				while (reader.Read() && reader.TokenType == JsonTokenType.PropertyName)
@@ -252,7 +252,7 @@ namespace EpicGames.Core
 		/// </summary>
 		public string GetLegacyLogLine()
 		{
-			Utf8JsonReader reader = new Utf8JsonReader(Data.Span);
+			Utf8JsonReader reader = new(Data.Span);
 			if (reader.Read() && reader.TokenType == JsonTokenType.StartObject)
 			{
 				Utf8String? message = null;
@@ -297,7 +297,7 @@ namespace EpicGames.Core
 		{
 			if (Data.Span.IndexOf(s_newlineEscaped.Span) != -1)
 			{
-				Utf8JsonReader reader = new Utf8JsonReader(Data.Span);
+				Utf8JsonReader reader = new(Data.Span);
 				if (reader.Read() && reader.TokenType == JsonTokenType.StartObject)
 				{
 					while (reader.Read() && reader.TokenType == JsonTokenType.PropertyName)
diff --git a/Engine/Source/Programs/Shared/EpicGames.Core/JsonObject.cs b/Engine/Source/Programs/Shared/EpicGames.Core/JsonObject.cs
--- a/Engine/Source/Programs/Shared/EpicGames.Core/JsonObject.cs
+++ b/Engine/Source/Programs/Shared/EpicGames.Core/JsonObject.cs
@@ -51,7 +51,7 @@ namespace EpicGames.Core
 
 			if (original is OrderedDictionary dictionary)
 			{
-				OrderedDictionary copy = new OrderedDictionary(StringComparer.InvariantCultureIgnoreCase);
+				OrderedDictionary copy = new(StringComparer.InvariantCultureIgnoreCase);
 				foreach (DictionaryEntry entry in dictionary)
 				{
 					copy.Add(entry.Key, DeepCopyHelper(entry.Value));
@@ -78,7 +78,7 @@ namespace EpicGames.Core
 		/// <returns> An ordered dictionary with a deep copy of elements from dictionary.</returns>
 		public static OrderedDictionary CreateDeepCopy(this OrderedDictionary dictionary)
 		{
-			OrderedDictionary copy = new OrderedDictionary(StringComparer.InvariantCultureIgnoreCase);
+			OrderedDictionary copy = new(StringComparer.InvariantCultureIgnoreCase);
 			foreach (DictionaryEntry entry in dictionary)
 			{
 				copy.Add(entry.Key, DeepCopyHelper(entry.Value));
@@ -196,7 +196,7 @@ namespace EpicGames.Core
 				case JsonValueKind.Number:
 					return element.GetDouble();
 				case JsonValueKind.Object:
-					OrderedDictionary dictionary = new OrderedDictionary(StringComparer.InvariantCultureIgnoreCase);
+					OrderedDictionary dictionary = new(StringComparer.InvariantCultureIgnoreCase);
 					foreach (JsonProperty property in element.EnumerateObject())
 					{
 						dictionary.Add(property.Name, ParseElement(property.Value));
@@ -935,15 +935,15 @@ namespace EpicGames.Core
 		public string ToJsonString()
 		{
 			SystemJsonObject jsonObjectToWrite = ToSystemJsonObject();
-			JsonWriterOptions options = new JsonWriterOptions();
+			JsonWriterOptions options = new();
 			options.Indented = true;
 			// IMPORTANT: Utf8JsonWriter blocks certain characters like +, &, <, >,` from being escaped in a global block list
 			// Best way around is to use the relaxed JavaScriptEncoder.UnsafeRelaxedJsonEscaping.
 			// However this can have security implications if this string is ever written to an HTML page or script
 			// https://learn.microsoft.com/en-us/dotnet/standard/serialization/system-text-json/character-encoding
 			options.Encoder = JavaScriptEncoder.UnsafeRelaxedJsonEscaping;
-			using MemoryStream jsonMemoryStream = new MemoryStream();
-			using (Utf8JsonWriter writer = new Utf8JsonWriter(jsonMemoryStream, options))
+			using MemoryStream jsonMemoryStream = new();
+			using (Utf8JsonWriter writer = new(jsonMemoryStream, options))
 			{
 				jsonObjectToWrite.WriteTo(writer);
 				writer.Flush();
@@ -951,7 +951,7 @@ namespace EpicGames.Core
 			// The Utf8JsonWriter doesn't format json the same as we want in UE, we massage the string to meet our standards 
 			string jsonString = Encoding.UTF8.GetString(jsonMemoryStream.ToArray());
 			string[] lines = jsonString.Split([Environment.NewLine], StringSplitOptions.None);
-			StringBuilder jsonStringBuilder = new StringBuilder();
+			StringBuilder jsonStringBuilder = new();
 			foreach (string line in lines)
 			{
 				// Utf8JsonWriter uses 2 spaces for indents, we replace them with tabs here 
diff --git a/Engine/Source/Programs/Shared/EpicGames.Core/JsonSchema.cs b/Engine/Source/Programs/Shared/EpicGames.Core/JsonSchema.cs
--- a/Engine/Source/Programs/Shared/EpicGames.Core/JsonSchema.cs
+++ b/Engine/Source/Programs/Shared/EpicGames.Core/JsonSchema.cs
@@ -201,7 +201,7 @@ namespace EpicGames.Core
 		public void Write(Utf8JsonWriter writer)
 		{
 			// Determine reference counts for each type. Any type referenced at least twice will be split off into a separate definition.
-			JsonTypeRefCollector refCollector = new JsonTypeRefCollector();
+			JsonTypeRefCollector refCollector = new();
 			refCollector.WriteType(RootType);
 
 			// Assign names to each type definition
@@ -231,7 +231,7 @@ namespace EpicGames.Core
 				writer.WriteString("$id", Id);
 			}
 
-			JsonSchemaWriter schemaWriter = new JsonSchemaWriter(writer, typeToDefinition);
+			JsonSchemaWriter schemaWriter = new(writer, typeToDefinition);
 			RootType.Write(schemaWriter);
 
 			if (typeToDefinition.Count > 0)
@@ -255,7 +255,7 @@ namespace EpicGames.Core
 		/// <param name="stream">The output stream</param>
 		public void Write(Stream stream)
 		{
-			using (Utf8JsonWriter writer = new Utf8JsonWriter(stream, new JsonWriterOptions { Indented = true }))
+			using (Utf8JsonWriter writer = new(stream, new JsonWriterOptions { Indented = true }))
 			{
 				Write(writer);
 			}
@@ -281,7 +281,7 @@ namespace EpicGames.Core
 		/// <returns>New schema object</returns>
 		public static JsonSchema FromType(Type type, XmlDocReader? xmlDocReader)
 		{
-			JsonSchemaFactory factory = new JsonSchemaFactory(xmlDocReader);
+			JsonSchemaFactory factory = new(xmlDocReader);
 			return factory.CreateSchema(type);
 		}
 
@@ -290,7 +290,7 @@ namespace EpicGames.Core
 		/// </summary>
 		public static JsonSchema CreateSchema(Type type, XmlDocReader? xmlDocReader)
 		{
-			JsonSchemaFactory factory = new JsonSchemaFactory(xmlDocReader);
+			JsonSchemaFactory factory = new(xmlDocReader);
 			return factory.CreateSchema(type);
 		}
 	}
diff --git a/Engine/Source/Programs/Shared/EpicGames.Core/JsonSchemaFactory.cs b/Engine/Source/Programs/Shared/EpicGames.Core/JsonSchemaFactory.cs
--- a/Engine/Source/Programs/Shared/EpicGames.Core/JsonSchemaFactory.cs
+++ b/Engine/Source/Programs/Shared/EpicGames.Core/JsonSchemaFactory.cs
@@ -137,7 +137,7 @@ namespace EpicGames.Core
 						}
 						if (interfaceType.GetGenericTypeDefinition() == typeof(IDictionary<,>))
 						{
-							JsonSchemaObject obj = new JsonSchemaObject();
+							JsonSchemaObject obj = new();
 							obj.AdditionalProperties = CreateSchemaType(arguments[1]);
 							return obj;
 						}
@@ -147,14 +147,14 @@ namespace EpicGames.Core
 				JsonKnownTypesAttribute? knownTypes = type.GetCustomAttribute<JsonKnownTypesAttribute>(false);
 				if (knownTypes != null)
 				{
-					JsonSchemaOneOf obj = new JsonSchemaOneOf();
+					JsonSchemaOneOf obj = new();
 					TypeCache[type] = obj;
 					SetOneOfProperties(obj, type, knownTypes.Types);
 					return obj;
 				}
 				else
 				{
-					JsonSchemaObject obj = new JsonSchemaObject();
+					JsonSchemaObject obj = new();
 					TypeCache[type] = obj;
 					SetObjectProperties(obj, type);
 					return obj;
@@ -193,7 +193,7 @@ namespace EpicGames.Core
 				JsonDiscriminatorAttribute? attribute = knownType.GetCustomAttribute<JsonDiscriminatorAttribute>();
 				if (attribute != null)
 				{
-					JsonSchemaObject knownObject = new JsonSchemaObject();
+					JsonSchemaObject knownObject = new();
 					knownObject.Properties.Add(new JsonSchemaProperty("type", "Type discriminator", new JsonSchemaEnum([attribute.Name], ["Identifier for the derived type"])));
 					SetObjectProperties(knownObject, knownType);
 					obj.Types.Add(knownObject);
diff --git a/Engine/Source/Programs/Shared/EpicGames.Core/JsonTracer.cs b/Engine/Source/Programs/Shared/EpicGames.Core/JsonTracer.cs
--- a/Engine/Source/Programs/Shared/EpicGames.Core/JsonTracer.cs
+++ b/Engine/Source/Programs/Shared/EpicGames.Core/JsonTracer.cs
@@ -70,7 +70,7 @@ namespace EpicGames.Core
 		private readonly List<JsonTracerSpan.Reference> _references;
 
 		public DateTimeOffset StartTimestamp { get; }
-		public Dictionary<string, object> Tags => new Dictionary<string, object>(_tags);
+		public Dictionary<string, object> Tags => new(_tags);
 
 		/// <summary>
 		/// The finish time of the span; only valid after a call to <see cref="Finish()"/>.
@@ -96,7 +96,7 @@ namespace EpicGames.Core
 		/// <seealso cref="MockSpan.References"/>
 		public string? ParentId { get; }
 
-		private readonly object _lock = new object();
+		private readonly object _lock = new();
 
 		public JsonTracerSpan(JsonTracer jsonTracer, string operationName, DateTimeOffset startTimestamp, Dictionary<string, object>? initialTags, List<Reference>? references)
 		{
@@ -404,7 +404,7 @@ namespace EpicGames.Core
 		public IScopeManager ScopeManager { get; }
 		public ISpan? ActiveSpan => ScopeManager.Active?.Span;
 
-		private readonly object _lock = new object();
+		private readonly object _lock = new();
 		private readonly List<JsonTracerSpan> _finishedSpans = [];
 		private readonly DirectoryReference? _telemetryDir;
 
@@ -419,7 +419,7 @@ namespace EpicGames.Core
 			string? telemetryDir = Environment.GetEnvironmentVariable("UE_TELEMETRY_DIR");
 			if (telemetryDir != null)
 			{
-				JsonTracer tracer = new JsonTracer(new DirectoryReference(telemetryDir));
+				JsonTracer tracer = new(new DirectoryReference(telemetryDir));
 				return GlobalTracer.RegisterIfAbsent(tracer) ? tracer : null;
 			}
 
@@ -456,7 +456,7 @@ namespace EpicGames.Core
 					file = FileReference.Combine(_telemetryDir, fileName);
 				}
 
-				using (JsonWriter writer = new JsonWriter(file))
+				using (JsonWriter writer = new(file))
 				{
 					GetFinishedSpansAsJson(writer);
 				}
diff --git a/Engine/Source/Programs/Shared/EpicGames.Core/JsonWriter.cs b/Engine/Source/Programs/Shared/EpicGames.Core/JsonWriter.cs
--- a/Engine/Source/Programs/Shared/EpicGames.Core/JsonWriter.cs
+++ b/Engine/Source/Programs/Shared/EpicGames.Core/JsonWriter.cs
@@ -389,7 +389,7 @@ namespace EpicGames.Core
 
 			// Otherwise, create the string builder, append the known portion that doesn't have an escape
 			// and continue processing the string starting at the first character needing to be escaped.
-			StringBuilder result = new StringBuilder();
+			StringBuilder result = new();
 			result.Append(value.AsSpan(0, idx));
 			for (; idx < value.Length; idx++)
 			{
diff --git a/Engine/Source/Programs/Shared/EpicGames.Core/ListSegment.cs b/Engine/Source/Programs/Shared/EpicGames.Core/ListSegment.cs
--- a/Engine/Source/Programs/Shared/EpicGames.Core/ListSegment.cs
+++ b/Engine/Source/Programs/Shared/EpicGames.Core/ListSegment.cs
@@ -16,7 +16,7 @@ namespace EpicGames.Core
 		/// <summary>
 		/// An empty list segment of the given element type
 		/// </summary>
-		public static ListSegment<T> Empty<T>() => new ListSegment<T>(Array.Empty<T>(), 0, 0);
+		public static ListSegment<T> Empty<T>() => new(Array.Empty<T>(), 0, 0);
 	}
 
 	/// <summary>
diff --git a/Engine/Source/Programs/Shared/EpicGames.Core/Log.cs b/Engine/Source/Programs/Shared/EpicGames.Core/Log.cs
--- a/Engine/Source/Programs/Shared/EpicGames.Core/Log.cs
+++ b/Engine/Source/Programs/Shared/EpicGames.Core/Log.cs
@@ -101,12 +101,12 @@ namespace EpicGames.Core
 		/// <summary>
 		/// Singleton instance of the default output logger
 		/// </summary>
-		private static readonly DefaultLogger s_defaultLogger = new DefaultLogger();
+		private static readonly DefaultLogger s_defaultLogger = new();
 
 		/// <summary>
 		/// Logger instance which parses events and forwards them to the main logger.
 		/// </summary>
-		private static readonly LegacyEventLogger s_legacyLogger = new LegacyEventLogger(s_defaultLogger);
+		private static readonly LegacyEventLogger s_legacyLogger = new(s_defaultLogger);
 
 		/// <summary>
 		/// Accessor for the global event parser from legacy events
@@ -237,7 +237,7 @@ namespace EpicGames.Core
 			string extension = outputFile.GetExtension();
 
 			Regex backupForm =
-				new Regex(filenameWithoutExtension + @"-backup-\d\d\d\d\.\d\d\.\d\d-\d\d\.\d\d\.\d\d" + extension);
+				new(filenameWithoutExtension + @"-backup-\d\d\d\d\.\d\d\.\d\d-\d\d\.\d\d\.\d\d" + extension);
 
 			foreach (FileReference oldBackup in DirectoryReference
 				.EnumerateFiles(outputFile.Directory)
@@ -372,7 +372,7 @@ namespace EpicGames.Core
 		{
 			if (Logger.IsEnabled((LogLevel)verbosity))
 			{
-				StringBuilder message = new StringBuilder();
+				StringBuilder message = new();
 
 				// Get the severity prefix for this message
 				if (IncludeSeverityPrefix && ((formatOptions & LogFormatOptions.NoSeverityPrefix) == 0))
@@ -861,7 +861,7 @@ namespace EpicGames.Core
 	/// </summary>
 	public class CaptureLogger : ILogger
 	{
-		readonly LoggerScopeCollection _scopeCollection = new LoggerScopeCollection();
+		readonly LoggerScopeCollection _scopeCollection = new();
 
 		/// <summary>
 		/// List of captured events
@@ -943,8 +943,8 @@ namespace EpicGames.Core
 			}
 		}
 
-		private readonly LoggerScopeCollection _scopeCollection = new LoggerScopeCollection();
-		private readonly object _syncObject = new object();
+		private readonly LoggerScopeCollection _scopeCollection = new();
+		private readonly object _syncObject = new();
 
 		/// <summary>
 		/// Minimum level for outputting messages
@@ -1008,7 +1008,7 @@ namespace EpicGames.Core
 		/// <summary>
 		/// Stack of status scope information.
 		/// </summary>
-		private readonly Stack<StatusMessage> _statusMessageStack = new Stack<StatusMessage>();
+		private readonly Stack<StatusMessage> _statusMessageStack = new();
 
 		/// <summary>
 		/// The currently visible status text
@@ -1023,7 +1023,7 @@ namespace EpicGames.Core
 		/// <summary>
 		/// Last time a status message was pushed to the stack
 		/// </summary>
-		private readonly Stopwatch _statusTimer = new Stopwatch();
+		private readonly Stopwatch _statusTimer = new();
 
 		/// <summary>
 		/// Background task for writing to files
@@ -1123,7 +1123,7 @@ namespace EpicGames.Core
 			{
 				OutputFile = outputFile;
 				DirectoryReference.CreateDirectory(outputFile.Directory);
-				TextWriterTraceListener logTraceListener = new TextWriterTraceListener(new StreamWriter(outputFile.FullName), name);
+				TextWriterTraceListener logTraceListener = new(new StreamWriter(outputFile.FullName), name);
 				lock (_syncObject)
 				{
 					Trace.Listeners.Add(logTraceListener);
@@ -1390,7 +1390,7 @@ namespace EpicGames.Core
 			{
 				FlushStatusHeading();
 
-				StatusMessage newStatusMessage = new StatusMessage(message, message);
+				StatusMessage newStatusMessage = new(message, message);
 				_statusMessageStack.Push(newStatusMessage);
 
 				_statusTimer.Restart();
@@ -1471,7 +1471,7 @@ namespace EpicGames.Core
 					numCommonChars = 0;
 				}
 
-				StringBuilder text = new StringBuilder();
+				StringBuilder text = new();
 				text.Append('\b', _statusText.Length - numCommonChars);
 				text.Append(newStatusText, numCommonChars, newStatusText.Length - numCommonChars);
 				if (newStatusText.Length < _statusText.Length)
diff --git a/Engine/Source/Programs/Shared/EpicGames.Core/LogEvent.cs b/Engine/Source/Programs/Shared/EpicGames.Core/LogEvent.cs
--- a/Engine/Source/Programs/Shared/EpicGames.Core/LogEvent.cs
+++ b/Engine/Source/Programs/Shared/EpicGames.Core/LogEvent.cs
@@ -20,28 +20,28 @@ namespace EpicGames.Core
 	public static class LogEventPropertyName
 	{
 #pragma warning disable CS1591 // Missing documentation
-		public static readonly Utf8String Time = new Utf8String("time");
-		public static readonly Utf8String Level = new Utf8String("level");
-		public static readonly Utf8String Id = new Utf8String("id");
-		public static readonly Utf8String Line = new Utf8String("line");
-		public static readonly Utf8String LineCount = new Utf8String("lineCount");
-		public static readonly Utf8String Message = new Utf8String("message");
-		public static readonly Utf8String Format = new Utf8String("format");
-		public static readonly Utf8String Properties = new Utf8String("properties");
-
-		public static readonly Utf8String Type = new Utf8String("$type");
-		public static readonly Utf8String Text = new Utf8String("$text");
-
-		public static readonly Utf8String File = new Utf8String("file"); // For source file / asset types
-		public static readonly Utf8String Identifier = new Utf8String("identifier"); // For symbols
-		public static readonly Utf8String RelativePath = new Utf8String("relativePath");
-		public static readonly Utf8String DepotPath = new Utf8String("depotPath");
-		public static readonly Utf8String Target = new Utf8String("target"); // For hyperlinks
-
-		public static readonly Utf8String Exception = new Utf8String("exception");
-		public static readonly Utf8String Trace = new Utf8String("trace");
-		public static readonly Utf8String InnerException = new Utf8String("innerException");
-		public static readonly Utf8String InnerExceptions = new Utf8String("innerExceptions");
+		public static readonly Utf8String Time = new("time");
+		public static readonly Utf8String Level = new("level");
+		public static readonly Utf8String Id = new("id");
+		public static readonly Utf8String Line = new("line");
+		public static readonly Utf8String LineCount = new("lineCount");
+		public static readonly Utf8String Message = new("message");
+		public static readonly Utf8String Format = new("format");
+		public static readonly Utf8String Properties = new("properties");
+
+		public static readonly Utf8String Type = new("$type");
+		public static readonly Utf8String Text = new("$text");
+
+		public static readonly Utf8String File = new("file"); // For source file / asset types
+		public static readonly Utf8String Identifier = new("identifier"); // For symbols
+		public static readonly Utf8String RelativePath = new("relativePath");
+		public static readonly Utf8String DepotPath = new("depotPath");
+		public static readonly Utf8String Target = new("target"); // For hyperlinks
+
+		public static readonly Utf8String Exception = new("exception");
+		public static readonly Utf8String Trace = new("trace");
+		public static readonly Utf8String InnerException = new("innerException");
+		public static readonly Utf8String InnerExceptions = new("innerExceptions");
 #pragma warning restore CS1591 // Missing documentation
 	}
 
@@ -121,7 +121,7 @@ namespace EpicGames.Core
 				=> _properties.GetEnumerator();
 		}
 
-		static readonly JsonSerializerOptions s_jsonSerializerOptions = new JsonSerializerOptions { PropertyNamingPolicy = JsonNamingPolicy.CamelCase, DefaultIgnoreCondition = JsonIgnoreCondition.WhenWritingNull };
+		static readonly JsonSerializerOptions s_jsonSerializerOptions = new() { PropertyNamingPolicy = JsonNamingPolicy.CamelCase, DefaultIgnoreCondition = JsonIgnoreCondition.WhenWritingNull };
 
 		/// <summary>
 		/// Constructor
@@ -254,7 +254,7 @@ namespace EpicGames.Core
 		/// <returns></returns>
 		public static LogEvent Read(ReadOnlySpan<byte> data)
 		{
-			Utf8JsonReader reader = new Utf8JsonReader(data);
+			Utf8JsonReader reader = new(data);
 			reader.Read();
 			return Read(ref reader);
 		}
@@ -268,9 +268,9 @@ namespace EpicGames.Core
 		public static LogEvent Read(ref Utf8JsonReader reader)
 #pragma warning restore CA1045 // Do not pass types by reference
 		{
-			DateTime time = new DateTime(0);
+			DateTime time = new(0);
 			LogLevel level = LogLevel.None;
-			EventId eventId = new EventId(0);
+			EventId eventId = new(0);
 			int line = 0;
 			int lineCount = 1;
 			string message = String.Empty;
@@ -588,8 +588,8 @@ namespace EpicGames.Core
 		/// </summary>
 		public byte[] ToJsonBytes()
 		{
-			ArrayBufferWriter<byte> buffer = new ArrayBufferWriter<byte>();
-			using (Utf8JsonWriter writer = new Utf8JsonWriter(buffer))
+			ArrayBufferWriter<byte> buffer = new();
+			using (Utf8JsonWriter writer = new(buffer))
 			{
 				Write(writer);
 			}
@@ -601,8 +601,8 @@ namespace EpicGames.Core
 		/// </summary>
 		public string ToJson()
 		{
-			ArrayBufferWriter<byte> buffer = new ArrayBufferWriter<byte>();
-			using (Utf8JsonWriter writer = new Utf8JsonWriter(buffer))
+			ArrayBufferWriter<byte> buffer = new();
+			using (Utf8JsonWriter writer = new(buffer))
 			{
 				Write(writer);
 			}
diff --git a/Engine/Source/Programs/Shared/EpicGames.Core/LogEventBuilder.cs b/Engine/Source/Programs/Shared/EpicGames.Core/LogEventBuilder.cs
--- a/Engine/Source/Programs/Shared/EpicGames.Core/LogEventBuilder.cs
+++ b/Engine/Source/Programs/Shared/EpicGames.Core/LogEventBuilder.cs
@@ -99,7 +99,7 @@ namespace EpicGames.Core
 		/// <returns>New span for the given range</returns>
 		public LogEventSpan AddSpan(int offset, int length, string name)
 		{
-			LogEventSpan span = new LogEventSpan(offset, Text.Substring(offset, length));
+			LogEventSpan span = new(offset, Text.Substring(offset, length));
 			Spans.Add(name, span);
 			return span;
 		}
@@ -251,7 +251,7 @@ namespace EpicGames.Core
 			string? format = null;
 			if (_spans != null)
 			{
-				StringBuilder builder = new StringBuilder();
+				StringBuilder builder = new();
 				foreach (LogSpan span in _spans)
 				{
 					if (span._offset >= offset)
@@ -278,7 +278,7 @@ namespace EpicGames.Core
 		/// <returns>New span for the given range</returns>
 		public void Annotate(string name, int offset, int length, object? value = null)
 		{
-			LogSpan span = new LogSpan(name, offset, length, value);
+			LogSpan span = new(name, offset, length, value);
 
 			_properties ??= new Dictionary<string, object>(StringComparer.OrdinalIgnoreCase);
 			_properties.Add(name, span);
diff --git a/Engine/Source/Programs/Shared/EpicGames.Core/LogEventMarkup.cs b/Engine/Source/Programs/Shared/EpicGames.Core/LogEventMarkup.cs
--- a/Engine/Source/Programs/Shared/EpicGames.Core/LogEventMarkup.cs
+++ b/Engine/Source/Programs/Shared/EpicGames.Core/LogEventMarkup.cs
@@ -13,15 +13,15 @@ namespace EpicGames.Core
 	public static class LogEventMarkup
 	{
 #pragma warning disable CS1591
-		public static LogValue Channel => new LogValue(LogValueType.Channel, "");
-		public static LogValue Severity => new LogValue(LogValueType.Severity, "");
-		public static LogValue Message => new LogValue(LogValueType.Message, "");
-		public static LogValue LineNumber => new LogValue(LogValueType.LineNumber, "");
-		public static LogValue ColumnNumber => new LogValue(LogValueType.ColumnNumber, "");
-		public static LogValue Symbol => new LogValue(LogValueType.Symbol, "");
-		public static LogValue ErrorCode => new LogValue(LogValueType.ErrorCode, "");
-		public static LogValue ToolName => new LogValue(LogValueType.ToolName, "");
-		public static LogValue ScreenshotTest => new LogValue(LogValueType.ScreenshotTest, "");
+		public static LogValue Channel => new(LogValueType.Channel, "");
+		public static LogValue Severity => new(LogValueType.Severity, "");
+		public static LogValue Message => new(LogValueType.Message, "");
+		public static LogValue LineNumber => new(LogValueType.LineNumber, "");
+		public static LogValue ColumnNumber => new(LogValueType.ColumnNumber, "");
+		public static LogValue Symbol => new(LogValueType.Symbol, "");
+		public static LogValue ErrorCode => new(LogValueType.ErrorCode, "");
+		public static LogValue ToolName => new(LogValueType.ToolName, "");
+		public static LogValue ScreenshotTest => new(LogValueType.ScreenshotTest, "");
 #pragma warning restore CS1591
 
 		static readonly char[] s_pathChars = ['/', '\\'];
diff --git a/Engine/Source/Programs/Shared/EpicGames.Core/LogEventParser.cs b/Engine/Source/Programs/Shared/EpicGames.Core/LogEventParser.cs
--- a/Engine/Source/Programs/Shared/EpicGames.Core/LogEventParser.cs
+++ b/Engine/Source/Programs/Shared/EpicGames.Core/LogEventParser.cs
@@ -171,7 +171,7 @@ namespace EpicGames.Core
 		/// <summary>
 		/// Buffer for holding partial line data
 		/// </summary>
-		readonly ByteArrayBuilder _partialLine = new ByteArrayBuilder();
+		readonly ByteArrayBuilder _partialLine = new();
 
 		/// <summary>
 		/// Whether matching is currently enabled
@@ -196,7 +196,7 @@ namespace EpicGames.Core
 		/// <summary>
 		/// Amount of time that the log parser has been processing events
 		/// </summary>
-		readonly Stopwatch _activeTimer = new Stopwatch();
+		readonly Stopwatch _activeTimer = new();
 
 		/// <summary>
 		/// Number of lines parsed in the last interval
diff --git a/Engine/Source/Programs/Shared/EpicGames.Core/LogIndent.cs b/Engine/Source/Programs/Shared/EpicGames.Core/LogIndent.cs
--- a/Engine/Source/Programs/Shared/EpicGames.Core/LogIndent.cs
+++ b/Engine/Source/Programs/Shared/EpicGames.Core/LogIndent.cs
@@ -28,7 +28,7 @@ namespace EpicGames.Core
 		/// <summary>
 		/// The current state value
 		/// </summary>
-		static readonly AsyncLocal<State?> s_currentState = new AsyncLocal<State?>();
+		static readonly AsyncLocal<State?> s_currentState = new();
 
 		/// <summary>
 		/// Gets the current indent string
diff --git a/Engine/Source/Programs/Shared/EpicGames.Core/LogStatusScope.cs b/Engine/Source/Programs/Shared/EpicGames.Core/LogStatusScope.cs
--- a/Engine/Source/Programs/Shared/EpicGames.Core/LogStatusScope.cs
+++ b/Engine/Source/Programs/Shared/EpicGames.Core/LogStatusScope.cs
@@ -61,7 +61,7 @@ namespace EpicGames.Core
 		/// <param name="progress">The progress message</param>
 		public void SetProgress(string progress)
 		{
-			StringBuilder fullMessage = new StringBuilder(_message);
+			StringBuilder fullMessage = new(_message);
 			fullMessage.Append(' ');
 			fullMessage.Append(progress);
 			Log.UpdateStatus(fullMessage.ToString());
@@ -74,7 +74,7 @@ namespace EpicGames.Core
 		/// <param name="args">Arguments for the status message</param>
 		public void SetProgress(string format, params object[] args)
 		{
-			StringBuilder fullMessage = new StringBuilder(_message);
+			StringBuilder fullMessage = new(_message);
 			fullMessage.Append(' ');
 			fullMessage.AppendFormat(format, args);
 			Log.UpdateStatus(fullMessage.ToString());
diff --git a/Engine/Source/Programs/Shared/EpicGames.Core/LogValue.cs b/Engine/Source/Programs/Shared/EpicGames.Core/LogValue.cs
--- a/Engine/Source/Programs/Shared/EpicGames.Core/LogValue.cs
+++ b/Engine/Source/Programs/Shared/EpicGames.Core/LogValue.cs
@@ -13,20 +13,20 @@ namespace EpicGames.Core
 	public static class LogValueType
 	{
 #pragma warning disable CS1591
-		public static readonly Utf8String Asset = new Utf8String("Asset");
-		public static readonly Utf8String SourceFile = new Utf8String("SourceFile");
-		public static readonly Utf8String Object = new Utf8String("Object"); // Arbitrary structured object
-		public static readonly Utf8String Channel = new Utf8String("Channel");
-		public static readonly Utf8String Severity = new Utf8String("Severity");
-		public static readonly Utf8String Message = new Utf8String("Message");
-		public static readonly Utf8String LineNumber = new Utf8String("Line");
-		public static readonly Utf8String ColumnNumber = new Utf8String("Column");
-		public static readonly Utf8String Symbol = new Utf8String("Symbol");
-		public static readonly Utf8String ErrorCode = new Utf8String("ErrorCode");
-		public static readonly Utf8String ToolName = new Utf8String("ToolName");
-		public static readonly Utf8String ScreenshotTest = new Utf8String("ScreenshotTest");
-		public static readonly Utf8String DepotPath = new Utf8String("DepotPath");
-		public static readonly Utf8String Link = new Utf8String("Link");
+		public static readonly Utf8String Asset = new("Asset");
+		public static readonly Utf8String SourceFile = new("SourceFile");
+		public static readonly Utf8String Object = new("Object"); // Arbitrary structured object
+		public static readonly Utf8String Channel = new("Channel");
+		public static readonly Utf8String Severity = new("Severity");
+		public static readonly Utf8String Message = new("Message");
+		public static readonly Utf8String LineNumber = new("Line");
+		public static readonly Utf8String ColumnNumber = new("Column");
+		public static readonly Utf8String Symbol = new("Symbol");
+		public static readonly Utf8String ErrorCode = new("ErrorCode");
+		public static readonly Utf8String ToolName = new("ToolName");
+		public static readonly Utf8String ScreenshotTest = new("ScreenshotTest");
+		public static readonly Utf8String DepotPath = new("DepotPath");
+		public static readonly Utf8String Link = new("Link");
 #pragma warning restore CS1591
 	}
 
diff --git a/Engine/Source/Programs/Shared/EpicGames.Core/LogValueFormatter.cs b/Engine/Source/Programs/Shared/EpicGames.Core/LogValueFormatter.cs
--- a/Engine/Source/Programs/Shared/EpicGames.Core/LogValueFormatter.cs
+++ b/Engine/Source/Programs/Shared/EpicGames.Core/LogValueFormatter.cs
@@ -273,13 +273,13 @@ namespace EpicGames.Core
 			}
 		}
 
-		static readonly StringFormatter s_stringFormatter = new StringFormatter();
+		static readonly StringFormatter s_stringFormatter = new();
 
 		static readonly ConcurrentDictionary<Type, ILogValueFormatter> s_formatters = GetDefaultFormatters();
 
 		static ConcurrentDictionary<Type, ILogValueFormatter> GetDefaultFormatters()
 		{
-			ConcurrentDictionary<Type, ILogValueFormatter> formatters = new ConcurrentDictionary<Type, ILogValueFormatter>();
+			ConcurrentDictionary<Type, ILogValueFormatter> formatters = new();
 			formatters.TryAdd(typeof(bool), new BoolFormatter());
 			formatters.TryAdd(typeof(int), new Int32Formatter());
 			formatters.TryAdd(typeof(uint), new UInt32Formatter());
diff --git a/Engine/Source/Programs/Shared/EpicGames.Core/ManagedProcess.cs b/Engine/Source/Programs/Shared/EpicGames.Core/ManagedProcess.cs
--- a/Engine/Source/Programs/Shared/EpicGames.Core/ManagedProcess.cs
+++ b/Engine/Source/Programs/Shared/EpicGames.Core/ManagedProcess.cs
@@ -148,7 +148,7 @@ namespace EpicGames.Core
 				}
 
 				// Configure the job object to terminate the processes added to it when the handle is closed
-				JOBOBJECT_EXTENDED_LIMIT_INFORMATION limitInformation = new JOBOBJECT_EXTENDED_LIMIT_INFORMATION();
+				JOBOBJECT_EXTENDED_LIMIT_INFORMATION limitInformation = new();
 				limitInformation.BasicLimitInformation.LimitFlags = JOB_OBJECT_LIMIT_BREAKAWAY_OK;
 				if (bKillOnJobClose)
 				{
@@ -478,7 +478,7 @@ namespace EpicGames.Core
 		/// Static lock object. This is used to synchronize the creation of child processes - in particular, the inheritance of stdout/stderr write pipes. If processes
 		/// inherit pipes meant for other processes, they won't be closed until both terminate.
 		/// </summary>
-		static readonly object s_lockObject = new object();
+		static readonly object s_lockObject = new();
 
 		/// <summary>
 		/// Used to perform CPU usage and resource accounting of all children process involved in a single compilation unit.
@@ -652,7 +652,7 @@ namespace EpicGames.Core
 					Marshal.Copy(environmentBytes.ToArray(), 0, environmentBlock, environmentBytes.Count);
 				}
 
-				PROCESS_INFORMATION processInfo = new PROCESS_INFORMATION();
+				PROCESS_INFORMATION processInfo = new();
 				try
 				{
 					// Get the flags to create the new process
@@ -690,7 +690,7 @@ namespace EpicGames.Core
 						try
 						{
 							// Create stdin and stdout pipes for the child process. We'll close the handles for the child process' ends after it's been created.
-							SECURITY_ATTRIBUTES securityAttributes = new SECURITY_ATTRIBUTES();
+							SECURITY_ATTRIBUTES securityAttributes = new();
 							securityAttributes.nLength = Marshal.SizeOf(securityAttributes);
 							securityAttributes.bInheritHandle = 1;
 
@@ -719,7 +719,7 @@ namespace EpicGames.Core
 							}
 
 							// Create the new process as suspended, so we can modify it before it starts executing (and potentially preempting us)
-							STARTUPINFO startupInfo = new STARTUPINFO();
+							STARTUPINFO startupInfo = new();
 							startupInfo.cb = Marshal.SizeOf(typeof(STARTUPINFOEX));
 							startupInfo.hStdInput = stdInRead;
 							startupInfo.hStdOutput = stdOutWrite;
@@ -792,7 +792,7 @@ namespace EpicGames.Core
 
 						uint groupProcessorCount = GetActiveProcessorCount(processorGroup);
 
-						GROUP_AFFINITY groupAffinity = new GROUP_AFFINITY();
+						GROUP_AFFINITY groupAffinity = new();
 						groupAffinity.Mask = ~0ul >> (int)(64 - groupProcessorCount);
 						groupAffinity.Group = processorGroup;
 						if (SetThreadGroupAffinity(processInfo.hThread, groupAffinity, null) == 0)
@@ -1077,7 +1077,7 @@ namespace EpicGames.Core
 		/// <returns></returns>
 		public async Task CopyToAsync(Func<byte[], int, int, CancellationToken, Task> writeOutputAsync, int bufferSize, CancellationToken cancellationToken)
 		{
-			TaskCompletionSource<bool> taskCompletionSource = new TaskCompletionSource<bool>(TaskCreationOptions.RunContinuationsAsynchronously);
+			TaskCompletionSource<bool> taskCompletionSource = new(TaskCreationOptions.RunContinuationsAsynchronously);
 			using (CancellationTokenRegistration registration = cancellationToken.Register(() => taskCompletionSource.SetResult(false)))
 			{
 				byte[] buffer = new byte[bufferSize];
@@ -1301,7 +1301,7 @@ namespace EpicGames.Core
 		{
 			if (_frameworkProcess == null)
 			{
-				using ProcessWaitHandle waitHandle = new ProcessWaitHandle(_processHandle!);
+				using ProcessWaitHandle waitHandle = new(_processHandle!);
 				await waitHandle.WaitOneAsync(cancellationToken);
 			}
 			else
diff --git a/Engine/Source/Programs/Shared/EpicGames.Core/MemoryMappedFileCache.cs b/Engine/Source/Programs/Shared/EpicGames.Core/MemoryMappedFileCache.cs
--- a/Engine/Source/Programs/Shared/EpicGames.Core/MemoryMappedFileCache.cs
+++ b/Engine/Source/Programs/Shared/EpicGames.Core/MemoryMappedFileCache.cs
@@ -148,9 +148,9 @@ namespace EpicGames.Core
 			}
 		}
 
-		readonly object _lockObject = new object();
+		readonly object _lockObject = new();
 		readonly Dictionary<FileReference, MappedFile> _pathToMappedFile = [];
-		readonly LinkedList<MappedFile> _mappedFiles = new LinkedList<MappedFile>();
+		readonly LinkedList<MappedFile> _mappedFiles = new();
 
 		long _mappedSize;
 
diff --git a/Engine/Source/Programs/Shared/EpicGames.Core/MemoryMappedView.cs b/Engine/Source/Programs/Shared/EpicGames.Core/MemoryMappedView.cs
--- a/Engine/Source/Programs/Shared/EpicGames.Core/MemoryMappedView.cs
+++ b/Engine/Source/Programs/Shared/EpicGames.Core/MemoryMappedView.cs
@@ -27,10 +27,10 @@ namespace EpicGames.Core
 			}
 
 			/// <inheritdoc/>
-			public override Span<byte> GetSpan() => new Span<byte>(_pointer, _length);
+			public override Span<byte> GetSpan() => new(_pointer, _length);
 
 			/// <inheritdoc/>
-			public override MemoryHandle Pin(int elementIndex) => new MemoryHandle(_pointer + elementIndex);
+			public override MemoryHandle Pin(int elementIndex) => new(_pointer + elementIndex);
 
 			/// <inheritdoc/>
 			public override void Unpin() { }
@@ -124,7 +124,7 @@ namespace EpicGames.Core
 		/// <returns></returns>
 		public Memory<byte> GetMemory(long offset, int length)
 		{
-			MemoryWrapper wrapper = new MemoryWrapper(_data + offset, length);
+			MemoryWrapper wrapper = new(_data + offset, length);
 			return wrapper.Memory;
 		}
 
diff --git a/Engine/Source/Programs/Shared/EpicGames.Core/MemoryReader.cs b/Engine/Source/Programs/Shared/EpicGames.Core/MemoryReader.cs
--- a/Engine/Source/Programs/Shared/EpicGames.Core/MemoryReader.cs
+++ b/Engine/Source/Programs/Shared/EpicGames.Core/MemoryReader.cs
@@ -206,7 +206,7 @@ namespace EpicGames.Core
 		/// <param name="reader">Reader to deserialize from</param>
 		public static Guid ReadGuidNetOrder(this IMemoryReader reader)
 		{
-			Guid guid = new Guid(reader.GetMemory(16).Slice(0, 16).Span);
+			Guid guid = new(reader.GetMemory(16).Slice(0, 16).Span);
 			reader.Advance(16);
 			return guid;
 		}
@@ -294,7 +294,7 @@ namespace EpicGames.Core
 		{
 			int length = (int)reader.ReadUnsignedVarInt();
 
-			List<T> list = new List<T>(length);
+			List<T> list = new(length);
 			for (int idx = 0; idx < length; idx++)
 			{
 				list.Add(readItem(reader));
@@ -428,7 +428,7 @@ namespace EpicGames.Core
 		/// <param name="comparer">Comparer for the new dictionary</param>
 		public static Dictionary<TKey, TValue> ReadDictionary<TKey, TValue>(this IMemoryReader reader, Func<TKey> readKey, Func<TValue> readValue, IEqualityComparer<TKey>? comparer = null) where TKey : notnull
 		{
-			Dictionary<TKey, TValue> dictionary = new Dictionary<TKey, TValue>(comparer);
+			Dictionary<TKey, TValue> dictionary = new(comparer);
 			ReadDictionary(reader, dictionary, readKey, readValue);
 			return dictionary;
 		}
@@ -442,7 +442,7 @@ namespace EpicGames.Core
 		/// <param name="comparer">Comparer for the new dictionary</param>
 		public static Dictionary<TKey, TValue> ReadDictionary<TKey, TValue>(this IMemoryReader reader, Func<IMemoryReader, TKey> readKey, Func<IMemoryReader, TValue> readValue, IEqualityComparer<TKey>? comparer = null) where TKey : notnull
 		{
-			Dictionary<TKey, TValue> dictionary = new Dictionary<TKey, TValue>(comparer);
+			Dictionary<TKey, TValue> dictionary = new(comparer);
 			ReadDictionary(reader, dictionary, readKey, readValue);
 			return dictionary;
 		}
diff --git a/Engine/Source/Programs/Shared/EpicGames.Core/MemoryWriter.cs b/Engine/Source/Programs/Shared/EpicGames.Core/MemoryWriter.cs
--- a/Engine/Source/Programs/Shared/EpicGames.Core/MemoryWriter.cs
+++ b/Engine/Source/Programs/Shared/EpicGames.Core/MemoryWriter.cs
@@ -413,7 +413,7 @@ namespace EpicGames.Core
 		/// <returns>Sequence of bytes</returns>
 		public ReadOnlySequence<byte> AsSequence()
 		{
-			ReadOnlySequenceBuilder<byte> builder = new ReadOnlySequenceBuilder<byte>();
+			ReadOnlySequenceBuilder<byte> builder = new();
 			AppendTo(builder);
 			return builder.Construct();
 		}
@@ -503,7 +503,7 @@ namespace EpicGames.Core
 		{ }
 
 		/// <inheritdoc/>
-		protected override Chunk CreateChunk(int runningIndex, int size) => new Chunk(runningIndex, new byte[size]);
+		protected override Chunk CreateChunk(int runningIndex, int size) => new(runningIndex, new byte[size]);
 	}
 
 	/// <summary>
diff --git a/Engine/Source/Programs/Shared/EpicGames.Core/MessageTemplate.cs b/Engine/Source/Programs/Shared/EpicGames.Core/MessageTemplate.cs
--- a/Engine/Source/Programs/Shared/EpicGames.Core/MessageTemplate.cs
+++ b/Engine/Source/Programs/Shared/EpicGames.Core/MessageTemplate.cs
@@ -26,7 +26,7 @@ namespace EpicGames.Core
 		/// <returns>The rendered string</returns>
 		public static string Render(string format, IEnumerable<KeyValuePair<string, object?>>? properties)
 		{
-			StringBuilder result = new StringBuilder();
+			StringBuilder result = new();
 			Render(format, properties, result);
 			return result.ToString();
 		}
@@ -127,7 +127,7 @@ namespace EpicGames.Core
 		/// <returns>The escaped string</returns>
 		public static string Escape(string text)
 		{
-			StringBuilder result = new StringBuilder();
+			StringBuilder result = new();
 			Escape(text, result);
 			return result.ToString();
 		}
@@ -156,7 +156,7 @@ namespace EpicGames.Core
 		/// <returns>The unescaped text</returns>
 		public static string Unescape(string text)
 		{
-			StringBuilder result = new StringBuilder();
+			StringBuilder result = new();
 			Unescape(text.AsSpan(), result);
 			return result.ToString();
 		}
diff --git a/Engine/Source/Programs/Shared/EpicGames.Core/ParallelTask.cs b/Engine/Source/Programs/Shared/EpicGames.Core/ParallelTask.cs
--- a/Engine/Source/Programs/Shared/EpicGames.Core/ParallelTask.cs
+++ b/Engine/Source/Programs/Shared/EpicGames.Core/ParallelTask.cs
@@ -36,10 +36,10 @@ namespace EpicGames.Core
 		/// <returns>Async task</returns>
 		public static Task ForEachAsync<T>(IEnumerable<T> collection, Action<T> action)
 		{
-			ExecutionDataflowBlockOptions options = new ExecutionDataflowBlockOptions();
+			ExecutionDataflowBlockOptions options = new();
 			options.MaxDegreeOfParallelism = DataflowBlockOptions.Unbounded;
 
-			ActionBlock<T> actions = new ActionBlock<T>(action, options);
+			ActionBlock<T> actions = new(action, options);
 			foreach (T item in collection)
 			{
 				actions.Post(item);
@@ -60,10 +60,10 @@ namespace EpicGames.Core
 		/// <returns>Async task</returns>
 		public static Task ForEachAsync<T>(IEnumerable<T> collection, Func<T, Task> action, int maxDegreeOfParallelism)
 		{
-			ExecutionDataflowBlockOptions options = new ExecutionDataflowBlockOptions();
+			ExecutionDataflowBlockOptions options = new();
 			options.MaxDegreeOfParallelism = maxDegreeOfParallelism;
 
-			ActionBlock<T> actions = new ActionBlock<T>(action, options);
+			ActionBlock<T> actions = new(action, options);
 			foreach (T item in collection)
 			{
 				actions.Post(item);
diff --git a/Engine/Source/Programs/Shared/EpicGames.Core/ProcessUtils.cs b/Engine/Source/Programs/Shared/EpicGames.Core/ProcessUtils.cs
--- a/Engine/Source/Programs/Shared/EpicGames.Core/ProcessUtils.cs
+++ b/Engine/Source/Programs/Shared/EpicGames.Core/ProcessUtils.cs
@@ -149,10 +149,10 @@ namespace EpicGames.Core
 							if (!handle.IsInvalid)
 							{
 								int characters = 260;
-								StringBuilder buffer = new StringBuilder(characters);
+								StringBuilder buffer = new(characters);
 								if (QueryFullProcessImageName(handle, 0, buffer, ref characters) > 0)
 								{
-									FileReference imageFile = new FileReference(buffer.ToString(0, characters));
+									FileReference imageFile = new(buffer.ToString(0, characters));
 									if (predicate(imageFile))
 									{
 										logger.LogInformation("Terminating {ImageName} ({ProcessId})", imageFile, processId);
@@ -305,7 +305,7 @@ namespace EpicGames.Core
 		{
 			try
 			{
-				ProcessBasicInformation pbi = new ProcessBasicInformation();
+				ProcessBasicInformation pbi = new();
 				if (NtQueryInformationProcess(process.Handle, 0, ref pbi, Marshal.SizeOf(pbi), out _) == 0)
 				{
 					return Process.GetProcessById(pbi._inheritedFromUniqueProcessId.ToInt32());
@@ -404,7 +404,7 @@ namespace EpicGames.Core
 			}
 
 			int capacity = 260;
-			StringBuilder imageNameBuilder = new StringBuilder(capacity);
+			StringBuilder imageNameBuilder = new(capacity);
 			if (QueryFullProcessImageName(hProcess, 0, imageNameBuilder, ref capacity) > 0)
 			{
 				return imageNameBuilder.ToString(0, capacity);
diff --git a/Engine/Source/Programs/Shared/EpicGames.Core/ProgressTextReader.cs b/Engine/Source/Programs/Shared/EpicGames.Core/ProgressTextReader.cs
--- a/Engine/Source/Programs/Shared/EpicGames.Core/ProgressTextReader.cs
+++ b/Engine/Source/Programs/Shared/EpicGames.Core/ProgressTextReader.cs
@@ -12,7 +12,7 @@ namespace EpicGames.Core
 	public class ProgressValue
 	{
 		Tuple<string, float> _state = null!;
-		readonly Stack<Tuple<float, float>> _ranges = new Stack<Tuple<float, float>>();
+		readonly Stack<Tuple<float, float>> _ranges = new();
 
 		/// <summary>
 		/// CTOR
diff --git a/Engine/Source/Programs/Shared/EpicGames.Core/QueryStringBuilder.cs b/Engine/Source/Programs/Shared/EpicGames.Core/QueryStringBuilder.cs
--- a/Engine/Source/Programs/Shared/EpicGames.Core/QueryStringBuilder.cs
+++ b/Engine/Source/Programs/Shared/EpicGames.Core/QueryStringBuilder.cs
@@ -11,7 +11,7 @@ namespace EpicGames.Core
 	/// </summary>
 	public class QueryStringBuilder
 	{
-		readonly StringBuilder _builder = new StringBuilder();
+		readonly StringBuilder _builder = new();
 
 		/// <summary>
 		/// Adds a new argument to the query string
diff --git a/Engine/Source/Programs/Shared/EpicGames.Core/ReadOnlyListExtensions.cs b/Engine/Source/Programs/Shared/EpicGames.Core/ReadOnlyListExtensions.cs
--- a/Engine/Source/Programs/Shared/EpicGames.Core/ReadOnlyListExtensions.cs
+++ b/Engine/Source/Programs/Shared/EpicGames.Core/ReadOnlyListExtensions.cs
@@ -94,7 +94,7 @@ namespace EpicGames.Core
 		/// <returns>New list of items</returns>
 		public static List<TOutput> ConvertAll<TInput, TOutput>(this IReadOnlyList<TInput> input, Func<TInput, TOutput> convert)
 		{
-			List<TOutput> outputs = new List<TOutput>(input.Count);
+			List<TOutput> outputs = new(input.Count);
 			for (int idx = 0; idx < input.Count; idx++)
 			{
 				outputs.Add(convert(input[idx]));
diff --git a/Engine/Source/Programs/Shared/EpicGames.Core/ReadOnlySequenceBuilder.cs b/Engine/Source/Programs/Shared/EpicGames.Core/ReadOnlySequenceBuilder.cs
--- a/Engine/Source/Programs/Shared/EpicGames.Core/ReadOnlySequenceBuilder.cs
+++ b/Engine/Source/Programs/Shared/EpicGames.Core/ReadOnlySequenceBuilder.cs
@@ -58,12 +58,12 @@ namespace EpicGames.Core
 				return ReadOnlySequence<T>.Empty;
 			}
 
-			Segment first = new Segment(0, segments[0]);
+			Segment first = new(0, segments[0]);
 			Segment last = first;
 
 			for (int idx = 1; idx < segments.Count; idx++)
 			{
-				Segment next = new Segment(last.RunningIndex + last.Memory.Length, segments[idx]);
+				Segment next = new(last.RunningIndex + last.Memory.Length, segments[idx]);
 				last.SetNext(next);
 				last = next;
 			}
diff --git a/Engine/Source/Programs/Shared/EpicGames.Core/SemVer.cs b/Engine/Source/Programs/Shared/EpicGames.Core/SemVer.cs
--- a/Engine/Source/Programs/Shared/EpicGames.Core/SemVer.cs
+++ b/Engine/Source/Programs/Shared/EpicGames.Core/SemVer.cs
@@ -30,7 +30,7 @@ namespace EpicGames.Core
 		ReadOnlySpan<char> Span => _text.AsSpan(0, _length);
 
 		// Regex taken from official site.
-		static readonly Regex s_pattern = new Regex(@"^((0|[1-9]\d*)\.(0|[1-9]\d*)\.(0|[1-9]\d*)(?:-((?:0|[1-9]\d*|\d*[a-zA-Z-][0-9a-zA-Z-]*)(?:\.(?:0|[1-9]\d*|\d*[a-zA-Z-][0-9a-zA-Z-]*))*))?)(?:\+([0-9a-zA-Z-]+(?:\.[0-9a-zA-Z-]+)*))?$");
+		static readonly Regex s_pattern = new(@"^((0|[1-9]\d*)\.(0|[1-9]\d*)\.(0|[1-9]\d*)(?:-((?:0|[1-9]\d*|\d*[a-zA-Z-][0-9a-zA-Z-]*)(?:\.(?:0|[1-9]\d*|\d*[a-zA-Z-][0-9a-zA-Z-]*))*))?)(?:\+([0-9a-zA-Z-]+(?:\.[0-9a-zA-Z-]+)*))?$");
 
 		/// <summary>
 		/// Constructor
diff --git a/Engine/Source/Programs/Shared/EpicGames.Core/SingleInstanceMutex.cs b/Engine/Source/Programs/Shared/EpicGames.Core/SingleInstanceMutex.cs
--- a/Engine/Source/Programs/Shared/EpicGames.Core/SingleInstanceMutex.cs
+++ b/Engine/Source/Programs/Shared/EpicGames.Core/SingleInstanceMutex.cs
@@ -15,8 +15,8 @@ namespace EpicGames.Core
 		class SingleInstanceMutexImpl : IDisposable
 		{
 			readonly string? _name;
-			readonly TaskCompletionSource _readyTcs = new TaskCompletionSource(TaskCreationOptions.RunContinuationsAsynchronously);
-			readonly ManualResetEvent _disposing = new ManualResetEvent(false);
+			readonly TaskCompletionSource _readyTcs = new(TaskCreationOptions.RunContinuationsAsynchronously);
+			readonly ManualResetEvent _disposing = new(false);
 			Thread? _thread;
 
 			public SingleInstanceMutexImpl(string name)
@@ -50,7 +50,7 @@ namespace EpicGames.Core
 
 			void ThreadFunc()
 			{
-				using Mutex mutex = new Mutex(false, _name);
+				using Mutex mutex = new(false, _name);
 				try
 				{
 					// Waiting for multiple handles is only supported on Windows
@@ -92,7 +92,7 @@ namespace EpicGames.Core
 		/// <returns>Handle to the mutex. Must be disposed when complete.</returns>
 		public static async Task<IDisposable> AcquireAsync(string name, CancellationToken cancellationToken = default)
 		{
-			SingleInstanceMutexImpl? impl = new SingleInstanceMutexImpl(name);
+			SingleInstanceMutexImpl? impl = new(name);
 			try
 			{
 				await impl.WaitAsync(cancellationToken);
diff --git a/Engine/Source/Programs/Shared/EpicGames.Core/StartupTraceListener.cs b/Engine/Source/Programs/Shared/EpicGames.Core/StartupTraceListener.cs
--- a/Engine/Source/Programs/Shared/EpicGames.Core/StartupTraceListener.cs
+++ b/Engine/Source/Programs/Shared/EpicGames.Core/StartupTraceListener.cs
@@ -10,7 +10,7 @@ namespace EpicGames.Core
 	/// </summary>
 	public class StartupTraceListener : TraceListener
 	{
-		readonly StringBuilder _buffer = new StringBuilder();
+		readonly StringBuilder _buffer = new();
 
 		/// <summary>
 		/// Copy the contents of the buffered output to another trace listener
diff --git a/Engine/Source/Programs/Shared/EpicGames.Core/StreamUtils.cs b/Engine/Source/Programs/Shared/EpicGames.Core/StreamUtils.cs
--- a/Engine/Source/Programs/Shared/EpicGames.Core/StreamUtils.cs
+++ b/Engine/Source/Programs/Shared/EpicGames.Core/StreamUtils.cs
@@ -79,7 +79,7 @@ namespace EpicGames.Core
 		/// <returns>Contents of the stream</returns>
 		public static async Task<byte[]> ReadAllBytesAsync(this Stream stream, CancellationToken cancellationToken = default)
 		{
-			using MemoryStream memoryStream = new MemoryStream();
+			using MemoryStream memoryStream = new();
 			await stream.CopyToAsync(memoryStream, cancellationToken);
 			return memoryStream.ToArray();
 		}
diff --git a/Engine/Source/Programs/Shared/EpicGames.Core/StringUtils.cs b/Engine/Source/Programs/Shared/EpicGames.Core/StringUtils.cs
--- a/Engine/Source/Programs/Shared/EpicGames.Core/StringUtils.cs
+++ b/Engine/Source/Programs/Shared/EpicGames.Core/StringUtils.cs
@@ -140,7 +140,7 @@ namespace EpicGames.Core
 		/// <param name="maxWidth">The maximum (non negative) length of the returned sentences</param>
 		public static IEnumerable<string> WordWrap(string text, int initialIndent, int hangingIndent, int maxWidth)
 		{
-			StringBuilder builder = new StringBuilder();
+			StringBuilder builder = new();
 
 			int minIdx = 0;
 			for (int lineIdx = 0; minIdx < text.Length; lineIdx++)
@@ -235,7 +235,7 @@ namespace EpicGames.Core
 		/// <returns>Formatted list of strings</returns>
 		public static string FormatList(IReadOnlyList<string> arguments, string conjunction = "and")
 		{
-			StringBuilder result = new StringBuilder();
+			StringBuilder result = new();
 			if (arguments.Count > 0)
 			{
 				result.Append(arguments[0]);
@@ -712,7 +712,7 @@ namespace EpicGames.Core
 			{
 				int lastIndex = 0;
 
-				StringBuilder result = new StringBuilder();
+				StringBuilder result = new();
 				for (; ; )
 				{
 					result.Append(line, lastIndex, index - lastIndex);
diff --git a/Engine/Source/Programs/Shared/EpicGames.Core/StringView.cs b/Engine/Source/Programs/Shared/EpicGames.Core/StringView.cs
--- a/Engine/Source/Programs/Shared/EpicGames.Core/StringView.cs
+++ b/Engine/Source/Programs/Shared/EpicGames.Core/StringView.cs
@@ -177,7 +177,7 @@ namespace EpicGames.Core
 		/// Implicit conversion operator from a regular string
 		/// </summary>
 		/// <param name="text">The string to construct from</param>
-		public static implicit operator StringView(string text) => new StringView(text);
+		public static implicit operator StringView(string text) => new(text);
 	}
 
 	/// <summary>
diff --git a/Engine/Source/Programs/Shared/EpicGames.Core/Telemetry/DataRouterTelemetryService.cs b/Engine/Source/Programs/Shared/EpicGames.Core/Telemetry/DataRouterTelemetryService.cs
--- a/Engine/Source/Programs/Shared/EpicGames.Core/Telemetry/DataRouterTelemetryService.cs
+++ b/Engine/Source/Programs/Shared/EpicGames.Core/Telemetry/DataRouterTelemetryService.cs
@@ -361,7 +361,7 @@ namespace EpicGames.Core.Telemetry
 		/// </summary>
 		public Uri BaseAddress { get; set; }
 
-		private readonly Timer _autoFlushTimer = new Timer();
+		private readonly Timer _autoFlushTimer = new();
 
 		/// <summary>
 		/// Length of an empty data router post when serialized to json
@@ -372,7 +372,7 @@ namespace EpicGames.Core.Telemetry
 			{
 				if (s_dataRouterPostEmptyJsonLength < 0)
 				{
-					DataRouterPost emptyPost = new DataRouterPost();
+					DataRouterPost emptyPost = new();
 					s_dataRouterPostEmptyJsonLength = JsonSerializer.Serialize(emptyPost).Length;
 				}
 				return s_dataRouterPostEmptyJsonLength;
@@ -390,7 +390,7 @@ namespace EpicGames.Core.Telemetry
 			{
 				string productName = String.IsNullOrEmpty(entryAssembly.GetName().Name) ? GetType().Name : entryAssembly.GetName().Name!;
 				string productVersion = entryAssembly.GetName().Version == null ? "1.0.0.0" : entryAssembly.GetName().Version!.ToString();
-				ProductInfoHeaderValue userAgent = new ProductInfoHeaderValue(productName, productVersion);
+				ProductInfoHeaderValue userAgent = new(productName, productVersion);
 				HttpClient.DefaultRequestHeaders.UserAgent.Add(userAgent);
 			}
 			else
@@ -469,7 +469,7 @@ namespace EpicGames.Core.Telemetry
 		{
 			outPostsList = [];
 
-			SerializedDataRouterPost currentBatch = new SerializedDataRouterPost();
+			SerializedDataRouterPost currentBatch = new();
 
 			int currentBatchSize = DataRouterPostEmptyJsonLength;
 
@@ -516,7 +516,7 @@ namespace EpicGames.Core.Telemetry
 		/// <returns>Http responce</returns>
 		private async Task<HttpResponseMessage> PostRequestAsync(string postDataKey, string jsonBody)
 		{
-			using HttpRequestMessage postRequest = new HttpRequestMessage(HttpMethod.Post, $"{BaseAddress}?{postDataKey}");
+			using HttpRequestMessage postRequest = new(HttpMethod.Post, $"{BaseAddress}?{postDataKey}");
 			postRequest.Content = new StringContent(jsonBody, System.Text.Encoding.UTF8, "application/json");
 			return await HttpClient.SendAsync(postRequest);
 		}
@@ -527,7 +527,7 @@ namespace EpicGames.Core.Telemetry
 		/// <returns>Http responce</returns>
 		private async Task<HttpResponseMessage> HeadRequestAsync()
 		{
-			using HttpRequestMessage headRequest = new HttpRequestMessage(HttpMethod.Head, $"{BaseAddress}");
+			using HttpRequestMessage headRequest = new(HttpMethod.Head, $"{BaseAddress}");
 			return await Task.Run(async () => await HttpClient.SendAsync(headRequest));
 		}
 
@@ -537,7 +537,7 @@ namespace EpicGames.Core.Telemetry
 		/// <returns></returns>
 		public virtual async Task FlushEventsAsync()
 		{
-			ConcurrentDictionary<string, DataRouterPost> posts = new ConcurrentDictionary<string, DataRouterPost>();
+			ConcurrentDictionary<string, DataRouterPost> posts = new();
 			DataRouterEvent? eventData;
 			while (EventsQueue.TryDequeue(out eventData))
 			{
@@ -563,7 +563,7 @@ namespace EpicGames.Core.Telemetry
 				return;
 			}
 
-			JsonSerializerOptions serializeOptions = new JsonSerializerOptions
+			JsonSerializerOptions serializeOptions = new()
 			{
 				WriteIndented = false,
 				Converters =
diff --git a/Engine/Source/Programs/Shared/EpicGames.Core/Telemetry/DefaultTelemetryHttpClient.cs b/Engine/Source/Programs/Shared/EpicGames.Core/Telemetry/DefaultTelemetryHttpClient.cs
--- a/Engine/Source/Programs/Shared/EpicGames.Core/Telemetry/DefaultTelemetryHttpClient.cs
+++ b/Engine/Source/Programs/Shared/EpicGames.Core/Telemetry/DefaultTelemetryHttpClient.cs
@@ -39,7 +39,7 @@ namespace EpicGames.Core.Telemetry
 			HttpMessageHandler socketHandler = HttpMessageHandlers.CreateDefaultHandler();
 
 #pragma warning disable CA2000 // Dispose objects before losing scope
-			PolicyHttpMessageHandler pollyHandler = new PolicyHttpMessageHandler(GetRetryPolicy())
+			PolicyHttpMessageHandler pollyHandler = new(GetRetryPolicy())
 			{
 				InnerHandler = socketHandler,
 			};
diff --git a/Engine/Source/Programs/Shared/EpicGames.Core/TraceSpan.cs b/Engine/Source/Programs/Shared/EpicGames.Core/TraceSpan.cs
--- a/Engine/Source/Programs/Shared/EpicGames.Core/TraceSpan.cs
+++ b/Engine/Source/Programs/Shared/EpicGames.Core/TraceSpan.cs
@@ -104,7 +104,7 @@ namespace EpicGames.Core
 		/// </summary>
 		public ITraceSpan Create(string name, string? resource = null, string? service = null)
 		{
-			TraceSpanImpl span = new TraceSpanImpl(name, resource, service);
+			TraceSpanImpl span = new(name, resource, service);
 			_spans.Add(span);
 			return span;
 		}
diff --git a/Engine/Source/Programs/Shared/EpicGames.Core/UnrealEngineTypes.cs b/Engine/Source/Programs/Shared/EpicGames.Core/UnrealEngineTypes.cs
--- a/Engine/Source/Programs/Shared/EpicGames.Core/UnrealEngineTypes.cs
+++ b/Engine/Source/Programs/Shared/EpicGames.Core/UnrealEngineTypes.cs
@@ -1323,7 +1323,7 @@ namespace EpicGames.Core
 	/// </summary>
 	public static class EPropertyFlagsExtensions
 	{
-		private static readonly Lazy<List<string>> BitNames = new Lazy<List<string>>(() =>
+		private static readonly Lazy<List<string>> BitNames = new(() =>
 		{
 			List<string> Out = [];
 			FieldInfo[] Fields = typeof(EPropertyFlags).GetFields();
diff --git a/Engine/Source/Programs/Shared/EpicGames.Core/Utf8String.cs b/Engine/Source/Programs/Shared/EpicGames.Core/Utf8String.cs
--- a/Engine/Source/Programs/Shared/EpicGames.Core/Utf8String.cs
+++ b/Engine/Source/Programs/Shared/EpicGames.Core/Utf8String.cs
@@ -14,7 +14,7 @@ namespace EpicGames.Core
 		/// <summary>
 		/// An empty string
 		/// </summary>
-		public static readonly Utf8String Empty = new Utf8String();
+		public static readonly Utf8String Empty = new();
 
 		/// <summary>
 		/// The data represented by this string
@@ -424,7 +424,7 @@ namespace EpicGames.Core
 			/// <inheritdoc/>
 			public override int GetHashCode(ReadOnlySpan<byte> str)
 			{
-				HashCode hashCode = new HashCode();
+				HashCode hashCode = new();
 				for (int idx = 0; idx < str.Length; idx++)
 				{
 					hashCode.Add(ToUpper(str[idx]));
@@ -538,7 +538,7 @@ namespace EpicGames.Core
 				int length = memory.Span.IndexOf((byte)0);
 				if (length != -1)
 				{
-					Utf8String str = new Utf8String(memory.Slice(0, length));
+					Utf8String str = new(memory.Slice(0, length));
 					reader.Advance(length + 1);
 					return str;
 				}
diff --git a/Engine/Source/Programs/Shared/EpicGames.Core/Utf8StringBuilder.cs b/Engine/Source/Programs/Shared/EpicGames.Core/Utf8StringBuilder.cs
--- a/Engine/Source/Programs/Shared/EpicGames.Core/Utf8StringBuilder.cs
+++ b/Engine/Source/Programs/Shared/EpicGames.Core/Utf8StringBuilder.cs
@@ -43,7 +43,7 @@ namespace EpicGames.Core
 		/// <summary>
 		/// Accessor for the written data
 		/// </summary>
-		public Utf8String WrittenString => new Utf8String(_writer.WrittenMemory);
+		public Utf8String WrittenString => new(_writer.WrittenMemory);
 
 		/// <summary>
 		/// Constructor
diff --git a/Engine/Source/Programs/Shared/EpicGames.Core/XmlDocReader.cs b/Engine/Source/Programs/Shared/EpicGames.Core/XmlDocReader.cs
--- a/Engine/Source/Programs/Shared/EpicGames.Core/XmlDocReader.cs
+++ b/Engine/Source/Programs/Shared/EpicGames.Core/XmlDocReader.cs
@@ -15,7 +15,7 @@ namespace EpicGames.Core
 	/// </summary>
 	public class XmlDocReader
 	{
-		readonly ConcurrentDictionary<Assembly, XmlDocument?> _cachedDocumentation = new ConcurrentDictionary<Assembly, XmlDocument?>();
+		readonly ConcurrentDictionary<Assembly, XmlDocument?> _cachedDocumentation = new();
 
 		/// <summary>
 		/// Adds a parsed file as the documentation for a particular assembly
diff --git a/Engine/Source/Programs/Shared/EpicGames.Core/ZipArchiveExtensions.cs b/Engine/Source/Programs/Shared/EpicGames.Core/ZipArchiveExtensions.cs
--- a/Engine/Source/Programs/Shared/EpicGames.Core/ZipArchiveExtensions.cs
+++ b/Engine/Source/Programs/Shared/EpicGames.Core/ZipArchiveExtensions.cs
@@ -53,7 +53,7 @@ namespace EpicGames.Core
 		{
 			// This seems to be consistently at least 35% faster than ZipFileExtensions.ExtractToFile(Entry, TargetFileName, Overwrite) when the
 			// archive contains many small files.
-			using (FileStream outputStream = new FileStream(targetFileName, overwrite ? FileMode.Create : FileMode.CreateNew, FileAccess.Write))
+			using (FileStream outputStream = new(targetFileName, overwrite ? FileMode.Create : FileMode.CreateNew, FileAccess.Write))
 			{
 				using (Stream inputStream = entry.Open())
 				{
diff --git a/Engine/Source/Programs/Shared/EpicGames.Horde.Tests/Commits/CommitTagTests.cs b/Engine/Source/Programs/Shared/EpicGames.Horde.Tests/Commits/CommitTagTests.cs
--- a/Engine/Source/Programs/Shared/EpicGames.Horde.Tests/Commits/CommitTagTests.cs
+++ b/Engine/Source/Programs/Shared/EpicGames.Horde.Tests/Commits/CommitTagTests.cs
@@ -210,7 +210,7 @@ namespace EpicGames.Horde.Tests.Commits
 			// Target size in bytes
 			int targetSize = 256 * 1024; // 256kb KB - double the current max size for CachedCommitDoc.
 
-			System.Text.StringBuilder sb = new System.Text.StringBuilder(targetSize);
+			System.Text.StringBuilder sb = new(targetSize);
 
 			while (sb.Length < targetSize)
 			{
@@ -237,7 +237,7 @@ namespace EpicGames.Horde.Tests.Commits
 			const int Iterations = 100000;
 
 			CommitDescriptionTagContainer? last = null;
-			Random rng = new Random();
+			Random rng = new();
 
 			for (int i = 0; i < Iterations; i++)
 			{
diff --git a/Engine/Source/Programs/Shared/EpicGames.Horde.Tests/Compute/BufferTests.cs b/Engine/Source/Programs/Shared/EpicGames.Horde.Tests/Compute/BufferTests.cs
--- a/Engine/Source/Programs/Shared/EpicGames.Horde.Tests/Compute/BufferTests.cs
+++ b/Engine/Source/Programs/Shared/EpicGames.Horde.Tests/Compute/BufferTests.cs
@@ -23,7 +23,7 @@ namespace EpicGames.Horde.Tests.Compute
 		[TestMethod]
 		public async Task TestSimpleBufferAsync()
 		{
-			using PooledBuffer buffer = new PooledBuffer(2, 1024);
+			using PooledBuffer buffer = new(2, 1024);
 			using ComputeBufferWriter bufferWriter = buffer.CreateWriter();
 			bufferWriter.AdvanceWritePosition(10);
 
@@ -36,7 +36,7 @@ namespace EpicGames.Horde.Tests.Compute
 		[TestMethod]
 		public async Task TestOverflowAsync()
 		{
-			using PooledBuffer buffer = new PooledBuffer(2, 20);
+			using PooledBuffer buffer = new(2, 20);
 			using ComputeBufferReader bufferReader = buffer.CreateReader();
 			using ComputeBufferWriter bufferWriter = buffer.CreateWriter();
 
@@ -101,8 +101,8 @@ namespace EpicGames.Horde.Tests.Compute
 
 			const int Length = 8000;
 
-			Pipe sourceToTargetPipe = new Pipe();
-			Pipe targetToSourcePipe = new Pipe();
+			Pipe sourceToTargetPipe = new();
+			Pipe targetToSourcePipe = new();
 
 			await using PipeTransport producerTransport = new(targetToSourcePipe.Reader, sourceToTargetPipe.Writer);
 			await using PipeTransport consumerTransport = new(sourceToTargetPipe.Reader, targetToSourcePipe.Writer);
@@ -165,21 +165,21 @@ namespace EpicGames.Horde.Tests.Compute
 		[TestMethod]
 		public async Task TestSendBufferCompleteAsync()
 		{
-			Pipe recvPipe = new Pipe();
-			Pipe sendPipe = new Pipe();
+			Pipe recvPipe = new();
+			Pipe sendPipe = new();
 			await using PipeTransport localTransport = new(sendPipe.Reader, recvPipe.Writer);
 			await using PipeTransport remoteTransport = new(recvPipe.Reader, sendPipe.Writer);
 			await using RemoteComputeSocket localSocket = new(localTransport, ComputeProtocol.Latest, NullLogger.Instance);
 			await using RemoteComputeSocket remoteSocket = new(remoteTransport, ComputeProtocol.Latest, NullLogger.Instance);
 
-			using (PooledBuffer remoteBuffer = new PooledBuffer(1024))
+			using (PooledBuffer remoteBuffer = new(1024))
 			{
 				remoteSocket.AttachRecvBuffer(1, remoteBuffer);
 
 				using ComputeBufferReader reader = remoteBuffer.CreateReader();
 
 				// Disposing of the buffer should mark the channel as complete
-				using (PooledBuffer localBuffer = new PooledBuffer(1024))
+				using (PooledBuffer localBuffer = new(1024))
 				{
 					localSocket.AttachSendBuffer(1, localBuffer);
 
diff --git a/Engine/Source/Programs/Shared/EpicGames.Horde.Tests/Compute/ComputeSocketTests.cs b/Engine/Source/Programs/Shared/EpicGames.Horde.Tests/Compute/ComputeSocketTests.cs
--- a/Engine/Source/Programs/Shared/EpicGames.Horde.Tests/Compute/ComputeSocketTests.cs
+++ b/Engine/Source/Programs/Shared/EpicGames.Horde.Tests/Compute/ComputeSocketTests.cs
@@ -74,8 +74,8 @@ namespace EpicGames.Horde.Tests.Compute
 		[TestMethod]
 		public async Task TestAgentMessageLoopPipeAsync()
 		{
-			Pipe recvPipe = new Pipe();
-			Pipe sendPipe = new Pipe();
+			Pipe recvPipe = new();
+			Pipe sendPipe = new();
 			await using PipeTransport localTransport = new(sendPipe.Reader, recvPipe.Writer);
 			await using PipeTransport agentTransport = new(recvPipe.Reader, sendPipe.Writer);
 			await using RemoteComputeSocket localSocket = new(localTransport, ComputeProtocol.Latest, new TestLogger("local"));
@@ -155,7 +155,7 @@ namespace EpicGames.Horde.Tests.Compute
 
 		static async Task RunAgentTestsAsync(RemoteComputeSocket localSocket, RemoteComputeSocket agentSocket, CancellationToken cancellationToken = default)
 		{
-			DirectoryReference tempDir = new DirectoryReference("test-temp-" + DateTime.UtcNow.Ticks);
+			DirectoryReference tempDir = new("test-temp-" + DateTime.UtcNow.Ticks);
 			await using (BackgroundTask agentTask = BackgroundTask.StartNew(ctx => RunAgentAsync(agentSocket, tempDir, ctx)))
 			{
 				const int PrimaryChannelId = 0;
@@ -206,15 +206,15 @@ namespace EpicGames.Horde.Tests.Compute
 
 						byte[] data = Encoding.UTF8.GetBytes("Hello world");
 
-						using ChunkedDataWriter writer = new ChunkedDataWriter(blobWriter, new ChunkingOptions());
+						using ChunkedDataWriter writer = new(blobWriter, new ChunkingOptions());
 						ChunkedData chunkedData = await writer.CreateAsync(data, cancellationToken);
 
-						DirectoryNode directory = new DirectoryNode();
+						DirectoryNode directory = new();
 						directory.AddFile("hello.txt", FileEntryFlags.None, chunkedData);
 
 						IHashedBlobRef<DirectoryNode> directoryRef = await blobWriter.WriteBlobAsync(directory, cancellationToken: cancellationToken);
 
-						DirectoryNode root = new DirectoryNode();
+						DirectoryNode root = new();
 						root.AddDirectory(new DirectoryEntry("subdir", directory.Length, directoryRef));
 
 						IHashedBlobRef<DirectoryNode> handle = await blobWriter.WriteBlobAsync(root, cancellationToken: cancellationToken);
@@ -277,7 +277,7 @@ namespace EpicGames.Horde.Tests.Compute
 		{
 			try
 			{
-				AgentMessageHandler handler = new AgentMessageHandler(tempDir, null, true, null, null, NullLogger.Instance);
+				AgentMessageHandler handler = new(tempDir, null, true, null, null, NullLogger.Instance);
 				await handler.RunAsync(socket, cancellationToken);
 			}
 			catch (Exception e)
diff --git a/Engine/Source/Programs/Shared/EpicGames.Horde.Tests/Compute/TransportTests.cs b/Engine/Source/Programs/Shared/EpicGames.Horde.Tests/Compute/TransportTests.cs
--- a/Engine/Source/Programs/Shared/EpicGames.Horde.Tests/Compute/TransportTests.cs
+++ b/Engine/Source/Programs/Shared/EpicGames.Horde.Tests/Compute/TransportTests.cs
@@ -23,16 +23,16 @@ namespace EpicGames.Horde.Tests.Compute
 
 			byte[] payload;
 			{
-				using MemoryStream memoryStream = new MemoryStream();
-				await using StreamTransport streamTransport = new StreamTransport(memoryStream);
+				using MemoryStream memoryStream = new();
+				await using StreamTransport streamTransport = new(memoryStream);
 				await streamTransport.SendAsync(input, CancellationToken.None);
 				payload = memoryStream.ToArray();
 			}
 
 			byte[] output = new byte[input.Length];
 			{
-				using MemoryStream memoryStream = new MemoryStream(payload);
-				await using StreamTransport streamTransport = new StreamTransport(memoryStream);
+				using MemoryStream memoryStream = new(payload);
+				await using StreamTransport streamTransport = new(memoryStream);
 				await CopyToAsync(streamTransport, output);
 			}
 
@@ -56,10 +56,10 @@ namespace EpicGames.Horde.Tests.Compute
 			byte[] encrypted;
 
 			{
-				using MemoryStream memoryStream = new MemoryStream();
-				await using (StreamTransport streamTransport = new StreamTransport(memoryStream))
+				using MemoryStream memoryStream = new();
+				await using (StreamTransport streamTransport = new(memoryStream))
 				{
-					await using (AesTransport aesTransport = new AesTransport(streamTransport, key))
+					await using (AesTransport aesTransport = new(streamTransport, key))
 					{
 						await aesTransport.SendAsync(input, CancellationToken.None);
 					}
@@ -69,10 +69,10 @@ namespace EpicGames.Horde.Tests.Compute
 
 			byte[] output = new byte[input.Length];
 			{
-				using MemoryStream memoryStream = new MemoryStream(encrypted);
-				await using (StreamTransport streamTransport = new StreamTransport(memoryStream))
+				using MemoryStream memoryStream = new(encrypted);
+				await using (StreamTransport streamTransport = new(memoryStream))
 				{
-					await using (AesTransport aesTransport = new AesTransport(streamTransport, key))
+					await using (AesTransport aesTransport = new(streamTransport, key))
 					{
 						for (int offset = 0; offset < output.Length;)
 						{
diff --git a/Engine/Source/Programs/Shared/EpicGames.Horde.Tests/HordeHttpTests.cs b/Engine/Source/Programs/Shared/EpicGames.Horde.Tests/HordeHttpTests.cs
--- a/Engine/Source/Programs/Shared/EpicGames.Horde.Tests/HordeHttpTests.cs
+++ b/Engine/Source/Programs/Shared/EpicGames.Horde.Tests/HordeHttpTests.cs
@@ -132,8 +132,8 @@ public class HordeHttpAuthHandlerTests
 
 		ILogger<OidcAuthState> logger = loggerFactory.CreateLogger<OidcAuthState>();
 		IOidcAuthStateConfig oidcAuthStateConfig = new OidcAuthStateConfig("Horde");
-		OidcAuthState state = new OidcAuthState(oidcTokenManager, new NullStaticAccessTokenProvider(), _clock, oidcAuthStateConfig, logger);
-		HttpOidcAuthHandler authHandler = new HttpOidcAuthHandler(server, state, new HttpOidcAuthHandlerConfig("Horde-AllowInteractiveLogin"));
+		OidcAuthState state = new(oidcTokenManager, new NullStaticAccessTokenProvider(), _clock, oidcAuthStateConfig, logger);
+		HttpOidcAuthHandler authHandler = new(server, state, new HttpOidcAuthHandlerConfig("Horde-AllowInteractiveLogin"));
 
 		authHandler.InnerHandler = server;
 		HttpClient client = new(authHandler);
diff --git a/Engine/Source/Programs/Shared/EpicGames.Horde.Tests/Logs/LogTests.cs b/Engine/Source/Programs/Shared/EpicGames.Horde.Tests/Logs/LogTests.cs
--- a/Engine/Source/Programs/Shared/EpicGames.Horde.Tests/Logs/LogTests.cs
+++ b/Engine/Source/Programs/Shared/EpicGames.Horde.Tests/Logs/LogTests.cs
@@ -23,7 +23,7 @@ namespace EpicGames.Horde.Tests
 		[TestMethod]
 		public void NgramTests()
 		{
-			NgramSetBuilder builder = new NgramSetBuilder();
+			NgramSetBuilder builder = new();
 
 			ulong[] values = [1, 2, 3, 5, 8, 13, 21, 34, 55, 89, 144, 233, 377, 610, 987, 1597, 2584, 4181, 6765, 10946, 0xfedcba9876543210UL];
 			foreach (ulong value in values)
@@ -43,7 +43,7 @@ namespace EpicGames.Horde.Tests
 			KeyValueStorageNamespace store = KeyValueStorageNamespace.CreateInMemory();
 
 			// Write the test data to the log file in blocks
-			LogBuilder builder = new LogBuilder(LogFormat.Text, 1, 1, NullLogger.Instance);
+			LogBuilder builder = new(LogFormat.Text, 1, 1, NullLogger.Instance);
 
 			int readOffset = 0;
 			int readLineIndex = 0;
@@ -112,7 +112,7 @@ namespace EpicGames.Horde.Tests
 				"wx\n"
 			];
 
-			LogBuilder builder = new LogBuilder(LogFormat.Text, 1, 1, NullLogger.Instance);
+			LogBuilder builder = new(LogFormat.Text, 1, 1, NullLogger.Instance);
 			for (int lineIdx = 0; lineIdx < lines.Length; lineIdx++)
 			{
 				builder.WriteData(Encoding.UTF8.GetBytes(lines[lineIdx]));
@@ -136,7 +136,7 @@ namespace EpicGames.Horde.Tests
 					{
 						string str = lines[lineIdx].Substring(strOfs, strLen);
 
-						SearchStats stats = new SearchStats();
+						SearchStats stats = new();
 						List<int> results = await index.SearchAsync(0, new SearchTerm(str), stats).ToListAsync();
 						Assert.AreEqual(1, results.Count);
 						Assert.AreEqual(lineIdx, results[0]);
@@ -204,7 +204,7 @@ namespace EpicGames.Horde.Tests
 
 		static async Task SearchLogDataTestAsync(LogIndexNode index, string text, int firstLine, int count, int[] expectedLines)
 		{
-			SearchStats stats = new SearchStats();
+			SearchStats stats = new();
 			List<int> lines = await index.SearchAsync(firstLine, new SearchTerm(text), stats, CancellationToken.None).Take(count).ToListAsync();
 			Assert.IsTrue(lines.SequenceEqual(expectedLines));
 		}
diff --git a/Engine/Source/Programs/Shared/EpicGames.Horde.Tests/Logs/ServerLogPacketBuilderTests.cs b/Engine/Source/Programs/Shared/EpicGames.Horde.Tests/Logs/ServerLogPacketBuilderTests.cs
--- a/Engine/Source/Programs/Shared/EpicGames.Horde.Tests/Logs/ServerLogPacketBuilderTests.cs
+++ b/Engine/Source/Programs/Shared/EpicGames.Horde.Tests/Logs/ServerLogPacketBuilderTests.cs
@@ -46,7 +46,7 @@ namespace EpicGames.Horde.Tests
 		[TestMethod]
 		public void JsonLogEventTest()
 		{
-			JsonLoggerImpl impl = new JsonLoggerImpl();
+			JsonLoggerImpl impl = new();
 			impl.LogInformation("Hello {Text}", "world");
 			impl.LogInformation("Hello {Text}", "world");
 
@@ -62,7 +62,7 @@ namespace EpicGames.Horde.Tests
 			JsonLogEvent jsonLogEvent;
 			Assert.IsTrue(JsonLogEvent.TryParse(Encoding.UTF8.GetBytes(message), out jsonLogEvent));
 
-			ServerLogPacketBuilder writer = new ServerLogPacketBuilder();
+			ServerLogPacketBuilder writer = new();
 			int count = writer.SanitizeAndWriteEvent(jsonLogEvent);
 
 			string[] result = Encoding.UTF8.GetString(writer.CreatePacket().Item1.Span).Split('\n', StringSplitOptions.RemoveEmptyEntries);
@@ -94,7 +94,7 @@ namespace EpicGames.Horde.Tests
 			data[idx] = 0xef;
 			data[idx + 1] = 0xbb;
 
-			ServerLogPacketBuilder writer = new ServerLogPacketBuilder();
+			ServerLogPacketBuilder writer = new();
 			int count = writer.SanitizeAndWriteEvent(JsonLogEvent.Parse(data));
 			Assert.AreEqual(1, count);
 
@@ -278,13 +278,13 @@ namespace EpicGames.Horde.Tests
 		[TestMethod]
 		public async Task StorageLoggerTestAsync()
 		{
-			await using BundleCache cache = new BundleCache();
-			await using FakeHordeClient hordeClient = new FakeHordeClient();
+			await using BundleCache cache = new();
+			await using FakeHordeClient hordeClient = new();
 
 			LogId logId = default;
 
 			const int Count = 20000;
-			await using (ServerLogger logger = new ServerLogger(hordeClient, logId, LogLevel.Information, NullLogger.Instance))
+			await using (ServerLogger logger = new(hordeClient, logId, LogLevel.Information, NullLogger.Instance))
 			{
 				for (int idx = 0; idx < Count; idx++)
 				{
@@ -335,22 +335,22 @@ namespace EpicGames.Horde.Tests
 		[TestMethod]
 		public void LongMessageTest()
 		{
-			DateTime time = new DateTime(2023, 1, 1, 0, 0, 0, 0, DateTimeKind.Utc);
+			DateTime time = new(2023, 1, 1, 0, 0, 0, 0, DateTimeKind.Utc);
 
 			string longLineA = String.Join("\n", Enumerable.Range(0, 5).Select(x => $"A{x}"));
 			string longLineB = String.Join("\n", Enumerable.Range(0, 5).Select(x => $"B{x}"));
 			string longLineC = String.Join("\n", Enumerable.Range(0, 5).Select(x => $"C{x}"));
 
-			Dictionary<string, object?> properties = new Dictionary<string, object?>
+			Dictionary<string, object?> properties = new()
 			{
 				["LongLineA"] = longLineA,
 				["LongLineB"] = longLineB,
 				["LongLineC"] = longLineC,
 			};
 
-			LogEvent baseEvent = new LogEvent(time, LogLevel.Information, default, $"start {longLineA} x {longLineB} y\nz {longLineC} end", "start {LongLineA} x {LongLineB} y\nz {LongLineC} end", properties, null);
+			LogEvent baseEvent = new(time, LogLevel.Information, default, $"start {longLineA} x {longLineB} y\nz {longLineC} end", "start {LongLineA} x {LongLineB} y\nz {LongLineC} end", properties, null);
 
-			ServerLogPacketBuilder writer = new ServerLogPacketBuilder();
+			ServerLogPacketBuilder writer = new();
 			writer.SanitizeAndWriteEvent(new JsonLogEvent(baseEvent));
 
 			string[] output = Encoding.UTF8.GetString(writer.CreatePacket().Item1.Span).Split('\n');
@@ -382,11 +382,11 @@ namespace EpicGames.Horde.Tests
 		[TestMethod]
 		public void MaxPacketLengthTest()
 		{
-			DateTime time = new DateTime(2023, 1, 1, 0, 0, 0, 0, DateTimeKind.Utc);
+			DateTime time = new(2023, 1, 1, 0, 0, 0, 0, DateTimeKind.Utc);
 
-			LogEvent baseEvent = new LogEvent(time, LogLevel.Information, default, "abcdefghijklmnopqrstuvwxyzABCDEFGHIJKLMNOPQRSTUVWXYZ", null, null, null);
+			LogEvent baseEvent = new(time, LogLevel.Information, default, "abcdefghijklmnopqrstuvwxyzABCDEFGHIJKLMNOPQRSTUVWXYZ", null, null, null);
 
-			ServerLogPacketBuilder writer = new ServerLogPacketBuilder(maxPacketLength: 85);
+			ServerLogPacketBuilder writer = new(maxPacketLength: 85);
 			writer.SanitizeAndWriteEvent(new JsonLogEvent(baseEvent));
 			writer.SanitizeAndWriteEvent(new JsonLogEvent(baseEvent));
 
@@ -401,11 +401,11 @@ namespace EpicGames.Horde.Tests
 		[TestMethod]
 		public void MaxLineLengthTest()
 		{
-			DateTime time = new DateTime(2023, 1, 1, 0, 0, 0, 0, DateTimeKind.Utc);
+			DateTime time = new(2023, 1, 1, 0, 0, 0, 0, DateTimeKind.Utc);
 
-			LogEvent baseEvent = new LogEvent(time, LogLevel.Information, default, "abcdefghijklmnopqrstuvwxyzABCDEFGHIJKLMNOPQRSTUVWXYZ", null, null, null);
+			LogEvent baseEvent = new(time, LogLevel.Information, default, "abcdefghijklmnopqrstuvwxyzABCDEFGHIJKLMNOPQRSTUVWXYZ", null, null, null);
 
-			ServerLogPacketBuilder writer = new ServerLogPacketBuilder(85);
+			ServerLogPacketBuilder writer = new(85);
 			writer.SanitizeAndWriteEvent(new JsonLogEvent(baseEvent));
 
 			string[] output = Encoding.UTF8.GetString(writer.CreatePacket().Item1.Span).Split('\n');
diff --git a/Engine/Source/Programs/Shared/EpicGames.Horde.Tests/Storage/BackendTests.cs b/Engine/Source/Programs/Shared/EpicGames.Horde.Tests/Storage/BackendTests.cs
--- a/Engine/Source/Programs/Shared/EpicGames.Horde.Tests/Storage/BackendTests.cs
+++ b/Engine/Source/Programs/Shared/EpicGames.Horde.Tests/Storage/BackendTests.cs
@@ -39,17 +39,17 @@ namespace EpicGames.Horde.Tests.Storage
 		[TestMethod]
 		public async Task TestMemoryBackendAsync()
 		{
-			MemoryStorageBackend backend = new MemoryStorageBackend();
+			MemoryStorageBackend backend = new();
 			await TestBackendAsync(backend);
 		}
 
 		[TestMethod]
 		public async Task TestFileBackendAsync()
 		{
-			using (TempDir tempDir = new TempDir("Cache-" + DateTime.UtcNow.Ticks))
+			using (TempDir tempDir = new("Cache-" + DateTime.UtcNow.Ticks))
 			{
-				using MemoryMappedFileCache memoryMappedFileCache = new MemoryMappedFileCache();
-				FileStorageBackend backend = new FileStorageBackend(tempDir.Location, memoryMappedFileCache, NullLogger.Instance);
+				using MemoryMappedFileCache memoryMappedFileCache = new();
+				FileStorageBackend backend = new(tempDir.Location, memoryMappedFileCache, NullLogger.Instance);
 				await TestBackendAsync(backend);
 			}
 		}
@@ -57,10 +57,10 @@ namespace EpicGames.Horde.Tests.Storage
 		[TestMethod]
 		public async Task TestFileBackendFileMappingAsync()
 		{
-			using (TempDir tempDir = new TempDir("Cache-" + DateTime.UtcNow.Ticks))
+			using (TempDir tempDir = new("Cache-" + DateTime.UtcNow.Ticks))
 			{
-				using MemoryMappedFileCache memoryMappedFileCache = new MemoryMappedFileCache();
-				FileStorageBackend backend = new FileStorageBackend(tempDir.Location, memoryMappedFileCache, NullLogger.Instance);
+				using MemoryMappedFileCache memoryMappedFileCache = new();
+				FileStorageBackend backend = new(tempDir.Location, memoryMappedFileCache, NullLogger.Instance);
 				BlobLocator path = await backend.WriteBytesAsync(Encoding.UTF8.GetBytes("hello world"), []);
 
 				using (IReadOnlyMemoryOwner<byte> handle = await backend.ReadBlobAsync(path, 0, null))
@@ -81,11 +81,11 @@ namespace EpicGames.Horde.Tests.Storage
 		[TestMethod]
 		public async Task TestCacheBackendAsync()
 		{
-			using (TempDir tempDir = new TempDir("Cache-" + DateTime.UtcNow.Ticks))
+			using (TempDir tempDir = new("Cache-" + DateTime.UtcNow.Ticks))
 			{
-				using StorageBackendCache cache = new StorageBackendCache(tempDir.Location, 12, NullLogger.Instance);
+				using StorageBackendCache cache = new(tempDir.Location, 12, NullLogger.Instance);
 
-				MemoryStorageBackend memoryBackend = new MemoryStorageBackend();
+				MemoryStorageBackend memoryBackend = new();
 				IStorageBackend cacheBackend = cache.CreateWrapper("", memoryBackend);
 
 				await TestBackendAsync(cacheBackend);
diff --git a/Engine/Source/Programs/Shared/EpicGames.Horde.Tests/Storage/BlobDataTests.cs b/Engine/Source/Programs/Shared/EpicGames.Horde.Tests/Storage/BlobDataTests.cs
--- a/Engine/Source/Programs/Shared/EpicGames.Horde.Tests/Storage/BlobDataTests.cs
+++ b/Engine/Source/Programs/Shared/EpicGames.Horde.Tests/Storage/BlobDataTests.cs
@@ -38,10 +38,10 @@ namespace EpicGames.Horde.Tests.Storage
 		{
 			List<IBlobRef> refs = [new DummyHandle("hello"), new DummyHandle("world")];
 
-			using BlobData blobData = new BlobData(new BlobType(Guid.NewGuid(), 4), new byte[] { 1, 2, 3 }, refs);
+			using BlobData blobData = new(new BlobType(Guid.NewGuid(), 4), new byte[] { 1, 2, 3 }, refs);
 
 			byte[] data = EncodedBlobData.Create(blobData);
-			EncodedBlobData packet = new EncodedBlobData(data);
+			EncodedBlobData packet = new(data);
 			Assert.AreEqual(49, data.Length);
 
 			Assert.AreEqual(blobData.Type.Guid, packet.Type.Guid);
diff --git a/Engine/Source/Programs/Shared/EpicGames.Horde.Tests/Storage/BlobPipelineTests.cs b/Engine/Source/Programs/Shared/EpicGames.Horde.Tests/Storage/BlobPipelineTests.cs
--- a/Engine/Source/Programs/Shared/EpicGames.Horde.Tests/Storage/BlobPipelineTests.cs
+++ b/Engine/Source/Programs/Shared/EpicGames.Horde.Tests/Storage/BlobPipelineTests.cs
@@ -15,7 +15,7 @@ namespace EpicGames.Horde.Tests.Storage
 	[TestClass]
 	public class BlobPipelineTests
 	{
-		static readonly BlobType s_blobType = new BlobType(Guid.Empty, 1);
+		static readonly BlobType s_blobType = new(Guid.Empty, 1);
 
 		[TestMethod]
 		public async Task TestAsync()
@@ -32,7 +32,7 @@ namespace EpicGames.Horde.Tests.Storage
 				}
 			}
 
-			await using (BlobPipeline<IoHash> pipeline = new BlobPipeline<IoHash>())
+			await using (BlobPipeline<IoHash> pipeline = new())
 			{
 				foreach (IHashedBlobRef blobRef in handles)
 				{
diff --git a/Engine/Source/Programs/Shared/EpicGames.Horde.Tests/Storage/BundleStoreTests.cs b/Engine/Source/Programs/Shared/EpicGames.Horde.Tests/Storage/BundleStoreTests.cs
--- a/Engine/Source/Programs/Shared/EpicGames.Horde.Tests/Storage/BundleStoreTests.cs
+++ b/Engine/Source/Programs/Shared/EpicGames.Horde.Tests/Storage/BundleStoreTests.cs
@@ -60,11 +60,11 @@ namespace EpicGames.Horde.Tests.Storage
 
 		static async Task<byte[]> CreateBundleNormalAsync()
 		{
-			MemoryStorageBackend memoryStore = new MemoryStorageBackend();
-			BundleStorageNamespace store = new BundleStorageNamespace(memoryStore, BundleCache.None, new BundleOptions { MaxVersion = BundleVersion.ImportHashes, CompressionFormat = BundleCompressionFormat.None }, NullLogger.Instance);
+			MemoryStorageBackend memoryStore = new();
+			BundleStorageNamespace store = new(memoryStore, BundleCache.None, new BundleOptions { MaxVersion = BundleVersion.ImportHashes, CompressionFormat = BundleCompressionFormat.None }, NullLogger.Instance);
 			await using IBlobWriter writer = store.CreateBlobWriter();
 
-			TextNode node = new TextNode("Hello world");
+			TextNode node = new("Hello world");
 			IHashedBlobRef<TextNode> handle = await writer.WriteBlobAsync(node);
 			await writer.FlushAsync();
 
@@ -74,7 +74,7 @@ namespace EpicGames.Horde.Tests.Storage
 
 		static byte[] CreateBundleManually()
 		{
-			ArrayMemoryWriter payloadWriter = new ArrayMemoryWriter(200);
+			ArrayMemoryWriter payloadWriter = new(200);
 			payloadWriter.WriteString("Hello world");
 			byte[] payload = payloadWriter.WrittenMemory.ToArray();
 
@@ -84,9 +84,9 @@ namespace EpicGames.Horde.Tests.Storage
 
 			List<BundlePacket> packets = [new BundlePacket(BundleCompressionFormat.None, 0, payload.Length, payload.Length)];
 
-			BundleHeader header = new BundleHeader([.. types], [], [.. exports], [.. packets]);
+			BundleHeader header = new([.. types], [], [.. exports], [.. packets]);
 
-			ReadOnlySequenceBuilder<byte> builder = new ReadOnlySequenceBuilder<byte>();
+			ReadOnlySequenceBuilder<byte> builder = new();
 			header.AppendTo(builder);
 			builder.Append(payload);
 			return builder.Construct().ToArray();
@@ -95,10 +95,10 @@ namespace EpicGames.Horde.Tests.Storage
 		[TestMethod]
 		public async Task TestTreeAsync()
 		{
-			MemoryStorageBackend blobStore = new MemoryStorageBackend();
+			MemoryStorageBackend blobStore = new();
 
-			BundleOptions bundleOptions = new BundleOptions { MaxBlobSize = 1024 * 1024 };
-			BundleStorageNamespace bundleStore = new BundleStorageNamespace(blobStore, BundleCache.None, bundleOptions, NullLogger.Instance);
+			BundleOptions bundleOptions = new() { MaxBlobSize = 1024 * 1024 };
+			BundleStorageNamespace bundleStore = new(blobStore, BundleCache.None, bundleOptions, NullLogger.Instance);
 
 			await TestTreeAsync(bundleStore);
 
@@ -109,10 +109,10 @@ namespace EpicGames.Horde.Tests.Storage
 		[TestMethod]
 		public async Task TestTreeSeparateBlobsAsync()
 		{
-			MemoryStorageBackend blobStore = new MemoryStorageBackend();
+			MemoryStorageBackend blobStore = new();
 
-			BundleOptions bundleOptions = new BundleOptions { MaxBlobSize = 1 };
-			BundleStorageNamespace bundleStore = new BundleStorageNamespace(blobStore, BundleCache.None, bundleOptions, NullLogger.Instance);
+			BundleOptions bundleOptions = new() { MaxBlobSize = 1 };
+			BundleStorageNamespace bundleStore = new(blobStore, BundleCache.None, bundleOptions, NullLogger.Instance);
 
 			await TestTreeAsync(bundleStore);
 
@@ -158,18 +158,18 @@ namespace EpicGames.Horde.Tests.Storage
 			{
 				await using IBlobWriter writer = store.CreateBlobWriter("test");
 
-				SimpleNode node1 = new SimpleNode(new ReadOnlySequence<byte>([1]), []);
-				SimpleNode node2 = new SimpleNode(new ReadOnlySequence<byte>([2]), [await writer.WriteBlobAsync(node1)]);
-				SimpleNode node3 = new SimpleNode(new ReadOnlySequence<byte>([3]), [await writer.WriteBlobAsync(node2)]);
-				SimpleNode node4 = new SimpleNode(new ReadOnlySequence<byte>([4]), []);
+				SimpleNode node1 = new(new ReadOnlySequence<byte>([1]), []);
+				SimpleNode node2 = new(new ReadOnlySequence<byte>([2]), [await writer.WriteBlobAsync(node1)]);
+				SimpleNode node3 = new(new ReadOnlySequence<byte>([3]), [await writer.WriteBlobAsync(node2)]);
+				SimpleNode node4 = new(new ReadOnlySequence<byte>([4]), []);
 
-				SimpleNode root = new SimpleNode(new ReadOnlySequence<byte>([5]), [await writer.WriteBlobAsync(node4), await writer.WriteBlobAsync(node3)]);
+				SimpleNode root = new(new ReadOnlySequence<byte>([5]), [await writer.WriteBlobAsync(node4), await writer.WriteBlobAsync(node3)]);
 				IHashedBlobRef<SimpleNode> rootRef = await writer.WriteBlobAsync(root);
 				await writer.FlushAsync();
 
 				await store.AddRefAsync(new RefName("test"), rootRef);
 
-				BundleReader reader = new BundleReader(store, BundleCache.None, NullLogger.Instance);
+				BundleReader reader = new(store, BundleCache.None, NullLogger.Instance);
 				await CheckTreeAsync(root);
 			}
 
@@ -218,7 +218,7 @@ namespace EpicGames.Horde.Tests.Storage
 		{
 			KeyValueStorageNamespace store = KeyValueStorageNamespace.CreateInMemory();
 
-			RefName refName = new RefName("test");
+			RefName refName = new("test");
 
 			IHashedBlobRef<SimpleNode> inputRef;
 			await using (IBlobWriter writer = store.CreateBlobWriter(refName))
@@ -241,14 +241,14 @@ namespace EpicGames.Horde.Tests.Storage
 			{
 				await using (IBlobWriter writer = store.CreateBlobWriter(new RefName("test")))
 				{
-					DirectoryNode world = new DirectoryNode();
+					DirectoryNode world = new();
 					IHashedBlobRef<DirectoryNode> worldRef = await writer.WriteBlobAsync(world);
 
-					DirectoryNode hello = new DirectoryNode();
+					DirectoryNode hello = new();
 					hello.AddDirectory(new DirectoryEntry("world", 0, worldRef));
 					IHashedBlobRef<DirectoryNode> helloRef = await writer.WriteBlobAsync(hello);
 
-					DirectoryNode root = new DirectoryNode(DirectoryFlags.None);
+					DirectoryNode root = new(DirectoryFlags.None);
 					root.AddDirectory(new DirectoryEntry("hello", 0, helloRef));
 					IHashedBlobRef<DirectoryNode> rootRef = await writer.WriteBlobAsync(root);
 
@@ -283,18 +283,18 @@ namespace EpicGames.Horde.Tests.Storage
 		{
 			using IMemoryCache cache = new MemoryCache(new MemoryCacheOptions());
 			BundleStorageNamespace store = BundleStorageNamespace.CreateInMemory(NullLogger.Instance);
-			BundleReader reader = new BundleReader(store, BundleCache.None, NullLogger.Instance);
+			BundleReader reader = new(store, BundleCache.None, NullLogger.Instance);
 
 			// Generate a tree
 			{
 				await using IBlobWriter writer = store.CreateBlobWriter();
 
-				using ChunkedDataWriter fileWriter = new ChunkedDataWriter(writer, new ChunkingOptions());
+				using ChunkedDataWriter fileWriter = new(writer, new ChunkingOptions());
 				ChunkedData fileHandle = await fileWriter.CreateAsync(Encoding.UTF8.GetBytes("world"), CancellationToken.None);
 
 				List<FileUpdate> fileUpdates = [new FileUpdate("hello/world", FileEntryFlags.None, fileWriter.Length, fileHandle)];
 
-				DirectoryNode root = new DirectoryNode();
+				DirectoryNode root = new();
 				await root.UpdateAsync(fileUpdates, writer);
 
 				IHashedBlobRef<DirectoryNode> rootRef = await writer.WriteBlobAsync(root);
@@ -329,7 +329,7 @@ namespace EpicGames.Horde.Tests.Storage
 		[TestMethod]
 		public async Task StreamTestAsync()
 		{
-			BundleOptions bundleOptions = new BundleOptions { MaxBlobSize = 1024 };
+			BundleOptions bundleOptions = new() { MaxBlobSize = 1024 };
 			BundleStorageNamespace store = BundleStorageNamespace.CreateInMemory(bundleOptions, NullLogger.Instance);
 
 			const int Length = 4096;
@@ -342,10 +342,10 @@ namespace EpicGames.Horde.Tests.Storage
 			{
 				await using IBlobWriter writer = store.CreateBlobWriter();
 
-				ChunkingOptions options = new ChunkingOptions();
+				ChunkingOptions options = new();
 				options.LeafOptions = new LeafChunkedDataNodeOptions(128, 256, 64 * 1024);
 
-				using ChunkedDataWriter fileWriter = new ChunkedDataWriter(writer, options);
+				using ChunkedDataWriter fileWriter = new(writer, options);
 				for (int idx = 0; idx < chunk.Length / 16; idx++)
 				{
 					await fileWriter.AppendAsync(chunk.AsMemory(idx * 16, 16), CancellationToken.None);
@@ -358,7 +358,7 @@ namespace EpicGames.Horde.Tests.Storage
 			{
 				ChunkedDataNode newRoot = await nodeRef.ReadBlobAsync();
 
-				using MemoryStream stream = new MemoryStream();
+				using MemoryStream stream = new();
 				await newRoot.CopyToStreamAsync(stream);
 
 				byte[] output = stream.ToArray();
@@ -369,7 +369,7 @@ namespace EpicGames.Horde.Tests.Storage
 		[TestMethod]
 		public async Task LargeFileTestAsync()
 		{
-			BundleOptions bundleOptions = new BundleOptions { MaxBlobSize = 1024 };
+			BundleOptions bundleOptions = new() { MaxBlobSize = 1024 };
 			BundleStorageNamespace store = BundleStorageNamespace.CreateInMemory(bundleOptions, NullLogger.Instance);
 
 			const int Length = 1024;
@@ -387,12 +387,12 @@ namespace EpicGames.Horde.Tests.Storage
 			// Generate a tree
 			DirectoryNode root;
 			{
-				await using DedupeBlobWriter writer = new DedupeBlobWriter(store.CreateBlobWriter());
+				await using DedupeBlobWriter writer = new(store.CreateBlobWriter());
 
-				ChunkingOptions options = new ChunkingOptions();
+				ChunkingOptions options = new();
 				options.LeafOptions = new LeafChunkedDataNodeOptions(128, 256, 64 * 1024);
 
-				using ChunkedDataWriter fileWriter = new ChunkedDataWriter(writer, options);
+				using ChunkedDataWriter fileWriter = new(writer, options);
 				ChunkedData chunkedData = await fileWriter.CreateAsync(data, CancellationToken.None);
 
 				root = new DirectoryNode(DirectoryFlags.None);
@@ -509,7 +509,7 @@ namespace EpicGames.Horde.Tests.Storage
 
 		static async Task<byte[]> GetFileDataAsync(ChunkedDataNode fileNode)
 		{
-			using (MemoryStream stream = new MemoryStream())
+			using (MemoryStream stream = new())
 			{
 				await fileNode.CopyToStreamAsync(stream);
 				return stream.ToArray();
diff --git a/Engine/Source/Programs/Shared/EpicGames.Horde.Tests/Storage/BundleTests.cs b/Engine/Source/Programs/Shared/EpicGames.Horde.Tests/Storage/BundleTests.cs
--- a/Engine/Source/Programs/Shared/EpicGames.Horde.Tests/Storage/BundleTests.cs
+++ b/Engine/Source/Programs/Shared/EpicGames.Horde.Tests/Storage/BundleTests.cs
@@ -53,13 +53,13 @@ namespace EpicGames.Horde.Tests.Storage
 		{
 			BundleStorageNamespace storage = BundleStorageNamespace.CreateInMemory(NullLogger.Instance);
 
-			RefName refName = new RefName("test");
+			RefName refName = new("test");
 			await using IBlobWriter writer = storage.CreateBlobWriter(refName);
 
-			ChunkingOptions options = new ChunkingOptions();
+			ChunkingOptions options = new();
 			options.LeafOptions = new LeafChunkedDataNodeOptions(8, 8, 8);
 
-			using ChunkedDataWriter fileNodeWriter = new ChunkedDataWriter(writer, options);
+			using ChunkedDataWriter fileNodeWriter = new(writer, options);
 
 			ChunkedDataNode node;
 			ChunkedDataNodeRef nodeRef;
@@ -178,21 +178,21 @@ namespace EpicGames.Horde.Tests.Storage
 			IHashedBlobRef<DirectoryNode> rootRef;
 			await using (IBlobWriter writer = store.CreateBlobWriter())
 			{
-				DirectoryNode world = new DirectoryNode();
+				DirectoryNode world = new();
 				IHashedBlobRef<DirectoryNode> worldRef = await writer.WriteBlobAsync(world);
 
-				DirectoryNode hello = new DirectoryNode();
+				DirectoryNode hello = new();
 				hello.AddDirectory(new DirectoryEntry("world", 0, worldRef));
 				IHashedBlobRef<DirectoryNode> helloRef = await writer.WriteBlobAsync(hello);
 
-				DirectoryNode root = new DirectoryNode();
+				DirectoryNode root = new();
 				root.AddDirectory(new DirectoryEntry("hello", 0, helloRef));
 				rootRef = await writer.WriteBlobAsync(root);
 
 				await writer.FlushAsync();
 			}
 
-			RefName refName = new RefName("testref");
+			RefName refName = new("testref");
 			await store.AddRefAsync(refName, rootRef);
 
 			// Should be stored inline
@@ -230,22 +230,22 @@ namespace EpicGames.Horde.Tests.Storage
 		[TestMethod]
 		public async Task DedupTestsAsync()
 		{
-			BundleOptions bundleOptions = new BundleOptions();
+			BundleOptions bundleOptions = new();
 			bundleOptions.MaxBlobSize = 1;
 
 			BundleStorageNamespace storage = BundleStorageNamespace.CreateInMemory(bundleOptions, NullLogger.Instance);
 
 			await using (IBlobWriter writer = new DedupeBlobWriter(storage.CreateBlobWriter()))
 			{
-				DirectoryNode root = new DirectoryNode();
+				DirectoryNode root = new();
 				for (int idx = 1; idx <= 3; idx++)
 				{
-					DirectoryNode node = new DirectoryNode();
+					DirectoryNode node = new();
 					IHashedBlobRef<DirectoryNode> nodeRef = await writer.WriteBlobAsync(node);
 					root.AddDirectory(new DirectoryEntry($"node{idx}", 0, nodeRef));
 				}
 
-				RefName refName = new RefName("ref");
+				RefName refName = new("ref");
 				IHashedBlobRef<DirectoryNode> rootRef = await writer.WriteBlobAsync(root);
 				await storage.AddRefAsync(refName, rootRef);
 			}
@@ -258,12 +258,12 @@ namespace EpicGames.Horde.Tests.Storage
 		[TestMethod]
 		public async Task ReloadTestsAsync()
 		{
-			BundleOptions bundleOptions = new BundleOptions();
+			BundleOptions bundleOptions = new();
 			bundleOptions.MaxBlobSize = 1;
 
 			BundleStorageNamespace storage = BundleStorageNamespace.CreateInMemory(bundleOptions, NullLogger.Instance);
 
-			RefName refName = new RefName("ref");
+			RefName refName = new("ref");
 
 			{
 				await using (IBlobWriter writer = storage.CreateBlobWriter())
@@ -271,7 +271,7 @@ namespace EpicGames.Horde.Tests.Storage
 					IHashedBlobRef<DirectoryNode> rootRef = await writer.WriteBlobAsync(new DirectoryNode());
 					for (int idx = 4; idx >= 1; idx--)
 					{
-						DirectoryNode next = new DirectoryNode();
+						DirectoryNode next = new();
 						next.AddDirectory(new DirectoryEntry($"node{idx}", 0, rootRef));
 						rootRef = await writer.WriteBlobAsync(next);
 					}
