Shared: prefer implicit object creation when type is apparent

diff --git a/Engine/Source/Programs/Shared/EpicGames.Horde.Tests/Storage/ChunkingTests.cs b/Engine/Source/Programs/Shared/EpicGames.Horde.Tests/Storage/ChunkingTests.cs
--- a/Engine/Source/Programs/Shared/EpicGames.Horde.Tests/Storage/ChunkingTests.cs
+++ b/Engine/Source/Programs/Shared/EpicGames.Horde.Tests/Storage/ChunkingTests.cs
@@ -58,15 +58,15 @@ namespace EpicGames.Horde.Tests.Storage
 			const string RefName = "hello";
 			await using (IBlobWriter writer = store.CreateBlobWriter(RefName))
 			{
-				ChunkingOptions options = new ChunkingOptions();
+				ChunkingOptions options = new();
 				options.LeafOptions = new LeafChunkedDataNodeOptions(64, 64, 64);
 				options.InteriorOptions = new InteriorChunkedDataNodeOptions(4, 4, 4);
 
-				using MemoryStream emptyStream = new MemoryStream();
+				using MemoryStream emptyStream = new();
 				LeafChunkedData leafChunkedData = await LeafChunkedDataNode.CreateFromStreamAsync(writer, emptyStream, new LeafChunkedDataNodeOptions(64, 64, 64), CancellationToken.None);
 				ChunkedData chunkedData = await InteriorChunkedDataNode.CreateTreeAsync(leafChunkedData, new InteriorChunkedDataNodeOptions(4, 4, 4), writer, CancellationToken.None);
 
-				DirectoryNode directory = new DirectoryNode();
+				DirectoryNode directory = new();
 				directory.AddFile("test.foo", FileEntryFlags.None, chunkedData);
 
 				IHashedBlobRef handle = await writer.WriteBlobAsync(directory);
@@ -77,7 +77,7 @@ namespace EpicGames.Horde.Tests.Storage
 		[TestMethod]
 		public async Task FixedSizeChunkingTestsAsync()
 		{
-			ChunkingOptions options = new ChunkingOptions();
+			ChunkingOptions options = new();
 			options.LeafOptions = new LeafChunkedDataNodeOptions(64, 64, 64);
 			options.InteriorOptions = new InteriorChunkedDataNodeOptions(4, 4, 4);
 
@@ -87,7 +87,7 @@ namespace EpicGames.Horde.Tests.Storage
 		[TestMethod]
 		public async Task VariableSizeChunkingTestsAsync()
 		{
-			ChunkingOptions options = new ChunkingOptions();
+			ChunkingOptions options = new();
 			options.LeafOptions = new LeafChunkedDataNodeOptions(32, 64, 96);
 			options.InteriorOptions = new InteriorChunkedDataNodeOptions(1, 4, 12);
 
@@ -96,7 +96,7 @@ namespace EpicGames.Horde.Tests.Storage
 
 		static async Task TestChunkingAsync(ChunkingOptions options)
 		{
-			using MemoryCache cache = new MemoryCache(new MemoryCacheOptions());
+			using MemoryCache cache = new(new MemoryCacheOptions());
 
 			KeyValueStorageNamespace store = KeyValueStorageNamespace.CreateInMemory();
 
@@ -114,7 +114,7 @@ namespace EpicGames.Horde.Tests.Storage
 
 			const int NumIterations = 100;
 			{
-				using ChunkedDataWriter fileWriter = new ChunkedDataWriter(writer, options);
+				using ChunkedDataWriter fileWriter = new(writer, options);
 
 				for (int idx = 0; idx < NumIterations; idx++)
 				{
@@ -127,7 +127,7 @@ namespace EpicGames.Horde.Tests.Storage
 			ChunkedDataNode root = await handle.ReadBlobAsync();
 
 			byte[] result;
-			using (MemoryStream stream = new MemoryStream())
+			using (MemoryStream stream = new())
 			{
 				await root.CopyToStreamAsync(stream);
 				result = stream.ToArray();
@@ -170,11 +170,11 @@ namespace EpicGames.Horde.Tests.Storage
 		[TestMethod]
 		public async Task ChunkingCompatV2Async()
 		{
-			ChunkingOptions chunkingOptions = new ChunkingOptions();
+			ChunkingOptions chunkingOptions = new();
 			chunkingOptions.LeafOptions = new LeafChunkedDataNodeOptions(2, 2, 2);
 			chunkingOptions.InteriorOptions = new InteriorChunkedDataNodeOptions(2, 2, 2);
 
-			BlobSerializerOptions serializerOptions = new BlobSerializerOptions();
+			BlobSerializerOptions serializerOptions = new();
 			serializerOptions.Converters.Add(new InteriorChunkedDataNodeConverter(HordeApiVersion.Initial)); // Does not include length fields in interior nodes
 
 			KeyValueStorageNamespace store = KeyValueStorageNamespace.CreateInMemory();
@@ -184,16 +184,16 @@ namespace EpicGames.Horde.Tests.Storage
 			IHashedBlobRef<DirectoryNode> handle;
 			await using (IBlobWriter writer = store.CreateBlobWriter(options: serializerOptions))
 			{
-				using ChunkedDataWriter chunkedWriter = new ChunkedDataWriter(writer, chunkingOptions);
+				using ChunkedDataWriter chunkedWriter = new(writer, chunkingOptions);
 				await chunkedWriter.AppendAsync(data, CancellationToken.None);
 				ChunkedData chunkedData = await chunkedWriter.FlushAsync(CancellationToken.None);
 
-				DirectoryNode directoryNode = new DirectoryNode();
+				DirectoryNode directoryNode = new();
 				directoryNode.AddFile("test", FileEntryFlags.None, chunkedData);
 				handle = await writer.WriteBlobAsync(directoryNode);
 			}
 
-			DirectoryInfo tempDir = new DirectoryInfo(Path.Combine(Path.GetTempPath(), Guid.NewGuid().ToString()));
+			DirectoryInfo tempDir = new(Path.Combine(Path.GetTempPath(), Guid.NewGuid().ToString()));
 			try
 			{
 				await handle.ExtractAsync(tempDir, new ExtractOptions(), NullLogger.Instance, CancellationToken.None);
@@ -210,14 +210,14 @@ namespace EpicGames.Horde.Tests.Storage
 		[TestMethod]
 		public async Task HashAsync()
 		{
-			ChunkingOptions chunkingOptions = new ChunkingOptions();
+			ChunkingOptions chunkingOptions = new();
 			chunkingOptions.LeafOptions = new LeafChunkedDataNodeOptions(64, 128, 92);
 
-			await using MemoryBlobWriter blobWriter = new MemoryBlobWriter(new BlobSerializerOptions());
+			await using MemoryBlobWriter blobWriter = new(new BlobSerializerOptions());
 			byte[] data = RandomNumberGenerator.GetBytes(4096);
 
 			ChunkedData chunkedData;
-			using (ChunkedDataWriter writer = new ChunkedDataWriter(blobWriter, chunkingOptions))
+			using (ChunkedDataWriter writer = new(blobWriter, chunkingOptions))
 			{
 				await writer.AppendAsync(data, default);
 				chunkedData = await writer.FlushAsync();
@@ -230,16 +230,16 @@ namespace EpicGames.Horde.Tests.Storage
 		[TestMethod]
 		public async Task ChunkOrderAsync()
 		{
-			InteriorChunkedDataNodeOptions interiorOptions = new InteriorChunkedDataNodeOptions(2, 2, 2);
+			InteriorChunkedDataNodeOptions interiorOptions = new(2, 2, 2);
 
-			BlobSerializerOptions serializerOptions = new BlobSerializerOptions();
+			BlobSerializerOptions serializerOptions = new();
 			serializerOptions.Converters.Add(new InteriorChunkedDataNodeConverter());
 
-			await using MemoryBlobWriter blobWriter = new MemoryBlobWriter(new BlobSerializerOptions());
+			await using MemoryBlobWriter blobWriter = new(new BlobSerializerOptions());
 
 			IHashedBlobRef<LeafChunkedDataNode> leafRef = await blobWriter.WriteBlobAsync(new LeafChunkedDataNode(new byte[] { 1, 2, 3 }));
 			int leafRefIndex = MemoryBlobWriter.GetIndex((IHashedBlobRef)leafRef);
-			ChunkedDataNodeRef leafChunkedRef = new ChunkedDataNodeRef(3, 0, leafRef);
+			ChunkedDataNodeRef leafChunkedRef = new(3, 0, leafRef);
 
 			List<ChunkedDataNodeRef> leafNodeRefs = [.. Enumerable.Repeat(leafChunkedRef, 10000)];
 			ChunkedDataNodeRef root = await InteriorChunkedDataNode.CreateTreeAsync(leafNodeRefs, interiorOptions, blobWriter, CancellationToken.None);
diff --git a/Engine/Source/Programs/Shared/EpicGames.Horde.Tests/Storage/CompressionTests.cs b/Engine/Source/Programs/Shared/EpicGames.Horde.Tests/Storage/CompressionTests.cs
--- a/Engine/Source/Programs/Shared/EpicGames.Horde.Tests/Storage/CompressionTests.cs
+++ b/Engine/Source/Programs/Shared/EpicGames.Horde.Tests/Storage/CompressionTests.cs
@@ -43,7 +43,7 @@ namespace EpicGames.Horde.Tests.Storage
 
 		static void RunTests(BundleCompressionFormat format, byte[] input)
 		{
-			ArrayMemoryWriter writer = new ArrayMemoryWriter(10 * 1024 * 1024);
+			ArrayMemoryWriter writer = new(10 * 1024 * 1024);
 			Stopwatch compressTimer = Stopwatch.StartNew();
 			BundleData.Compress(format, input, writer);
 			compressTimer.Stop();
diff --git a/Engine/Source/Programs/Shared/EpicGames.Horde.Tests/Storage/ContentChunkerTests.cs b/Engine/Source/Programs/Shared/EpicGames.Horde.Tests/Storage/ContentChunkerTests.cs
--- a/Engine/Source/Programs/Shared/EpicGames.Horde.Tests/Storage/ContentChunkerTests.cs
+++ b/Engine/Source/Programs/Shared/EpicGames.Horde.Tests/Storage/ContentChunkerTests.cs
@@ -21,10 +21,10 @@ namespace EpicGames.Horde.Tests.Storage
 		public async Task SerialTestAsync()
 		{
 			List<SplitPoint> legacyPoints;
-			await using (AsyncPipeline pipeline = new AsyncPipeline(CancellationToken.None))
+			await using (AsyncPipeline pipeline = new(CancellationToken.None))
 			{
 				byte[] buffer = new byte[1024 * 1024];
-				SerialBuzHashChunker chunker = new SerialBuzHashChunker(pipeline, buffer, LeafChunkedDataNodeOptions.Default);
+				SerialBuzHashChunker chunker = new(pipeline, buffer, LeafChunkedDataNodeOptions.Default);
 				legacyPoints = await TestPipelineAsync(chunker);
 			}
 		}
@@ -33,10 +33,10 @@ namespace EpicGames.Horde.Tests.Storage
 		public async Task ParallelTestAsync()
 		{
 			List<SplitPoint> latestPoints;
-			await using (AsyncPipeline pipeline = new AsyncPipeline(CancellationToken.None))
+			await using (AsyncPipeline pipeline = new(CancellationToken.None))
 			{
 				byte[] buffer = new byte[1024 * 1024];
-				ParallelBuzHashChunker chunker = new ParallelBuzHashChunker(pipeline, buffer, LeafChunkedDataNodeOptions.Default);
+				ParallelBuzHashChunker chunker = new(pipeline, buffer, LeafChunkedDataNodeOptions.Default);
 				latestPoints = await TestPipelineAsync(chunker);
 			}
 		}
@@ -45,18 +45,18 @@ namespace EpicGames.Horde.Tests.Storage
 		public async Task CompareTestAsync()
 		{
 			List<SplitPoint> legacyPoints;
-			await using (AsyncPipeline pipeline = new AsyncPipeline(CancellationToken.None))
+			await using (AsyncPipeline pipeline = new(CancellationToken.None))
 			{
 				byte[] buffer = new byte[1024 * 1024];
-				SerialBuzHashChunker chunker = new SerialBuzHashChunker(pipeline, buffer, LeafChunkedDataNodeOptions.Default);
+				SerialBuzHashChunker chunker = new(pipeline, buffer, LeafChunkedDataNodeOptions.Default);
 				legacyPoints = await TestPipelineAsync(chunker);
 			}
 
 			List<SplitPoint> latestPoints;
-			await using (AsyncPipeline pipeline = new AsyncPipeline(CancellationToken.None))
+			await using (AsyncPipeline pipeline = new(CancellationToken.None))
 			{
 				byte[] buffer = new byte[1024 * 1024];
-				ParallelBuzHashChunker chunker = new ParallelBuzHashChunker(pipeline, buffer, LeafChunkedDataNodeOptions.Default);
+				ParallelBuzHashChunker chunker = new(pipeline, buffer, LeafChunkedDataNodeOptions.Default);
 				latestPoints = await TestPipelineAsync(chunker);
 			}
 
@@ -65,7 +65,7 @@ namespace EpicGames.Horde.Tests.Storage
 
 		static async Task<List<SplitPoint>> TestPipelineAsync(ContentChunker pipeline)
 		{
-			Random random = new Random(0);
+			Random random = new(0);
 
 			int[] lengths = [0, 1, 32 * 1024, 129 * 1024, 256 * 1024, 4 * 1024 * 1024];
 
diff --git a/Engine/Source/Programs/Shared/EpicGames.Horde.Tests/Storage/FileMetadataDbTests.cs b/Engine/Source/Programs/Shared/EpicGames.Horde.Tests/Storage/FileMetadataDbTests.cs
--- a/Engine/Source/Programs/Shared/EpicGames.Horde.Tests/Storage/FileMetadataDbTests.cs
+++ b/Engine/Source/Programs/Shared/EpicGames.Horde.Tests/Storage/FileMetadataDbTests.cs
@@ -60,7 +60,7 @@ namespace EpicGames.Horde.Tests.Storage
 
 			int a = await db.AddDirectoryAsync(FileMetadataDb.RootDirectoryId, "A");
 
-			DateTime time = new DateTime(2024, 1, 1, 12, 30, 45, DateTimeKind.Utc);
+			DateTime time = new(2024, 1, 1, 12, 30, 45, DateTimeKind.Utc);
 			int c = await db.AddFileAsync(FileMetadataDb.RootDirectoryId, "C.txt", time, 1234);
 			int d = await db.AddFileAsync(FileMetadataDb.RootDirectoryId, "D.txt", time, 2345);
 			int e = await db.AddFileAsync(a, "E.txt", time, 3456);
@@ -131,7 +131,7 @@ namespace EpicGames.Horde.Tests.Storage
 			int d1 = await db.AddDirectoryAsync(FileMetadataDb.RootDirectoryId, "D1");
 			int d2 = await db.AddDirectoryAsync(d1, "D2");
 
-			DateTime time = new DateTime(2024, 1, 1, 12, 30, 45, DateTimeKind.Utc);
+			DateTime time = new(2024, 1, 1, 12, 30, 45, DateTimeKind.Utc);
 			int f1 = await db.AddFileAsync(d1, "F1.txt", time, 1234);
 			int f2 = await db.AddFileAsync(d1, "F2.txt", time, 2345);
 			int f3 = await db.AddFileAsync(d2, "F3.txt", time, 3456);
@@ -171,7 +171,7 @@ namespace EpicGames.Horde.Tests.Storage
 			int d4 = await db.AddDirectoryAsync(d2, "D4");
 			int d5 = await db.AddDirectoryAsync(d4, "D5");
 
-			DateTime time = new DateTime(2024, 1, 1, 12, 30, 45, DateTimeKind.Utc);
+			DateTime time = new(2024, 1, 1, 12, 30, 45, DateTimeKind.Utc);
 			int f1 = await db.AddFileAsync(d1, "F1.txt", time, 1234);
 			int f2 = await db.AddFileAsync(d2, "F2.txt", time, 2345);
 			int f3 = await db.AddFileAsync(d3, "F3.txt", time, 3456);
diff --git a/Engine/Source/Programs/Shared/EpicGames.Horde.Tests/Storage/StorageClientTests.cs b/Engine/Source/Programs/Shared/EpicGames.Horde.Tests/Storage/StorageClientTests.cs
--- a/Engine/Source/Programs/Shared/EpicGames.Horde.Tests/Storage/StorageClientTests.cs
+++ b/Engine/Source/Programs/Shared/EpicGames.Horde.Tests/Storage/StorageClientTests.cs
@@ -17,7 +17,7 @@ namespace EpicGames.Horde.Tests.Storage
 		[BlobConverter(typeof(TestNodeConverter))]
 		class TestNode
 		{
-			public static BlobType BlobType = new BlobType("{99601905-4F6E-A089-03D6-F187711BAFEE}", 1);
+			public static BlobType BlobType = new("{99601905-4F6E-A089-03D6-F187711BAFEE}", 1);
 
 			public int Value { get; }
 			public byte[] Padding { get; set; } = [];
@@ -88,7 +88,7 @@ namespace EpicGames.Horde.Tests.Storage
 		[TestMethod]
 		public async Task PendingRefsAsync()
 		{
-			await using BundleCache cache = new BundleCache();
+			await using BundleCache cache = new();
 
 			BundleStorageNamespace store = BundleStorageNamespace.CreateInMemory(NullLogger.Instance);
 
@@ -112,9 +112,9 @@ namespace EpicGames.Horde.Tests.Storage
 		[TestMethod]
 		public async Task PacketFlushAsync()
 		{
-			await using BundleCache cache = new BundleCache();
+			await using BundleCache cache = new();
 
-			BundleOptions bundleOptions = new BundleOptions { MinCompressionPacketSize = 100, MaxBlobSize = 1024 * 1024, MaxVersion = BundleVersion.LatestV2 };
+			BundleOptions bundleOptions = new() { MinCompressionPacketSize = 100, MaxBlobSize = 1024 * 1024, MaxVersion = BundleVersion.LatestV2 };
 			BundleStorageNamespace store = BundleStorageNamespace.CreateInMemory(bundleOptions, NullLogger.Instance);
 
 			await using IBlobWriter writer = store.CreateBlobWriter();
diff --git a/Engine/Source/Programs/Shared/EpicGames.Horde.Tests/Storage/WorkspaceTests.cs b/Engine/Source/Programs/Shared/EpicGames.Horde.Tests/Storage/WorkspaceTests.cs
--- a/Engine/Source/Programs/Shared/EpicGames.Horde.Tests/Storage/WorkspaceTests.cs
+++ b/Engine/Source/Programs/Shared/EpicGames.Horde.Tests/Storage/WorkspaceTests.cs
@@ -22,7 +22,7 @@ namespace EpicGames.Horde.Tests.Storage
 
 		static async Task<IHashedBlobRef<DirectoryNode>> CreateBasicTreeAsync(IBlobWriter blobWriter, CancellationToken cancellationToken)
 		{
-			DirectoryNode rootDir = new DirectoryNode();
+			DirectoryNode rootDir = new();
 			rootDir.AddFile("readme.txt", FileEntryFlags.None, await ChunkedData.CreateAsync(blobWriter, HelloBytes, cancellationToken));
 			return await blobWriter.WriteBlobAsync(rootDir, cancellationToken);
 		}
@@ -59,7 +59,7 @@ namespace EpicGames.Horde.Tests.Storage
 
 		static DirectoryInfo CreateTempDirectory([CallerMemberName] string? callerName = null)
 		{
-			DirectoryInfo tempDir = new DirectoryInfo($"temp/{callerName ?? Guid.NewGuid().ToString()}");
+			DirectoryInfo tempDir = new($"temp/{callerName ?? Guid.NewGuid().ToString()}");
 			tempDir.Create();
 			FileUtils.ForceDeleteDirectoryContents(tempDir);
 			return tempDir;
@@ -68,7 +68,7 @@ namespace EpicGames.Horde.Tests.Storage
 		[TestMethod]
 		public async Task MaterializeManuallyAsync()
 		{
-			await using MemoryBlobWriter blobWriter = new MemoryBlobWriter(new BlobSerializerOptions());
+			await using MemoryBlobWriter blobWriter = new(new BlobSerializerOptions());
 			IHashedBlobRef<DirectoryNode> input = await CreateBasicTreeAsync(blobWriter, default);
 
 			DirectoryInfo tempDir = CreateTempDirectory();
@@ -83,19 +83,19 @@ namespace EpicGames.Horde.Tests.Storage
 		{
 			DirectoryInfo tempDir = CreateTempDirectory();
 			Workspace ws = await Workspace.CreateAsync(new DirectoryReference(tempDir), NullLogger.Instance);
-			await using MemoryBlobWriter blobWriter = new MemoryBlobWriter(new BlobSerializerOptions());
+			await using MemoryBlobWriter blobWriter = new(new BlobSerializerOptions());
 
 			IHashedBlobRef<DirectoryNode> initialState;
 			{
-				DirectoryNode rootDir = new DirectoryNode();
+				DirectoryNode rootDir = new();
 				rootDir.AddFile("readme.txt", FileEntryFlags.None, await ChunkedData.CreateAsync(blobWriter, HelloBytes));
 				initialState = await blobWriter.WriteBlobAsync(rootDir);
 			}
 
-			WorkspaceSyncStats initialStats = new WorkspaceSyncStats();
+			WorkspaceSyncStats initialStats = new();
 			await ws.SyncAsync(WorkspaceLayerId.Default, initialState, null, stats: initialStats);
 
-			WorkspaceSyncStats updateStats = new WorkspaceSyncStats();
+			WorkspaceSyncStats updateStats = new();
 			await ws.SyncAsync(WorkspaceLayerId.Default, initialState, null, stats: updateStats);
 			Assert.AreEqual(1, updateStats.NumFilesKept);
 			Assert.AreEqual(0, updateStats.NumBytesReused);
@@ -110,11 +110,11 @@ namespace EpicGames.Horde.Tests.Storage
 		{
 			DirectoryInfo tempDir = CreateTempDirectory();
 			Workspace ws = await Workspace.CreateAsync(new DirectoryReference(tempDir), NullLogger.Instance);
-			await using MemoryBlobWriter blobWriter = new MemoryBlobWriter(new BlobSerializerOptions());
+			await using MemoryBlobWriter blobWriter = new(new BlobSerializerOptions());
 
 			IHashedBlobRef<DirectoryNode> initialState;
 			{
-				DirectoryNode rootDir = new DirectoryNode();
+				DirectoryNode rootDir = new();
 				rootDir.AddFile("readme.txt", FileEntryFlags.None, await ChunkedData.CreateAsync(blobWriter, HelloBytes));
 				initialState = await blobWriter.WriteBlobAsync(rootDir);
 			}
@@ -122,7 +122,7 @@ namespace EpicGames.Horde.Tests.Storage
 
 			IHashedBlobRef<DirectoryNode> updatedState;
 			{
-				DirectoryNode rootDir = new DirectoryNode();
+				DirectoryNode rootDir = new();
 				rootDir.AddFile("file2.txt", FileEntryFlags.None, await ChunkedData.CreateAsync(blobWriter, HelloBytes));
 				updatedState = await blobWriter.WriteBlobAsync(rootDir);
 			}
@@ -137,11 +137,11 @@ namespace EpicGames.Horde.Tests.Storage
 		{
 			DirectoryInfo tempDir = CreateTempDirectory();
 			Workspace ws = await Workspace.CreateAsync(new DirectoryReference(tempDir), NullLogger.Instance);
-			await using MemoryBlobWriter blobWriter = new MemoryBlobWriter(new BlobSerializerOptions());
+			await using MemoryBlobWriter blobWriter = new(new BlobSerializerOptions());
 
 			IHashedBlobRef<DirectoryNode> initialState;
 			{
-				DirectoryNode rootDir = new DirectoryNode();
+				DirectoryNode rootDir = new();
 				rootDir.AddFile("readme.txt", FileEntryFlags.None, await ChunkedData.CreateAsync(blobWriter, HelloBytes));
 				initialState = await blobWriter.WriteBlobAsync(rootDir);
 			}
@@ -149,7 +149,7 @@ namespace EpicGames.Horde.Tests.Storage
 
 			IHashedBlobRef<DirectoryNode> updatedState;
 			{
-				DirectoryNode rootDir = new DirectoryNode();
+				DirectoryNode rootDir = new();
 				rootDir.AddFile("file2.txt", FileEntryFlags.None, await ChunkedData.CreateAsync(blobWriter, HelloBytes));
 				rootDir.AddFile("file3.txt", FileEntryFlags.None, await ChunkedData.CreateAsync(blobWriter, HelloBytes));
 				updatedState = await blobWriter.WriteBlobAsync(rootDir);
@@ -165,34 +165,34 @@ namespace EpicGames.Horde.Tests.Storage
 		{
 			DirectoryInfo tempDir = CreateTempDirectory();
 			Workspace ws = await Workspace.CreateAsync(new DirectoryReference(tempDir), NullLogger.Instance);
-			await using MemoryBlobWriter blobWriter = new MemoryBlobWriter(new BlobSerializerOptions());
+			await using MemoryBlobWriter blobWriter = new(new BlobSerializerOptions());
 
 			IHashedBlobRef<DirectoryNode> emptyState;
 			{
-				DirectoryNode rootDir = new DirectoryNode();
+				DirectoryNode rootDir = new();
 				emptyState = await blobWriter.WriteBlobAsync(rootDir);
 			}
 
 			IHashedBlobRef<DirectoryNode> initialState;
 			{
-				DirectoryNode subDirA = new DirectoryNode();
+				DirectoryNode subDirA = new();
 				subDirA.AddFile("hello.txt", FileEntryFlags.None, await ChunkedData.CreateAsync(blobWriter, HelloBytes));
 				subDirA.AddFile("world.txt", FileEntryFlags.None, await ChunkedData.CreateAsync(blobWriter, WorldBytes));
 				IHashedBlobRef<DirectoryNode> subDirRefA = await blobWriter.WriteBlobAsync(subDirA);
 
-				DirectoryNode rootDir = new DirectoryNode();
+				DirectoryNode rootDir = new();
 				rootDir.AddDirectory(new DirectoryEntry("A", subDirA.Length, subDirRefA));
 				initialState = await blobWriter.WriteBlobAsync(rootDir);
 			}
 
-			WorkspaceLayerId layerA = new WorkspaceLayerId("layer-a");
+			WorkspaceLayerId layerA = new("layer-a");
 			ws.AddLayer(layerA);
 
-			WorkspaceLayerId layerB = new WorkspaceLayerId("layer-b");
+			WorkspaceLayerId layerB = new("layer-b");
 			ws.AddLayer(layerB);
 
 			// Add the files to layer A
-			WorkspaceSyncStats stats = new WorkspaceSyncStats();
+			WorkspaceSyncStats stats = new();
 			await ws.SyncAsync(layerA, initialState, stats: stats);
 			Assert.AreEqual(stats.NumFilesKept, 0);
 			Assert.AreEqual(stats.NumBytesDownloaded, 10);
@@ -229,16 +229,16 @@ namespace EpicGames.Horde.Tests.Storage
 		{
 			DirectoryInfo tempDir = CreateTempDirectory();
 			Workspace ws = await Workspace.CreateAsync(new DirectoryReference(tempDir), NullLogger.Instance);
-			await using MemoryBlobWriter blobWriter = new MemoryBlobWriter(new BlobSerializerOptions());
+			await using MemoryBlobWriter blobWriter = new(new BlobSerializerOptions());
 
 			IHashedBlobRef<DirectoryNode> initialState;
 			{
-				DirectoryNode subDirA = new DirectoryNode();
+				DirectoryNode subDirA = new();
 				subDirA.AddFile("hello.txt", FileEntryFlags.None, await ChunkedData.CreateAsync(blobWriter, HelloBytes));
 				subDirA.AddFile("world.txt", FileEntryFlags.None, await ChunkedData.CreateAsync(blobWriter, WorldBytes));
 				IHashedBlobRef<DirectoryNode> subDirRefA = await blobWriter.WriteBlobAsync(subDirA);
 
-				DirectoryNode rootDir = new DirectoryNode();
+				DirectoryNode rootDir = new();
 				rootDir.AddDirectory(new DirectoryEntry("A", subDirA.Length, subDirRefA));
 				rootDir.AddFile(new FileEntry("B", FileEntryFlags.None, await ChunkedData.CreateAsync(blobWriter, HelloBytes)));
 				initialState = await blobWriter.WriteBlobAsync(rootDir);
@@ -247,18 +247,18 @@ namespace EpicGames.Horde.Tests.Storage
 
 			IHashedBlobRef<DirectoryNode> updatedState;
 			{
-				DirectoryNode subDirB = new DirectoryNode();
+				DirectoryNode subDirB = new();
 				subDirB.AddFile("hello.txt", FileEntryFlags.None, await ChunkedData.CreateAsync(blobWriter, HelloBytes));
 				subDirB.AddFile("world.txt", FileEntryFlags.None, await ChunkedData.CreateAsync(blobWriter, WorldBytes));
 				IHashedBlobRef<DirectoryNode> subDirRefB = await blobWriter.WriteBlobAsync(subDirB);
 
-				DirectoryNode rootDir = new DirectoryNode();
+				DirectoryNode rootDir = new();
 				rootDir.AddFile(new FileEntry("A", FileEntryFlags.None, await ChunkedData.CreateAsync(blobWriter, HelloBytes)));
 				rootDir.AddDirectory(new DirectoryEntry("B", subDirB.Length, subDirRefB));
 				updatedState = await blobWriter.WriteBlobAsync(rootDir);
 			}
 
-			WorkspaceSyncStats syncStats = new WorkspaceSyncStats();
+			WorkspaceSyncStats syncStats = new();
 			await ws.SyncAsync(WorkspaceLayerId.Default, updatedState, stats: syncStats);
 
 			IHashedBlobRef<DirectoryNode> output = await blobWriter.AddFilesAsync(new DirectoryReference(tempDir), null, null, default);
diff --git a/Engine/Source/Programs/Shared/EpicGames.Horde.Tests/Utilities/IntervalJsonConverterTests.cs b/Engine/Source/Programs/Shared/EpicGames.Horde.Tests/Utilities/IntervalJsonConverterTests.cs
--- a/Engine/Source/Programs/Shared/EpicGames.Horde.Tests/Utilities/IntervalJsonConverterTests.cs
+++ b/Engine/Source/Programs/Shared/EpicGames.Horde.Tests/Utilities/IntervalJsonConverterTests.cs
@@ -17,7 +17,7 @@ namespace EpicGames.Horde.Tests.Utilities
 			public TimeSpan Time { get; set; }
 		}
 
-		readonly JsonSerializerOptions _options = new JsonSerializerOptions();
+		readonly JsonSerializerOptions _options = new();
 
 		[TestMethod]
 		public void Test()
diff --git a/Engine/Source/Programs/Shared/EpicGames.Horde/Accounts/AccountId.cs b/Engine/Source/Programs/Shared/EpicGames.Horde/Accounts/AccountId.cs
--- a/Engine/Source/Programs/Shared/EpicGames.Horde/Accounts/AccountId.cs
+++ b/Engine/Source/Programs/Shared/EpicGames.Horde/Accounts/AccountId.cs
@@ -17,7 +17,7 @@ namespace EpicGames.Horde.Accounts
 	public readonly record struct AccountId(BinaryId Id)
 	{
 		/// <inheritdoc cref="BinaryId.Parse(System.String)"/>
-		public static AccountId Parse(string text) => new AccountId(BinaryId.Parse(text));
+		public static AccountId Parse(string text) => new(BinaryId.Parse(text));
 
 		/// <inheritdoc cref="BinaryId.TryParse(ReadOnlySpan{Char}, out BinaryId)"/>
 		public static bool TryParse(ReadOnlySpan<char> text, out AccountId result)
@@ -45,7 +45,7 @@ namespace EpicGames.Horde.Accounts
 	class AccountIdConverter : BinaryIdConverter<AccountId>
 	{
 		/// <inheritdoc/>
-		public override AccountId FromBinaryId(BinaryId id) => new AccountId(id);
+		public override AccountId FromBinaryId(BinaryId id) => new(id);
 
 		/// <inheritdoc/>
 		public override BinaryId ToBinaryId(AccountId value) => value.Id;
diff --git a/Engine/Source/Programs/Shared/EpicGames.Horde/Acls/AclAction.cs b/Engine/Source/Programs/Shared/EpicGames.Horde/Acls/AclAction.cs
--- a/Engine/Source/Programs/Shared/EpicGames.Horde/Acls/AclAction.cs
+++ b/Engine/Source/Programs/Shared/EpicGames.Horde/Acls/AclAction.cs
@@ -28,7 +28,7 @@ namespace EpicGames.Horde.Acls
 	sealed class AclActionJsonConverter : JsonConverter<AclAction>
 	{
 		/// <inheritdoc/>
-		public override AclAction Read(ref Utf8JsonReader reader, Type typeToConvert, JsonSerializerOptions options) => new AclAction(reader.GetString()!);
+		public override AclAction Read(ref Utf8JsonReader reader, Type typeToConvert, JsonSerializerOptions options) => new(reader.GetString()!);
 
 		/// <inheritdoc/>
 		public override void Write(Utf8JsonWriter writer, AclAction value, JsonSerializerOptions options) => writer.WriteStringValue(value.Name);
diff --git a/Engine/Source/Programs/Shared/EpicGames.Horde/Acls/AclScopeName.cs b/Engine/Source/Programs/Shared/EpicGames.Horde/Acls/AclScopeName.cs
--- a/Engine/Source/Programs/Shared/EpicGames.Horde/Acls/AclScopeName.cs
+++ b/Engine/Source/Programs/Shared/EpicGames.Horde/Acls/AclScopeName.cs
@@ -48,7 +48,7 @@ namespace EpicGames.Horde.Acls
 	class AclScopeNameJsonConverter : JsonConverter<AclScopeName>
 	{
 		/// <inheritdoc/>
-		public override AclScopeName Read(ref Utf8JsonReader reader, Type typeToConvert, JsonSerializerOptions options) => new AclScopeName(reader.GetString() ?? String.Empty);
+		public override AclScopeName Read(ref Utf8JsonReader reader, Type typeToConvert, JsonSerializerOptions options) => new(reader.GetString() ?? String.Empty);
 
 		/// <inheritdoc/>
 		public override void Write(Utf8JsonWriter writer, AclScopeName value, JsonSerializerOptions options) => writer.WriteStringValue(value.Text);
diff --git a/Engine/Source/Programs/Shared/EpicGames.Horde/Agents/AgentId.cs b/Engine/Source/Programs/Shared/EpicGames.Horde/Agents/AgentId.cs
--- a/Engine/Source/Programs/Shared/EpicGames.Horde/Agents/AgentId.cs
+++ b/Engine/Source/Programs/Shared/EpicGames.Horde/Agents/AgentId.cs
@@ -165,7 +165,7 @@ namespace EpicGames.Horde.Agents
 	public sealed class AgentIdJsonConverter : JsonConverter<AgentId>
 	{
 		/// <inheritdoc/>
-		public override AgentId Read(ref Utf8JsonReader reader, Type typeToConvert, JsonSerializerOptions options) => new AgentId(reader.GetString() ?? String.Empty);
+		public override AgentId Read(ref Utf8JsonReader reader, Type typeToConvert, JsonSerializerOptions options) => new(reader.GetString() ?? String.Empty);
 
 		/// <inheritdoc/>
 		public override void Write(Utf8JsonWriter writer, AgentId value, JsonSerializerOptions options) => writer.WriteStringValue(value.ToString());
diff --git a/Engine/Source/Programs/Shared/EpicGames.Horde/Agents/Leases/LeaseId.cs b/Engine/Source/Programs/Shared/EpicGames.Horde/Agents/Leases/LeaseId.cs
--- a/Engine/Source/Programs/Shared/EpicGames.Horde/Agents/Leases/LeaseId.cs
+++ b/Engine/Source/Programs/Shared/EpicGames.Horde/Agents/Leases/LeaseId.cs
@@ -32,7 +32,7 @@ namespace EpicGames.Horde.Agents.Leases
 		}
 
 		/// <inheritdoc cref="BinaryId.Parse(System.String)"/>
-		public static LeaseId Parse(string text) => new LeaseId(BinaryId.Parse(text));
+		public static LeaseId Parse(string text) => new(BinaryId.Parse(text));
 
 		/// <inheritdoc/>
 		public override string ToString() => Id.ToString();
@@ -44,7 +44,7 @@ namespace EpicGames.Horde.Agents.Leases
 	class LeaseIdConverter : BinaryIdConverter<LeaseId>
 	{
 		/// <inheritdoc/>
-		public override LeaseId FromBinaryId(BinaryId id) => new LeaseId(id);
+		public override LeaseId FromBinaryId(BinaryId id) => new(id);
 
 		/// <inheritdoc/>
 		public override BinaryId ToBinaryId(LeaseId value) => value.Id;
diff --git a/Engine/Source/Programs/Shared/EpicGames.Horde/Agents/Pools/PoolId.cs b/Engine/Source/Programs/Shared/EpicGames.Horde/Agents/Pools/PoolId.cs
--- a/Engine/Source/Programs/Shared/EpicGames.Horde/Agents/Pools/PoolId.cs
+++ b/Engine/Source/Programs/Shared/EpicGames.Horde/Agents/Pools/PoolId.cs
@@ -37,7 +37,7 @@ namespace EpicGames.Horde.Agents.Pools
 	class PoolIdConverter : StringIdConverter<PoolId>
 	{
 		/// <inheritdoc/>
-		public override PoolId FromStringId(StringId id) => new PoolId(id);
+		public override PoolId FromStringId(StringId id) => new(id);
 
 		/// <inheritdoc/>
 		public override StringId ToStringId(PoolId value) => value.Id;
diff --git a/Engine/Source/Programs/Shared/EpicGames.Horde/Agents/Sessions/SessionId.cs b/Engine/Source/Programs/Shared/EpicGames.Horde/Agents/Sessions/SessionId.cs
--- a/Engine/Source/Programs/Shared/EpicGames.Horde/Agents/Sessions/SessionId.cs
+++ b/Engine/Source/Programs/Shared/EpicGames.Horde/Agents/Sessions/SessionId.cs
@@ -21,7 +21,7 @@ namespace EpicGames.Horde.Agents.Sessions
 		public static SessionId Empty { get; } = default;
 
 		/// <inheritdoc cref="BinaryId.Parse(System.String)"/>
-		public static SessionId Parse(string text) => new SessionId(BinaryId.Parse(text));
+		public static SessionId Parse(string text) => new(BinaryId.Parse(text));
 
 		/// <inheritdoc cref="BinaryId.TryParse(System.String, out BinaryId)"/>
 		public static bool TryParse(string text, out SessionId id)
@@ -48,7 +48,7 @@ namespace EpicGames.Horde.Agents.Sessions
 	class SessionIdConverter : BinaryIdConverter<SessionId>
 	{
 		/// <inheritdoc/>
-		public override SessionId FromBinaryId(BinaryId id) => new SessionId(id);
+		public override SessionId FromBinaryId(BinaryId id) => new(id);
 
 		/// <inheritdoc/>
 		public override BinaryId ToBinaryId(SessionId value) => value.Id;
diff --git a/Engine/Source/Programs/Shared/EpicGames.Horde/Artifacts/ArtifactDescriptor.cs b/Engine/Source/Programs/Shared/EpicGames.Horde/Artifacts/ArtifactDescriptor.cs
--- a/Engine/Source/Programs/Shared/EpicGames.Horde/Artifacts/ArtifactDescriptor.cs
+++ b/Engine/Source/Programs/Shared/EpicGames.Horde/Artifacts/ArtifactDescriptor.cs
@@ -142,7 +142,7 @@ namespace EpicGames.Horde.Storage
 		/// <returns>Data for the serialized descriptor</returns>
 		public byte[] Serialize()
 		{
-			JsonSerializerOptions serializerOptions = new JsonSerializerOptions(HordeHttpClient.JsonSerializerOptions);
+			JsonSerializerOptions serializerOptions = new(HordeHttpClient.JsonSerializerOptions);
 			serializerOptions.WriteIndented = true;
 
 			return JsonSerializer.SerializeToUtf8Bytes<ArtifactDescriptor>(this, serializerOptions);
diff --git a/Engine/Source/Programs/Shared/EpicGames.Horde/Artifacts/ArtifactId.cs b/Engine/Source/Programs/Shared/EpicGames.Horde/Artifacts/ArtifactId.cs
--- a/Engine/Source/Programs/Shared/EpicGames.Horde/Artifacts/ArtifactId.cs
+++ b/Engine/Source/Programs/Shared/EpicGames.Horde/Artifacts/ArtifactId.cs
@@ -16,7 +16,7 @@ namespace EpicGames.Horde.Artifacts
 	public readonly record struct ArtifactId(BinaryId Id)
 	{
 		/// <inheritdoc cref="BinaryId.Parse(System.String)"/>
-		public static ArtifactId Parse(string text) => new ArtifactId(BinaryId.Parse(text));
+		public static ArtifactId Parse(string text) => new(BinaryId.Parse(text));
 
 		/// <inheritdoc/>
 		public override readonly string ToString() => Id.ToString();
@@ -28,7 +28,7 @@ namespace EpicGames.Horde.Artifacts
 	class ArtifactIdConverter : BinaryIdConverter<ArtifactId>
 	{
 		/// <inheritdoc/>
-		public override ArtifactId FromBinaryId(BinaryId id) => new ArtifactId(id);
+		public override ArtifactId FromBinaryId(BinaryId id) => new(id);
 
 		/// <inheritdoc/>
 		public override BinaryId ToBinaryId(ArtifactId value) => value.Id;
diff --git a/Engine/Source/Programs/Shared/EpicGames.Horde/Artifacts/ArtifactName.cs b/Engine/Source/Programs/Shared/EpicGames.Horde/Artifacts/ArtifactName.cs
--- a/Engine/Source/Programs/Shared/EpicGames.Horde/Artifacts/ArtifactName.cs
+++ b/Engine/Source/Programs/Shared/EpicGames.Horde/Artifacts/ArtifactName.cs
@@ -30,7 +30,7 @@ namespace EpicGames.Horde.Artifacts
 	class ArtifactNameConverter : StringIdConverter<ArtifactName>
 	{
 		/// <inheritdoc/>
-		public override ArtifactName FromStringId(StringId id) => new ArtifactName(id);
+		public override ArtifactName FromStringId(StringId id) => new(id);
 
 		/// <inheritdoc/>
 		public override StringId ToStringId(ArtifactName value) => value.Id;
diff --git a/Engine/Source/Programs/Shared/EpicGames.Horde/Artifacts/ArtifactType.cs b/Engine/Source/Programs/Shared/EpicGames.Horde/Artifacts/ArtifactType.cs
--- a/Engine/Source/Programs/Shared/EpicGames.Horde/Artifacts/ArtifactType.cs
+++ b/Engine/Source/Programs/Shared/EpicGames.Horde/Artifacts/ArtifactType.cs
@@ -57,7 +57,7 @@ namespace EpicGames.Horde.Artifacts
 	class ArtifactTypeConverter : StringIdConverter<ArtifactType>
 	{
 		/// <inheritdoc/>
-		public override ArtifactType FromStringId(StringId id) => new ArtifactType(id);
+		public override ArtifactType FromStringId(StringId id) => new(id);
 
 		/// <inheritdoc/>
 		public override StringId ToStringId(ArtifactType value) => value.Id;
diff --git a/Engine/Source/Programs/Shared/EpicGames.Horde/Artifacts/HttpArtifactCollection.cs b/Engine/Source/Programs/Shared/EpicGames.Horde/Artifacts/HttpArtifactCollection.cs
--- a/Engine/Source/Programs/Shared/EpicGames.Horde/Artifacts/HttpArtifactCollection.cs
+++ b/Engine/Source/Programs/Shared/EpicGames.Horde/Artifacts/HttpArtifactCollection.cs
@@ -129,7 +129,7 @@ namespace EpicGames.Horde.Artifacts
 		{
 			HordeHttpClient hordeHttpClient = _hordeClient.CreateHttpClient();
 			CreateArtifactResponse response = await hordeHttpClient.CreateArtifactAsync(name, type, description, streamId, commitId, keys, metadata, cancellationToken);
-			Artifact artifact = new Artifact(this, response.ArtifactId, name, type, description, streamId, CommitIdWithOrder.FromPerforceChange(commitId.GetPerforceChange()), [.. keys], [.. metadata], response.NamespaceId, response.RefName, DateTime.UtcNow);
+			Artifact artifact = new(this, response.ArtifactId, name, type, description, streamId, CommitIdWithOrder.FromPerforceChange(commitId.GetPerforceChange()), [.. keys], [.. metadata], response.NamespaceId, response.RefName, DateTime.UtcNow);
 			return new ArtifactBuilder(this, artifact, response);
 		}
 
diff --git a/Engine/Source/Programs/Shared/EpicGames.Horde/Auth/HordeAuthConfigProvider.cs b/Engine/Source/Programs/Shared/EpicGames.Horde/Auth/HordeAuthConfigProvider.cs
--- a/Engine/Source/Programs/Shared/EpicGames.Horde/Auth/HordeAuthConfigProvider.cs
+++ b/Engine/Source/Programs/Shared/EpicGames.Horde/Auth/HordeAuthConfigProvider.cs
@@ -38,12 +38,12 @@ internal sealed class HordeAuthConfigProvider(
 		}
 
 		GetAuthConfigResponse? authConfig;
-		using (HttpClient httpClient = new HttpClient(_httpMessageHandler, false))
+		using (HttpClient httpClient = new(_httpMessageHandler, false))
 		{
 			httpClient.BaseAddress = _serverUri;
 			_logger.LogDebug("Retrieving auth configuration for {Server}", _serverUri);
 
-			JsonSerializerOptions jsonOptions = new JsonSerializerOptions();
+			JsonSerializerOptions jsonOptions = new();
 			HordeHttpClient.ConfigureJsonSerializer(jsonOptions);
 
 			authConfig = await httpClient.GetFromJsonAsync<GetAuthConfigResponse>("api/v1/server/auth", jsonOptions, cancellationToken);
@@ -67,5 +67,5 @@ internal sealed class HordeAuthConfigProvider(
 	private readonly HttpMessageHandler _httpMessageHandler = httpMessageHandler;
 	private readonly Uri _serverUri = serverUri;
 	private readonly ILogger<HordeAuthConfigProvider> _logger = logger;
-	private readonly MemoryCache _memoryCache = new MemoryCache(Options.Create(new MemoryCacheOptions()));
+	private readonly MemoryCache _memoryCache = new(Options.Create(new MemoryCacheOptions()));
 }
diff --git a/Engine/Source/Programs/Shared/EpicGames.Horde/Auth/HordeStaticAccessTokenProvider.cs b/Engine/Source/Programs/Shared/EpicGames.Horde/Auth/HordeStaticAccessTokenProvider.cs
--- a/Engine/Source/Programs/Shared/EpicGames.Horde/Auth/HordeStaticAccessTokenProvider.cs
+++ b/Engine/Source/Programs/Shared/EpicGames.Horde/Auth/HordeStaticAccessTokenProvider.cs
@@ -22,7 +22,7 @@ internal sealed class HordeStaticAccessTokenProvider(IOptionsSnapshot<HordeOptio
 			string? hordeUrlEnvVar = Environment.GetEnvironmentVariable(HordeHttpClient.HordeUrlEnvVarName);
 			if (!String.IsNullOrEmpty(hordeUrlEnvVar))
 			{
-				Uri hordeUrl = new Uri(hordeUrlEnvVar);
+				Uri hordeUrl = new(hordeUrlEnvVar);
 				if (_options.Value.ServerUrl == null || String.Equals(_options.Value.ServerUrl.Host, hordeUrl.Host, StringComparison.OrdinalIgnoreCase))
 				{
 					string? hordeToken = Environment.GetEnvironmentVariable(HordeHttpClient.HordeTokenEnvVarName);
diff --git a/Engine/Source/Programs/Shared/EpicGames.Horde/Auth/OAuthHandler.cs b/Engine/Source/Programs/Shared/EpicGames.Horde/Auth/OAuthHandler.cs
--- a/Engine/Source/Programs/Shared/EpicGames.Horde/Auth/OAuthHandler.cs
+++ b/Engine/Source/Programs/Shared/EpicGames.Horde/Auth/OAuthHandler.cs
@@ -115,7 +115,7 @@ namespace EpicGames.Horde.Auth
 
 			try
 			{
-				using HttpRequestMessage message = new HttpRequestMessage(HttpMethod.Post, _options.AuthUrl);
+				using HttpRequestMessage message = new(HttpMethod.Post, _options.AuthUrl);
 				message.Content = new FormUrlEncodedContent(content);
 
 				HttpResponseMessage response = await _client.SendAsync(message, cancellationToken);
@@ -164,6 +164,6 @@ namespace EpicGames.Horde.Auth
 		/// </summary>
 		/// <param name="options"></param>
 		/// <returns></returns>
-		public OAuthHandler<T> Create<T>(IOAuthOptions options) => new OAuthHandler<T>(_httpClient, options);
+		public OAuthHandler<T> Create<T>(IOAuthOptions options) => new(_httpClient, options);
 	}
 }
diff --git a/Engine/Source/Programs/Shared/EpicGames.Horde/BinaryIdConverter.cs b/Engine/Source/Programs/Shared/EpicGames.Horde/BinaryIdConverter.cs
--- a/Engine/Source/Programs/Shared/EpicGames.Horde/BinaryIdConverter.cs
+++ b/Engine/Source/Programs/Shared/EpicGames.Horde/BinaryIdConverter.cs
@@ -50,7 +50,7 @@ namespace EpicGames.Horde
 	/// </summary>
 	public sealed class BinaryIdCbConverter<TValue, TConverter> : CbConverter<TValue> where TValue : struct where TConverter : BinaryIdConverter<TValue>, new()
 	{
-		readonly TConverter _converter = new TConverter();
+		readonly TConverter _converter = new();
 
 		/// <inheritdoc/>
 		public override TValue Read(CbField field)
@@ -76,7 +76,7 @@ namespace EpicGames.Horde
 	/// </summary>
 	public sealed class BinaryIdTypeConverter<TValue, TConverter> : TypeConverter where TValue : struct where TConverter : BinaryIdConverter<TValue>, new()
 	{
-		readonly TConverter _converter = new TConverter();
+		readonly TConverter _converter = new();
 
 		/// <inheritdoc/>
 		public override bool CanConvertFrom(ITypeDescriptorContext? context, Type sourceType)
@@ -121,7 +121,7 @@ namespace EpicGames.Horde
 	/// </summary>
 	public sealed class BinaryIdJsonConverter<TValue, TConverter> : JsonConverter<TValue> where TValue : struct where TConverter : BinaryIdConverter<TValue>, new()
 	{
-		readonly TConverter _converter = new TConverter();
+		readonly TConverter _converter = new();
 
 		/// <inheritdoc/>
 		public override TValue Read(ref Utf8JsonReader reader, Type typeToConvert, JsonSerializerOptions options) => _converter.FromBinaryId(BinaryId.Parse(reader.GetUtf8String()));
diff --git a/Engine/Source/Programs/Shared/EpicGames.Horde/BuildHealth/BuildHealthFilterId.cs b/Engine/Source/Programs/Shared/EpicGames.Horde/BuildHealth/BuildHealthFilterId.cs
--- a/Engine/Source/Programs/Shared/EpicGames.Horde/BuildHealth/BuildHealthFilterId.cs
+++ b/Engine/Source/Programs/Shared/EpicGames.Horde/BuildHealth/BuildHealthFilterId.cs
@@ -20,7 +20,7 @@ namespace EpicGames.Horde.BuildHealth
 		public static BuildHealthFilterId Empty { get; } = default;
 
 		/// <inheritdoc cref="BinaryId.Parse(System.String)"/>
-		public static BuildHealthFilterId Parse(string text) => new BuildHealthFilterId(BinaryId.Parse(text));
+		public static BuildHealthFilterId Parse(string text) => new(BinaryId.Parse(text));
 
 		/// <inheritdoc/>
 		public override string ToString() => Id.ToString();
@@ -32,7 +32,7 @@ namespace EpicGames.Horde.BuildHealth
 	class FilterIdConverter : BinaryIdConverter<BuildHealthFilterId>
 	{
 		/// <inheritdoc/>
-		public override BuildHealthFilterId FromBinaryId(BinaryId id) => new BuildHealthFilterId(id);
+		public override BuildHealthFilterId FromBinaryId(BinaryId id) => new(id);
 
 		/// <inheritdoc/>
 		public override BinaryId ToBinaryId(BuildHealthFilterId value) => value.Id;
diff --git a/Engine/Source/Programs/Shared/EpicGames.Horde/Commits/CommitId.cs b/Engine/Source/Programs/Shared/EpicGames.Horde/Commits/CommitId.cs
--- a/Engine/Source/Programs/Shared/EpicGames.Horde/Commits/CommitId.cs
+++ b/Engine/Source/Programs/Shared/EpicGames.Horde/Commits/CommitId.cs
@@ -38,7 +38,7 @@ namespace EpicGames.Horde.Commits
 		/// Creates a commit it from a Perforce changelist number. Temporary helper method for migration purposes.
 		/// </summary>
 		public static CommitId FromPerforceChange(int change)
-			=> new CommitId($"{change}");
+			=> new($"{change}");
 
 		/// <summary>
 		/// Creates a commit it from a Perforce changelist number. Temporary helper method for migration purposes.
@@ -157,7 +157,7 @@ namespace EpicGames.Horde.Commits
 		/// Creates a commit it from a Perforce changelist number. Temporary helper method for migration purposes.
 		/// </summary>
 		public static new CommitIdWithOrder FromPerforceChange(int change)
-			=> new CommitIdWithOrder($"{change}", change);
+			=> new($"{change}", change);
 
 		/// <summary>
 		/// Creates a commit it from a Perforce changelist number. Temporary helper method for migration purposes.
diff --git a/Engine/Source/Programs/Shared/EpicGames.Horde/Common/Condition.cs b/Engine/Source/Programs/Shared/EpicGames.Horde/Common/Condition.cs
--- a/Engine/Source/Programs/Shared/EpicGames.Horde/Common/Condition.cs
+++ b/Engine/Source/Programs/Shared/EpicGames.Horde/Common/Condition.cs
@@ -43,7 +43,7 @@ namespace EpicGames.Horde.Common
 		public int Offset { get; private set; }
 		public int Length { get; private set; }
 
-		public StringView Token => new StringView(Input, Offset, Length);
+		public StringView Token => new(Input, Offset, Length);
 		public TokenType Type { get; private set; }
 		public string? Scalar { get; private set; }
 
@@ -218,7 +218,7 @@ namespace EpicGames.Horde.Common
 			return true;
 		}
 
-		static readonly Dictionary<StringView, ulong> s_sizeSuffixes = new Dictionary<StringView, ulong>(StringViewComparer.OrdinalIgnoreCase)
+		static readonly Dictionary<StringView, ulong> s_sizeSuffixes = new(StringViewComparer.OrdinalIgnoreCase)
 		{
 			["kb"] = 1024UL,
 			["Mb"] = 1024UL * 1024,
@@ -245,7 +245,7 @@ namespace EpicGames.Horde.Common
 					endIdx++;
 				}
 
-				StringView suffix = new StringView(Input, offset, endIdx - offset);
+				StringView suffix = new(Input, offset, endIdx - offset);
 
 				ulong size;
 				if (!s_sizeSuffixes.TryGetValue(suffix, out size))
@@ -343,7 +343,7 @@ namespace EpicGames.Horde.Common
 		{
 			Text = text;
 
-			TokenReader reader = new TokenReader(text);
+			TokenReader reader = new(text);
 			if (reader.Type != TokenType.End)
 			{
 				Error = ParseOrExpr(reader);
@@ -377,7 +377,7 @@ namespace EpicGames.Horde.Common
 		/// <returns>The new condition object</returns>
 		public static Condition Parse(string text)
 		{
-			Condition condition = new Condition(text);
+			Condition condition = new(text);
 			if (!condition.IsValid())
 			{
 				throw new ConditionException(condition.Error!);
diff --git a/Engine/Source/Programs/Shared/EpicGames.Horde/Compute/AgentManagedProcess.cs b/Engine/Source/Programs/Shared/EpicGames.Horde/Compute/AgentManagedProcess.cs
--- a/Engine/Source/Programs/Shared/EpicGames.Horde/Compute/AgentManagedProcess.cs
+++ b/Engine/Source/Programs/Shared/EpicGames.Horde/Compute/AgentManagedProcess.cs
@@ -16,7 +16,7 @@ namespace EpicGames.Horde.Compute
 	{
 		readonly Channel<string> _output;
 		readonly BackgroundTask _backgroundTask;
-		readonly TaskCompletionSource<int> _result = new TaskCompletionSource<int>(TaskCreationOptions.RunContinuationsAsynchronously);
+		readonly TaskCompletionSource<int> _result = new(TaskCreationOptions.RunContinuationsAsynchronously);
 
 		byte[] _buffer = new byte[1024];
 		int _bufferLength;
diff --git a/Engine/Source/Programs/Shared/EpicGames.Horde/Compute/AgentMessage.cs b/Engine/Source/Programs/Shared/EpicGames.Horde/Compute/AgentMessage.cs
--- a/Engine/Source/Programs/Shared/EpicGames.Horde/Compute/AgentMessage.cs
+++ b/Engine/Source/Programs/Shared/EpicGames.Horde/Compute/AgentMessage.cs
@@ -460,7 +460,7 @@ namespace EpicGames.Horde.Compute
 				path = path.Substring(atIdx + 1);
 			}
 
-			BlobLocator locator = new BlobLocator(path);
+			BlobLocator locator = new(path);
 			return new UploadFilesMessage(name, locator);
 		}
 
@@ -640,7 +640,7 @@ namespace EpicGames.Horde.Compute
 		/// <returns></returns>
 		public static ReadBlobMessage ParseReadBlobRequest(this AgentMessage message)
 		{
-			BlobLocator locator = new BlobLocator(message.ReadUtf8String());
+			BlobLocator locator = new(message.ReadUtf8String());
 			int offset = (int)message.ReadUnsignedVarInt();
 			int length = (int)message.ReadUnsignedVarInt();
 			return new ReadBlobMessage(locator, offset, length);
diff --git a/Engine/Source/Programs/Shared/EpicGames.Horde/Compute/AgentMessageChannel.cs b/Engine/Source/Programs/Shared/EpicGames.Horde/Compute/AgentMessageChannel.cs
--- a/Engine/Source/Programs/Shared/EpicGames.Horde/Compute/AgentMessageChannel.cs
+++ b/Engine/Source/Programs/Shared/EpicGames.Horde/Compute/AgentMessageChannel.cs
@@ -164,7 +164,7 @@ namespace EpicGames.Horde.Compute
 				}
 
 				AgentMessageType type = (AgentMessageType)memory.Span[0];
-				AgentMessage message = new AgentMessage(type, memory.Slice(HeaderLength, messageLength));
+				AgentMessage message = new(type, memory.Slice(HeaderLength, messageLength));
 				// if (_logger.IsEnabled(LogLevel.Trace))
 				// {
 				// 	LogMessageInfo("RECV", message.Type, message.Data.Span);
@@ -179,7 +179,7 @@ namespace EpicGames.Horde.Compute
 		[SuppressMessage("CodeQuality", "IDE0051:Remove unused private members", Justification = "Log calls disabled")]
 		void LogMessageInfo(string verb, AgentMessageType type, ReadOnlySpan<byte> data)
 		{
-			StringBuilder bytes = new StringBuilder();
+			StringBuilder bytes = new();
 			for (int offset = 0; offset < 16 && offset < data.Length; offset++)
 			{
 				bytes.Append($" {data[offset]:X2}");
@@ -234,7 +234,7 @@ namespace EpicGames.Horde.Compute
 		/// <param name="socket">Socket to create a channel for</param>
 		/// <param name="channelId">Identifier for the channel</param>
 		/// <param name="sendBufferSize">Size of the send buffer</param>
-		/// <param name="recvBufferSize">Size of the recieve buffer</param>
+		/// <param name="recvBufferSize">Size of the receive buffer</param>
 		public static AgentMessageChannel CreateAgentMessageChannel(this ComputeSocket socket, int channelId, int sendBufferSize, int recvBufferSize)
 		{
 			using ComputeBuffer sendBuffer = new PooledBuffer(sendBufferSize);
diff --git a/Engine/Source/Programs/Shared/EpicGames.Horde/Compute/AgentMessageHandler.cs b/Engine/Source/Programs/Shared/EpicGames.Horde/Compute/AgentMessageHandler.cs
--- a/Engine/Source/Programs/Shared/EpicGames.Horde/Compute/AgentMessageHandler.cs
+++ b/Engine/Source/Programs/Shared/EpicGames.Horde/Compute/AgentMessageHandler.cs
@@ -183,11 +183,11 @@ namespace EpicGames.Horde.Compute
 
 		async Task WriteFilesAsync(AgentMessageChannel channel, string path, BlobLocator locator, BlobSerializerOptions? options = null, CancellationToken cancellationToken = default)
 		{
-			using AgentStorageBackend innerStore = new AgentStorageBackend(channel);
-			await using BundleCache cache = new BundleCache(new BundleCacheOptions { HeaderCacheSize = 10 * 1024 * 1024, PacketCacheSize = 128 * 1024 * 1024 });
+			using AgentStorageBackend innerStore = new(channel);
+			await using BundleCache cache = new(new BundleCacheOptions { HeaderCacheSize = 10 * 1024 * 1024, PacketCacheSize = 128 * 1024 * 1024 });
 
 			BundleOptions bundleOptions = ComputeProtocolUtilities.GetBundleOptions(channel.Protocol);
-			BundleStorageNamespace store = new BundleStorageNamespace(innerStore, cache, bundleOptions, _logger);
+			BundleStorageNamespace store = new(innerStore, cache, bundleOptions, _logger);
 
 			IBlobRef<DirectoryNode> directoryRef = store.CreateBlobRef<DirectoryNode>(locator, options);
 
@@ -246,7 +246,7 @@ namespace EpicGames.Horde.Compute
 
 		void DeleteFiles(IReadOnlyList<string> deleteFiles)
 		{
-			FileFilter filter = new FileFilter(deleteFiles);
+			FileFilter filter = new(deleteFiles);
 
 			List<FileReference> files = filter.ApplyToDirectory(_sandboxDir, false);
 			foreach (FileReference file in files)
@@ -294,7 +294,7 @@ namespace EpicGames.Horde.Compute
 
 		async Task ExecuteProcessWindowsAsync(ComputeSocket socket, AgentMessageChannel channel, string executable, IReadOnlyList<string> arguments, string? workingDir, IReadOnlyDictionary<string, string?>? envVars, ExecuteProcessFlags flags, CancellationToken cancellationToken)
 		{
-			Dictionary<string, string?> newEnvVars = new Dictionary<string, string?>(_envVars);
+			Dictionary<string, string?> newEnvVars = new(_envVars);
 			if (envVars != null)
 			{
 				foreach ((string name, string? value) in envVars)
@@ -328,7 +328,7 @@ namespace EpicGames.Horde.Compute
 				while (await ipcReader.WaitToReadAsync(1, cancellationToken))
 				{
 					ReadOnlyMemory<byte> memory = ipcReader.GetReadBuffer();
-					MemoryReader reader = new MemoryReader(memory);
+					MemoryReader reader = new(memory);
 
 					IpcMessage message = (IpcMessage)reader.ReadUnsignedVarInt();
 					try
@@ -388,7 +388,7 @@ namespace EpicGames.Horde.Compute
 		class ProcessOutputWriter
 		{
 			readonly string _prefix;
-			readonly ByteArrayBuilder _lineBuffer = new ByteArrayBuilder();
+			readonly ByteArrayBuilder _lineBuffer = new();
 			readonly ILogger _logger;
 
 			public ProcessOutputWriter(string prefix, ILogger logger)
@@ -452,9 +452,9 @@ namespace EpicGames.Horde.Compute
 
 				_logger.LogWarning("Note: Loading and running {Assembly} in process", assemblyPath);
 
-				TaskCompletionSource<int> resultTcs = new TaskCompletionSource<int>(TaskCreationOptions.RunContinuationsAsynchronously);
+				TaskCompletionSource<int> resultTcs = new(TaskCreationOptions.RunContinuationsAsynchronously);
 
-				Thread thread = new Thread(() => resultTcs.SetResult(AppDomain.CurrentDomain.ExecuteAssembly(assemblyPath, mainArgs)));
+				Thread thread = new(() => resultTcs.SetResult(AppDomain.CurrentDomain.ExecuteAssembly(assemblyPath, mainArgs)));
 				thread.Start();
 
 				int result = await resultTcs.Task;
@@ -556,11 +556,11 @@ namespace EpicGames.Horde.Compute
 					_logger.LogWarning("Working dir {Path} does not exist", resolvedWorkingDir);
 				}
 
-				using ManagedProcessGroup group = new ManagedProcessGroup();
-				using ManagedProcess process = new ManagedProcess(group, resolvedExecutable, resolvedCommandLine, resolvedWorkingDir, resolvedEnvVars, null, ProcessPriorityClass.Normal);
+				using ManagedProcessGroup group = new();
+				using ManagedProcess process = new(group, resolvedExecutable, resolvedCommandLine, resolvedWorkingDir, resolvedEnvVars, null, ProcessPriorityClass.Normal);
 				byte[] buffer = new byte[1024];
 
-				ProcessOutputWriter outputWriter = new ProcessOutputWriter($"{Path.GetFileNameWithoutExtension(resolvedExecutable)}> ", _logger);
+				ProcessOutputWriter outputWriter = new($"{Path.GetFileNameWithoutExtension(resolvedExecutable)}> ", _logger);
 				for (; ; )
 				{
 					// Use WaitAsync() as ReadAsync() does not respect the cancellation token when reading
diff --git a/Engine/Source/Programs/Shared/EpicGames.Horde/Compute/Buffers/PooledBuffer.cs b/Engine/Source/Programs/Shared/EpicGames.Horde/Compute/Buffers/PooledBuffer.cs
--- a/Engine/Source/Programs/Shared/EpicGames.Horde/Compute/Buffers/PooledBuffer.cs
+++ b/Engine/Source/Programs/Shared/EpicGames.Horde/Compute/Buffers/PooledBuffer.cs
@@ -55,8 +55,8 @@ namespace EpicGames.Horde.Compute.Buffers
 		GCHandle _headerHandle;
 		IMemoryOwner<byte>[] _chunkOwners;
 
-		readonly AsyncEvent _writerEvent = new AsyncEvent();
-		readonly AsyncEvent _readerEvent = new AsyncEvent();
+		readonly AsyncEvent _writerEvent = new();
+		readonly AsyncEvent _readerEvent = new();
 
 		internal PooledBufferDetail(HeaderPtr headerPtr, GCHandle headerHandle, Memory<byte>[] chunks, IMemoryOwner<byte>[] chunkOwners)
 			: base(headerPtr, chunks)
@@ -69,7 +69,7 @@ namespace EpicGames.Horde.Compute.Buffers
 		{
 			byte[] header = new byte[HeaderSize];
 			GCHandle headerHandle = GCHandle.Alloc(header, GCHandleType.Pinned);
-			HeaderPtr headerPtr = new HeaderPtr((ulong*)headerHandle.AddrOfPinnedObject().ToPointer(), numReaders, numChunks, chunkLength);
+			HeaderPtr headerPtr = new((ulong*)headerHandle.AddrOfPinnedObject().ToPointer(), numReaders, numChunks, chunkLength);
 
 			Memory<byte>[] chunks = new Memory<byte>[numChunks];
 			IMemoryOwner<byte>[] chunkOwners = new IMemoryOwner<byte>[numChunks];
diff --git a/Engine/Source/Programs/Shared/EpicGames.Horde/Compute/Buffers/SharedMemoryBuffer.cs b/Engine/Source/Programs/Shared/EpicGames.Horde/Compute/Buffers/SharedMemoryBuffer.cs
--- a/Engine/Source/Programs/Shared/EpicGames.Horde/Compute/Buffers/SharedMemoryBuffer.cs
+++ b/Engine/Source/Programs/Shared/EpicGames.Horde/Compute/Buffers/SharedMemoryBuffer.cs
@@ -50,13 +50,13 @@ namespace EpicGames.Horde.Compute.Buffers
 		/// <param name="name">Name of the buffer</param>
 		/// <param name="numChunks">Number of chunks in the buffer</param>
 		/// <param name="chunkLength">Length of each chunk</param>
-		public static unsafe SharedMemoryBuffer CreateNew(string? name, int numChunks, int chunkLength) => new SharedMemoryBuffer(SharedMemoryBufferDetail.CreateNew(name, numChunks, chunkLength));
+		public static unsafe SharedMemoryBuffer CreateNew(string? name, int numChunks, int chunkLength) => new(SharedMemoryBufferDetail.CreateNew(name, numChunks, chunkLength));
 
 		/// <summary>
 		/// Open an existing buffer by name
 		/// </summary>
 		/// <param name="name">Name of the buffer to open</param>
-		public static unsafe SharedMemoryBuffer OpenExisting(string name) => new SharedMemoryBuffer(SharedMemoryBufferDetail.OpenExisting(name));
+		public static unsafe SharedMemoryBuffer OpenExisting(string name) => new(SharedMemoryBufferDetail.OpenExisting(name));
 	}
 
 	/// <summary>
@@ -100,12 +100,12 @@ namespace EpicGames.Horde.Compute.Buffers
 
 			MemoryMappedFile memoryMappedFile = MemoryMappedFile.CreateNew($"{name}_M", capacity, MemoryMappedFileAccess.ReadWrite, MemoryMappedFileOptions.None, HandleInheritability.Inheritable);
 			MemoryMappedViewAccessor memoryMappedViewAccessor = memoryMappedFile.CreateViewAccessor();
-			MemoryMappedView memoryMappedView = new MemoryMappedView(memoryMappedViewAccessor);
+			MemoryMappedView memoryMappedView = new(memoryMappedViewAccessor);
 
 			Native.EventHandle writerEvent = Native.EventHandle.CreateNew($"{name}_W", EventResetMode.AutoReset, true, HandleInheritability.Inheritable);
 			Native.EventHandle readerEvent = Native.EventHandle.CreateNew($"{name}_R0", EventResetMode.AutoReset, true, HandleInheritability.Inheritable);
 
-			HeaderPtr headerPtr = new HeaderPtr((ulong*)memoryMappedView.GetPointer(), 1, numChunks, chunkLength);
+			HeaderPtr headerPtr = new((ulong*)memoryMappedView.GetPointer(), 1, numChunks, chunkLength);
 			Memory<byte>[] chunks = CreateChunks(headerPtr.NumChunks, headerPtr.ChunkLength, memoryMappedView);
 
 			return new SharedMemoryBufferDetail(name, headerPtr, chunks, memoryMappedFile, memoryMappedViewAccessor, memoryMappedView, readerEvent, writerEvent);
@@ -124,12 +124,12 @@ namespace EpicGames.Horde.Compute.Buffers
 
 			MemoryMappedFile memoryMappedFile = MemoryMappedFile.OpenExisting($"{name}_M");
 			MemoryMappedViewAccessor memoryMappedViewAccessor = memoryMappedFile.CreateViewAccessor();
-			MemoryMappedView memoryMappedView = new MemoryMappedView(memoryMappedViewAccessor);
+			MemoryMappedView memoryMappedView = new(memoryMappedViewAccessor);
 
 			Native.EventHandle readerEvent = Native.EventHandle.OpenExisting($"{name}_R0");
 			Native.EventHandle writerEvent = Native.EventHandle.OpenExisting($"{name}_W");
 
-			HeaderPtr headerPtr = new HeaderPtr((ulong*)memoryMappedView.GetPointer());
+			HeaderPtr headerPtr = new((ulong*)memoryMappedView.GetPointer());
 			Memory<byte>[] chunks = CreateChunks(headerPtr.NumChunks, headerPtr.ChunkLength, memoryMappedView);
 
 			return new SharedMemoryBufferDetail(name, headerPtr, chunks, memoryMappedFile, memoryMappedViewAccessor, memoryMappedView, readerEvent, writerEvent);
diff --git a/Engine/Source/Programs/Shared/EpicGames.Horde/Compute/Clients/AgentComputeClient.cs b/Engine/Source/Programs/Shared/EpicGames.Horde/Compute/Clients/AgentComputeClient.cs
--- a/Engine/Source/Programs/Shared/EpicGames.Horde/Compute/Clients/AgentComputeClient.cs
+++ b/Engine/Source/Programs/Shared/EpicGames.Horde/Compute/Clients/AgentComputeClient.cs
@@ -116,7 +116,7 @@ namespace EpicGames.Horde.Compute.Clients
 
 		async IAsyncEnumerable<RemoteComputeSocket> ConnectAsync(ILogger logger, [EnumeratorCancellation] CancellationToken cancellationToken)
 		{
-			using Socket listener = new Socket(SocketType.Stream, ProtocolType.IP);
+			using Socket listener = new(SocketType.Stream, ProtocolType.IP);
 			listener.Bind(new IPEndPoint(IPAddress.Loopback, _port));
 			listener.Listen();
 
@@ -132,12 +132,12 @@ namespace EpicGames.Horde.Compute.Clients
 
 		static async Task RunAgentAsync(string hordeAgentAssembly, int port, ILogger logger, CancellationToken cancellationToken)
 		{
-			using ManagedProcessGroup group = new ManagedProcessGroup();
+			using ManagedProcessGroup group = new();
 			if (RuntimeInformation.IsOSPlatform(OSPlatform.Windows))
 			{
 				List<string> arguments = [hordeAgentAssembly, "computeworker", $"-port={port}"];
 
-				using ManagedProcess process = new ManagedProcess(group, "dotnet", CommandLineArguments.Join(arguments), null, null, ProcessPriorityClass.Normal);
+				using ManagedProcess process = new(group, "dotnet", CommandLineArguments.Join(arguments), null, null, ProcessPriorityClass.Normal);
 
 				string? line;
 				while ((line = await process.ReadLineAsync(cancellationToken)) != null)
@@ -149,7 +149,7 @@ namespace EpicGames.Horde.Compute.Clients
 			}
 			else
 			{
-				using Process process = new Process();
+				using Process process = new();
 				process.StartInfo.FileName = "dotnet";
 				process.StartInfo.ArgumentList.Add(hordeAgentAssembly);
 				process.StartInfo.ArgumentList.Add("computeworker");
diff --git a/Engine/Source/Programs/Shared/EpicGames.Horde/Compute/Clients/LocalComputeClient.cs b/Engine/Source/Programs/Shared/EpicGames.Horde/Compute/Clients/LocalComputeClient.cs
--- a/Engine/Source/Programs/Shared/EpicGames.Horde/Compute/Clients/LocalComputeClient.cs
+++ b/Engine/Source/Programs/Shared/EpicGames.Horde/Compute/Clients/LocalComputeClient.cs
@@ -98,7 +98,7 @@ namespace EpicGames.Horde.Compute.Clients
 		public Task<IComputeLease?> TryAssignWorkerAsync(ClusterId? clusterId, Requirements? requirements, string? requestId, ConnectionMetadataRequest? connection, bool? useUbaCache, ILogger logger, CancellationToken cancellationToken)
 		{
 #pragma warning disable CA2000 // Dispose objects before losing scope
-			RemoteComputeSocket socket = new RemoteComputeSocket(new TcpTransport(_socket), ComputeProtocol.Latest, logger);
+			RemoteComputeSocket socket = new(new TcpTransport(_socket), ComputeProtocol.Latest, logger);
 			return Task.FromResult<IComputeLease?>(new LeaseImpl(socket));
 #pragma warning restore CA2000 // Dispose objects before losing scope
 		}
diff --git a/Engine/Source/Programs/Shared/EpicGames.Horde/Compute/Clients/ServerComputeClient.cs b/Engine/Source/Programs/Shared/EpicGames.Horde/Compute/Clients/ServerComputeClient.cs
--- a/Engine/Source/Programs/Shared/EpicGames.Horde/Compute/Clients/ServerComputeClient.cs
+++ b/Engine/Source/Programs/Shared/EpicGames.Horde/Compute/Clients/ServerComputeClient.cs
@@ -181,7 +181,7 @@ namespace EpicGames.Horde.Compute.Clients
 		}
 
 		readonly HttpClient _httpClient;
-		readonly CancellationTokenSource _cancellationSource = new CancellationTokenSource();
+		readonly CancellationTokenSource _cancellationSource = new();
 		readonly string _sessionId;
 		readonly ILogger _logger;
 
diff --git a/Engine/Source/Programs/Shared/EpicGames.Horde/Compute/ClusterId.cs b/Engine/Source/Programs/Shared/EpicGames.Horde/Compute/ClusterId.cs
--- a/Engine/Source/Programs/Shared/EpicGames.Horde/Compute/ClusterId.cs
+++ b/Engine/Source/Programs/Shared/EpicGames.Horde/Compute/ClusterId.cs
@@ -59,7 +59,7 @@ namespace EpicGames.Horde.Compute
 	sealed class ClusterIdCbConverter : CbConverter<ClusterId>
 	{
 		/// <inheritdoc/>
-		public override ClusterId Read(CbField field) => new ClusterId(field.AsString());
+		public override ClusterId Read(CbField field) => new(field.AsString());
 
 		/// <inheritdoc/>
 		public override void Write(CbWriter writer, ClusterId value) => writer.WriteStringValue(value.ToString());
@@ -74,7 +74,7 @@ namespace EpicGames.Horde.Compute
 	sealed class ClusterIdJsonConverter : JsonConverter<ClusterId>
 	{
 		/// <inheritdoc/>
-		public override ClusterId Read(ref Utf8JsonReader reader, Type typeToConvert, JsonSerializerOptions options) => new ClusterId(reader.GetString() ?? String.Empty);
+		public override ClusterId Read(ref Utf8JsonReader reader, Type typeToConvert, JsonSerializerOptions options) => new(reader.GetString() ?? String.Empty);
 
 		/// <inheritdoc/>
 		public override void Write(Utf8JsonWriter writer, ClusterId value, JsonSerializerOptions options) => writer.WriteStringValue(value.ToString());
diff --git a/Engine/Source/Programs/Shared/EpicGames.Horde/Compute/ComputeBuffer.cs b/Engine/Source/Programs/Shared/EpicGames.Horde/Compute/ComputeBuffer.cs
--- a/Engine/Source/Programs/Shared/EpicGames.Horde/Compute/ComputeBuffer.cs
+++ b/Engine/Source/Programs/Shared/EpicGames.Horde/Compute/ComputeBuffer.cs
@@ -307,7 +307,7 @@ namespace EpicGames.Horde.Compute
 			public ChunkStatePtr(ulong* data) => _data = data;
 
 			// Current value of the chunk state
-			public ChunkState Get() => new ChunkState(Interlocked.CompareExchange(ref *_data, 0, 0));
+			public ChunkState Get() => new(Interlocked.CompareExchange(ref *_data, 0, 0));
 
 			// Set the current state
 			public void Set(ChunkState value) => Interlocked.Exchange(ref *_data, value.Value);
@@ -362,7 +362,7 @@ namespace EpicGames.Horde.Compute
 			public ReaderStatePtr(ulong* data) => _data = data;
 
 			// Current value of the chunk state
-			public ReaderState Get() => new ReaderState(Interlocked.CompareExchange(ref *_data, 0, 0));
+			public ReaderState Get() => new(Interlocked.CompareExchange(ref *_data, 0, 0));
 
 			// Update current state
 			public void Set(ReaderState value) => Interlocked.Exchange(ref *_data, value.Value);
@@ -399,7 +399,7 @@ namespace EpicGames.Horde.Compute
 			public WriterStatePtr(ulong* data) => _data = data;
 
 			// Get the current value
-			public WriterState Get() => new WriterState(Interlocked.CompareExchange(ref *_data, 0, 0));
+			public WriterState Get() => new(Interlocked.CompareExchange(ref *_data, 0, 0));
 
 			// Set the current value
 			public void Set(WriterState state) => Interlocked.Exchange(ref *_data, state.Value);
@@ -429,11 +429,11 @@ namespace EpicGames.Horde.Compute
 			public int NumChunks => (int)((_data[0] >> 16) & 0xffff);
 			public int ChunkLength => (int)(_data[0] >> 32);
 
-			public WriterStatePtr GetWriterStatePtr() => new WriterStatePtr(_data + 1);
+			public WriterStatePtr GetWriterStatePtr() => new(_data + 1);
 
-			public ReaderStatePtr GetReaderStatePtr(int readerIdx) => new ReaderStatePtr(_data + 2 + readerIdx);
+			public ReaderStatePtr GetReaderStatePtr(int readerIdx) => new(_data + 2 + readerIdx);
 
-			public ChunkStatePtr GetChunkStatePtr(int chunkIdx) => new ChunkStatePtr(_data + 2 + ComputeBuffer.MaxReaders + chunkIdx);
+			public ChunkStatePtr GetChunkStatePtr(int chunkIdx) => new(_data + 2 + ComputeBuffer.MaxReaders + chunkIdx);
 		}
 
 		class BufferDebugProxy
diff --git a/Engine/Source/Programs/Shared/EpicGames.Horde/Compute/ComputeSocket.cs b/Engine/Source/Programs/Shared/EpicGames.Horde/Compute/ComputeSocket.cs
--- a/Engine/Source/Programs/Shared/EpicGames.Horde/Compute/ComputeSocket.cs
+++ b/Engine/Source/Programs/Shared/EpicGames.Horde/Compute/ComputeSocket.cs
@@ -163,7 +163,7 @@ namespace EpicGames.Horde.Compute
 
 		void AttachBuffer(IpcMessage message, int channelId, string bufferName)
 		{
-			MemoryWriter writer = new MemoryWriter(_commandBufferWriter.GetWriteBuffer());
+			MemoryWriter writer = new(_commandBufferWriter.GetWriteBuffer());
 			writer.WriteUnsignedVarInt((int)message);
 			writer.WriteUnsignedVarInt(channelId);
 			writer.WriteString(bufferName);
@@ -250,7 +250,7 @@ namespace EpicGames.Horde.Compute
 				while (await ipcReader.WaitToReadAsync(1, cancellationToken))
 				{
 					ReadOnlyMemory<byte> memory = ipcReader.GetReadBuffer();
-					MemoryReader reader = new MemoryReader(memory);
+					MemoryReader reader = new(memory);
 
 					IpcMessage message = (IpcMessage)reader.ReadUnsignedVarInt();
 					try
@@ -321,7 +321,7 @@ namespace EpicGames.Horde.Compute
 		class RecvBuffer : IDisposable
 		{
 			public ComputeBufferWriter? _writer;
-			public readonly SemaphoreSlim Semaphore = new SemaphoreSlim(1);
+			public readonly SemaphoreSlim Semaphore = new(1);
 			public int _refCount = 1;
 
 			public RecvBuffer(ComputeBufferWriter writer) => _writer = writer;
@@ -346,7 +346,7 @@ namespace EpicGames.Horde.Compute
 		class SendBuffer : IDisposable
 		{
 			public ComputeBufferReader? _reader;
-			public readonly SemaphoreSlim Semaphore = new SemaphoreSlim(1);
+			public readonly SemaphoreSlim Semaphore = new(1);
 			[SuppressMessage("Usage", "CA2213:Disposable fields should be disposed")]
 			public readonly BackgroundTask Task;
 			int _refCount = 1;
@@ -374,7 +374,7 @@ namespace EpicGames.Horde.Compute
 			}
 		}
 
-		readonly object _lockObject = new object();
+		readonly object _lockObject = new();
 
 		bool _complete;
 		bool _isClosing;
@@ -386,7 +386,7 @@ namespace EpicGames.Horde.Compute
 		BackgroundTask? _recvTask;
 		readonly Dictionary<int, RecvBuffer> _recvBuffers = [];
 
-		readonly SemaphoreSlim _sendSemaphore = new SemaphoreSlim(1, 1);
+		readonly SemaphoreSlim _sendSemaphore = new(1, 1);
 		readonly Dictionary<int, SendBuffer> _sendBuffers = [];
 
 		/// <inheritdoc/>
@@ -634,8 +634,8 @@ namespace EpicGames.Horde.Compute
 		}
 
 		readonly byte[] _header = new byte[8];
-		readonly SendSegment _headerSegment = new SendSegment();
-		readonly SendSegment _bodySegment = new SendSegment();
+		readonly SendSegment _headerSegment = new();
+		readonly SendSegment _bodySegment = new();
 
 		/// <inheritdoc/>
 		async ValueTask SendAsync(int id, ReadOnlyMemory<byte> memory, CancellationToken cancellationToken = default)
@@ -659,7 +659,7 @@ namespace EpicGames.Horde.Compute
 				_headerSegment.Set(_header, _bodySegment, 0);
 				_bodySegment.Set(memory, null, _header.Length);
 
-				ReadOnlySequence<byte> sequence = new ReadOnlySequence<byte>(_headerSegment, 0, _bodySegment, memory.Length);
+				ReadOnlySequence<byte> sequence = new(_headerSegment, 0, _bodySegment, memory.Length);
 				await _transport.SendAsync(sequence, cancellationToken);
 			}
 			finally
@@ -774,7 +774,7 @@ namespace EpicGames.Horde.Compute
 				}
 
 				ComputeBufferReader sendBufferReader = sendBuffer.CreateReader();
-				SendBuffer sendBufferInfo = new SendBuffer(sendBufferReader, (buffer, ctx) => SendFromBufferAsync(channelId, buffer, ctx));
+				SendBuffer sendBufferInfo = new(sendBufferReader, (buffer, ctx) => SendFromBufferAsync(channelId, buffer, ctx));
 				_sendBuffers.Add(channelId, sendBufferInfo);
 			}
 		}
diff --git a/Engine/Source/Programs/Shared/EpicGames.Horde/Compute/Native.cs b/Engine/Source/Programs/Shared/EpicGames.Horde/Compute/Native.cs
--- a/Engine/Source/Programs/Shared/EpicGames.Horde/Compute/Native.cs
+++ b/Engine/Source/Programs/Shared/EpicGames.Horde/Compute/Native.cs
@@ -54,7 +54,7 @@ namespace EpicGames.Horde.Compute
 
 			public static EventHandle CreateNew(string? name, EventResetMode resetMode, bool initialState, HandleInheritability handleInheritability)
 			{
-				Native.SECURITY_ATTRIBUTES securityAttributes = new Native.SECURITY_ATTRIBUTES();
+				Native.SECURITY_ATTRIBUTES securityAttributes = new();
 				securityAttributes.nLength = Marshal.SizeOf(securityAttributes);
 				securityAttributes.bInheritHandle = (handleInheritability == HandleInheritability.Inheritable) ? 1 : 0;
 
diff --git a/Engine/Source/Programs/Shared/EpicGames.Horde/Compute/Transports/StreamTransport.cs b/Engine/Source/Programs/Shared/EpicGames.Horde/Compute/Transports/StreamTransport.cs
--- a/Engine/Source/Programs/Shared/EpicGames.Horde/Compute/Transports/StreamTransport.cs
+++ b/Engine/Source/Programs/Shared/EpicGames.Horde/Compute/Transports/StreamTransport.cs
@@ -57,6 +57,6 @@ namespace EpicGames.Horde.Compute.Transports
 		}
 
 		/// <inheritdoc/>
-		public override ValueTask MarkCompleteAsync(CancellationToken cancellationToken) => new ValueTask();
+		public override ValueTask MarkCompleteAsync(CancellationToken cancellationToken) => new();
 	}
 }
diff --git a/Engine/Source/Programs/Shared/EpicGames.Horde/Devices/DeviceId.cs b/Engine/Source/Programs/Shared/EpicGames.Horde/Devices/DeviceId.cs
--- a/Engine/Source/Programs/Shared/EpicGames.Horde/Devices/DeviceId.cs
+++ b/Engine/Source/Programs/Shared/EpicGames.Horde/Devices/DeviceId.cs
@@ -28,7 +28,7 @@ namespace EpicGames.Horde.Devices
 		public bool IsEmpty => Id.IsEmpty;
 
 		/// <inheritdoc cref="StringId.Sanitize(System.String)"/>
-		public static DeviceId Sanitize(string text) => new DeviceId(StringId.Sanitize(text));
+		public static DeviceId Sanitize(string text) => new(StringId.Sanitize(text));
 
 		/// <inheritdoc/>
 		public override string ToString() => Id.ToString();
@@ -40,7 +40,7 @@ namespace EpicGames.Horde.Devices
 	class DeviceIdConverter : StringIdConverter<DeviceId>
 	{
 		/// <inheritdoc/>
-		public override DeviceId FromStringId(StringId id) => new DeviceId(id);
+		public override DeviceId FromStringId(StringId id) => new(id);
 
 		/// <inheritdoc/>
 		public override StringId ToStringId(DeviceId value) => value.Id;
diff --git a/Engine/Source/Programs/Shared/EpicGames.Horde/Devices/DevicePlatformId.cs b/Engine/Source/Programs/Shared/EpicGames.Horde/Devices/DevicePlatformId.cs
--- a/Engine/Source/Programs/Shared/EpicGames.Horde/Devices/DevicePlatformId.cs
+++ b/Engine/Source/Programs/Shared/EpicGames.Horde/Devices/DevicePlatformId.cs
@@ -28,7 +28,7 @@ namespace EpicGames.Horde.Devices
 		public bool IsEmpty => Id.IsEmpty;
 
 		/// <inheritdoc cref="StringId.Sanitize(System.String)"/>
-		public static DevicePlatformId Sanitize(string text) => new DevicePlatformId(StringId.Sanitize(text));
+		public static DevicePlatformId Sanitize(string text) => new(StringId.Sanitize(text));
 
 		/// <inheritdoc/>
 		public override string ToString() => Id.ToString();
@@ -40,7 +40,7 @@ namespace EpicGames.Horde.Devices
 	class DevicePlatformIdConverter : StringIdConverter<DevicePlatformId>
 	{
 		/// <inheritdoc/>
-		public override DevicePlatformId FromStringId(StringId id) => new DevicePlatformId(id);
+		public override DevicePlatformId FromStringId(StringId id) => new(id);
 
 		/// <inheritdoc/>
 		public override StringId ToStringId(DevicePlatformId value) => value.Id;
diff --git a/Engine/Source/Programs/Shared/EpicGames.Horde/Devices/DevicePoolId.cs b/Engine/Source/Programs/Shared/EpicGames.Horde/Devices/DevicePoolId.cs
--- a/Engine/Source/Programs/Shared/EpicGames.Horde/Devices/DevicePoolId.cs
+++ b/Engine/Source/Programs/Shared/EpicGames.Horde/Devices/DevicePoolId.cs
@@ -28,7 +28,7 @@ namespace EpicGames.Horde.Devices
 		public bool IsEmpty => Id.IsEmpty;
 
 		/// <inheritdoc cref="StringId.Sanitize(System.String)"/>
-		public static DevicePoolId Sanitize(string text) => new DevicePoolId(StringId.Sanitize(text));
+		public static DevicePoolId Sanitize(string text) => new(StringId.Sanitize(text));
 
 		/// <inheritdoc/>
 		public override string ToString() => Id.ToString();
@@ -40,7 +40,7 @@ namespace EpicGames.Horde.Devices
 	class DevicePoolIdConverter : StringIdConverter<DevicePoolId>
 	{
 		/// <inheritdoc/>
-		public override DevicePoolId FromStringId(StringId id) => new DevicePoolId(id);
+		public override DevicePoolId FromStringId(StringId id) => new(id);
 
 		/// <inheritdoc/>
 		public override StringId ToStringId(DevicePoolId value) => value.Id;
diff --git a/Engine/Source/Programs/Shared/EpicGames.Horde/HordeClient.cs b/Engine/Source/Programs/Shared/EpicGames.Horde/HordeClient.cs
--- a/Engine/Source/Programs/Shared/EpicGames.Horde/HordeClient.cs
+++ b/Engine/Source/Programs/Shared/EpicGames.Horde/HordeClient.cs
@@ -143,18 +143,18 @@ namespace EpicGames.Horde
 				GetPortsResponse? ports = await response.Content.ReadFromJsonAsync<GetPortsResponse>(HordeHttpClient.JsonSerializerOptions, cancellationToken);
 				if (ports is { UnencryptedHttp2: not null } && ports.UnencryptedHttp2.Value != 0)
 				{
-					UriBuilder builder = new UriBuilder(serverUri);
+					UriBuilder builder = new(serverUri);
 					builder.Port = ports.UnencryptedHttp2.Value;
 					serverUri = builder.Uri;
 				}
 			}
 
 #pragma warning disable CA2000 // Dispose objects before losing scope
-			SocketsHttpHandler httpHandler = new SocketsHttpHandler();
+			SocketsHttpHandler httpHandler = new();
 #pragma warning restore CA2000 // Dispose objects before losing scope
 
 			// Create options for the new channel
-			GrpcChannelOptions channelOptions = new GrpcChannelOptions();
+			GrpcChannelOptions channelOptions = new();
 			channelOptions.MaxReceiveMessageSize = 1024 * 1024 * 1024; // 1 GB 		// Required payloads coming from CAS service can be large
 			channelOptions.MaxSendMessageSize = 1024 * 1024 * 1024; // 1 GB
 			channelOptions.LoggerFactory = _loggerFactory;
@@ -203,7 +203,7 @@ namespace EpicGames.Horde
 
 		/// <inheritdoc/>
 		public HordeHttpClient CreateHttpClient()
-			=> new HordeHttpClient(CreateAuthenticatedHttpClient());
+			=> new(CreateAuthenticatedHttpClient());
 
 		/// <inheritdoc/>
 		ServerComputeClient CreateComputeClient()
@@ -238,7 +238,7 @@ namespace EpicGames.Horde
 				}
 			}
 
-			HttpStorageBackend httpStorageBackend = new HttpStorageBackend(basePath, CreateClient, CreateUnauthenticatedHttpClient, _loggerFactory.CreateLogger<HttpStorageBackend>());
+			HttpStorageBackend httpStorageBackend = new(basePath, CreateClient, CreateUnauthenticatedHttpClient, _loggerFactory.CreateLogger<HttpStorageBackend>());
 			return new BundleStorageNamespace(httpStorageBackend, _bundleCache, _hordeOptions.Bundle, _loggerFactory.CreateLogger<BundleStorageNamespace>());
 		}
 
@@ -286,7 +286,7 @@ namespace EpicGames.Horde
 		{
 			using (HttpClient httpClient = CreateAuthenticatedHttpClient())
 			{
-				using HttpRequestMessage request = new HttpRequestMessage(HttpMethod.Get, "api/v1/dashboard/challenge");
+				using HttpRequestMessage request = new(HttpMethod.Get, "api/v1/dashboard/challenge");
 				using HttpResponseMessage response = await httpClient.SendAsync(request, cancellationToken);
 				return response.IsSuccessStatusCode;
 			}
@@ -303,7 +303,7 @@ namespace EpicGames.Horde
 		/// <inheritdoc/>
 		protected override HttpClient CreateAuthenticatedHttpClient()
 		{
-			HttpClient httpClient = new HttpClient(_httpMessageHandler.Instance, false);
+			HttpClient httpClient = new(_httpMessageHandler.Instance, false);
 			httpClient.BaseAddress = ServerUrl;
 			if (!String.IsNullOrEmpty(_accessToken))
 			{
@@ -315,7 +315,7 @@ namespace EpicGames.Horde
 		/// <inheritdoc/>
 		protected override HttpClient CreateUnauthenticatedHttpClient()
 		{
-			HttpClient httpClient = new HttpClient(_httpMessageHandler.Instance, false);
+			HttpClient httpClient = new(_httpMessageHandler.Instance, false);
 			httpClient.BaseAddress = ServerUrl;
 			return httpClient;
 		}
@@ -378,7 +378,7 @@ namespace EpicGames.Horde
 			// Send a request to the server to log in automatically
 			using (HttpClient httpClient = CreateAuthenticatedHttpClient())
 			{
-				using HttpRequestMessage request = new HttpRequestMessage(HttpMethod.Get, "api/v1/dashboard/challenge");
+				using HttpRequestMessage request = new(HttpMethod.Get, "api/v1/dashboard/challenge");
 				request.Options.Set(s_allowInteractiveLogin, allowLogin);
 
 				using HttpResponseMessage response = await httpClient.SendAsync(request, cancellationToken);
@@ -406,7 +406,7 @@ namespace EpicGames.Horde
 		/// <inheritdoc/>
 		protected override HttpClient CreateAuthenticatedHttpClient()
 		{
-			HttpClient httpClient = new HttpClient(_authHttpMessageHandler, false);
+			HttpClient httpClient = new(_authHttpMessageHandler, false);
 			httpClient.BaseAddress = ServerUrl;
 			return httpClient;
 		}
@@ -414,13 +414,13 @@ namespace EpicGames.Horde
 		/// <inheritdoc/>
 		protected override HttpClient CreateUnauthenticatedHttpClient()
 		{
-			HttpClient httpClient = new HttpClient(_baseHttpMessageHandler.Instance, false);
+			HttpClient httpClient = new(_baseHttpMessageHandler.Instance, false);
 			httpClient.BaseAddress = ServerUrl;
 			return httpClient;
 		}
 
 		private static HttpOidcAuthHandlerConfig HordeOptionsToIHttpOidcAuthHandlerConfig(HordeOptions options) =>
-			new HttpOidcAuthHandlerConfig(AllowInteractiveLoginHeaderKey, options.AllowAuthPrompt);
+			new(AllowInteractiveLoginHeaderKey, options.AllowAuthPrompt);
 	}
 
 	/// <summary>
diff --git a/Engine/Source/Programs/Shared/EpicGames.Horde/HordeHttpClient.cs b/Engine/Source/Programs/Shared/EpicGames.Horde/HordeHttpClient.cs
--- a/Engine/Source/Programs/Shared/EpicGames.Horde/HordeHttpClient.cs
+++ b/Engine/Source/Programs/Shared/EpicGames.Horde/HordeHttpClient.cs
@@ -202,7 +202,7 @@ namespace EpicGames.Horde
 		/// <returns>Information about all the artifacts</returns>
 		public async Task<List<GetArtifactResponse>> FindArtifactsAsync(IEnumerable<ArtifactId>? ids, StreamId? streamId = null, CommitId? minCommitId = null, CommitId? maxCommitId = null, ArtifactName? name = null, ArtifactType? type = null, IEnumerable<string>? keys = null, int maxResults = 100, CancellationToken cancellationToken = default)
 		{
-			QueryStringBuilder queryParams = new QueryStringBuilder();
+			QueryStringBuilder queryParams = new();
 
 			if (ids != null)
 			{
@@ -401,7 +401,7 @@ namespace EpicGames.Horde
 		/// <param name="cancellationToken">Cancellation token for the operation</param>
 		public async Task<ReadRefResponse?> TryReadRefAsync(string path, RefCacheTime cacheTime = default, CancellationToken cancellationToken = default)
 		{
-			using (HttpRequestMessage request = new HttpRequestMessage(HttpMethod.Get, path))
+			using (HttpRequestMessage request = new(HttpMethod.Get, path))
 			{
 				if (cacheTime.IsSet())
 				{
@@ -443,7 +443,7 @@ namespace EpicGames.Horde
 		/// <remarks>BadRequest can be returned under circusmtances where the payload is not serializable (primitive or null payload, <see cref="NotSupportedException"/>, and <see cref="JsonException"/>).</remarks>
 		public async Task<HttpResponseMessage> PostTelemetryAsync(string appId, int appVersion, string appEnvironment, object[] payloads, string? sessionId = null, CancellationToken cancellationToken = default)
 		{
-			QueryStringBuilder queryParams = new QueryStringBuilder();
+			QueryStringBuilder queryParams = new();
 			queryParams.Add("appId", appId);
 			queryParams.Add("appVersion", appVersion.ToString());
 			queryParams.Add("appEnvironment", appEnvironment);
@@ -454,7 +454,7 @@ namespace EpicGames.Horde
 				queryParams.Add("SessionId", sessionId);
 			}
 
-			PostTelemetryEventStreamRequest request = new PostTelemetryEventStreamRequest();
+			PostTelemetryEventStreamRequest request = new();
 
 			try
 			{
@@ -496,7 +496,7 @@ namespace EpicGames.Horde
 		/// <param name="cancellationToken">Cancellation token for the operation</param>
 		public Task<List<GetUtilizationDataResponse>> GetTelemetryAsync(DateTime endDate, int range, int? tzOffset = null, CancellationToken cancellationToken = default)
 		{
-			QueryStringBuilder queryParams = new QueryStringBuilder();
+			QueryStringBuilder queryParams = new();
 			queryParams.Add("Range", range.ToString());
 			if (tzOffset != null)
 			{
@@ -559,7 +559,7 @@ namespace EpicGames.Horde
 		/// <param name="cancellationToken">Cancellation token for the operation</param>
 		public async Task<ToolDeploymentId> CreateToolDeploymentAsync(ToolId id, string? version, double? duration, bool? createPaused, HashedBlobRefValue target, CancellationToken cancellationToken = default)
 		{
-			CreateToolDeploymentRequest request = new CreateToolDeploymentRequest(version ?? String.Empty, duration, createPaused, target);
+			CreateToolDeploymentRequest request = new(version ?? String.Empty, duration, createPaused, target);
 			CreateToolDeploymentResponse response = await PostAsync<CreateToolDeploymentResponse, CreateToolDeploymentRequest>(_httpClient, $"api/v2/tools/{id}/deployments", request, cancellationToken);
 			return response.Id;
 		}
@@ -617,7 +617,7 @@ namespace EpicGames.Horde
 			int count = 100,
 			CancellationToken cancellationToken = default)
 		{
-			StringBuilder queryBuilder = new StringBuilder($"/api/v1/jobs?streamId={streamId}&count={count}&includePreflight={(includePreflight ? "true" : "false")}");
+			StringBuilder queryBuilder = new($"/api/v1/jobs?streamId={streamId}&count={count}&includePreflight={(includePreflight ? "true" : "false")}");
 
 			if (templateIds != null)
 			{
@@ -713,7 +713,7 @@ namespace EpicGames.Horde
 		/// <returns></returns>
 		public Task<GetUgsMetadataListResponse> GetUgsMetadataAsync(StreamId streamId, CommitId commitId, ProjectId projectId, CancellationToken cancellationToken = default)
 		{
-			Regex streamRe = new Regex(@"(.*?)\-(.*)");
+			Regex streamRe = new(@"(.*?)\-(.*)");
 			Match streamMatch = streamRe.Match(streamId.ToString());
 			if (streamMatch.Success)
 			{
diff --git a/Engine/Source/Programs/Shared/EpicGames.Horde/HordeHttpRequest.cs b/Engine/Source/Programs/Shared/EpicGames.Horde/HordeHttpRequest.cs
--- a/Engine/Source/Programs/Shared/EpicGames.Horde/HordeHttpRequest.cs
+++ b/Engine/Source/Programs/Shared/EpicGames.Horde/HordeHttpRequest.cs
@@ -26,7 +26,7 @@ namespace EpicGames.Horde
 		/// </summary>
 		static JsonSerializerOptions CreateJsonSerializerOptions()
 		{
-			JsonSerializerOptions options = new JsonSerializerOptions();
+			JsonSerializerOptions options = new();
 			ConfigureJsonSerializer(options);
 			return options;
 		}
diff --git a/Engine/Source/Programs/Shared/EpicGames.Horde/HordeOptions.cs b/Engine/Source/Programs/Shared/EpicGames.Horde/HordeOptions.cs
--- a/Engine/Source/Programs/Shared/EpicGames.Horde/HordeOptions.cs
+++ b/Engine/Source/Programs/Shared/EpicGames.Horde/HordeOptions.cs
@@ -165,7 +165,7 @@ namespace EpicGames.Horde
 
 					using (FileStream stream = FileReference.Open(configFile, FileMode.Create, FileAccess.ReadWrite, FileShare.Read))
 					{
-						using Utf8JsonWriter writer = new Utf8JsonWriter(stream, new JsonWriterOptions { Indented = true });
+						using Utf8JsonWriter writer = new(stream, new JsonWriterOptions { Indented = true });
 						root.WriteTo(writer);
 					}
 				}
diff --git a/Engine/Source/Programs/Shared/EpicGames.Horde/Issues/Handlers/BuildGraphIssueHandler.cs b/Engine/Source/Programs/Shared/EpicGames.Horde/Issues/Handlers/BuildGraphIssueHandler.cs
--- a/Engine/Source/Programs/Shared/EpicGames.Horde/Issues/Handlers/BuildGraphIssueHandler.cs
+++ b/Engine/Source/Programs/Shared/EpicGames.Horde/Issues/Handlers/BuildGraphIssueHandler.cs
@@ -37,7 +37,7 @@ namespace EpicGames.Horde.Issues.Handlers
 				EventId eventId = issueEvent.EventId.Value;
 				if (IsMatchingEventId(eventId))
 				{
-					IssueEventGroup issue = new IssueEventGroup("BuildGraph", "BuildGraph {Severity} in {Files}", IssueChangeFilter.All);
+					IssueEventGroup issue = new("BuildGraph", "BuildGraph {Severity} in {Files}", IssueChangeFilter.All);
 					issue.Events.Add(issueEvent);
 					issue.Keys.AddSourceFiles(issueEvent);
 					_issues.Add(issue);
diff --git a/Engine/Source/Programs/Shared/EpicGames.Horde/Issues/Handlers/CompileIssueHandler.cs b/Engine/Source/Programs/Shared/EpicGames.Horde/Issues/Handlers/CompileIssueHandler.cs
--- a/Engine/Source/Programs/Shared/EpicGames.Horde/Issues/Handlers/CompileIssueHandler.cs
+++ b/Engine/Source/Programs/Shared/EpicGames.Horde/Issues/Handlers/CompileIssueHandler.cs
@@ -68,7 +68,7 @@ namespace EpicGames.Horde.Issues.Handlers
 						fingerprintType = $"{fingerprintType}:{group}";
 					}
 
-					IssueEventGroup issue = new IssueEventGroup(fingerprintType, "{Meta:CompileType} {Severity} in {Files}", IssueChangeFilter.Code);
+					IssueEventGroup issue = new(fingerprintType, "{Meta:CompileType} {Severity} in {Files}", IssueChangeFilter.Code);
 					issue.Events.Add(issueEvent);
 					issue.Keys.AddSourceFiles(issueEvent);
 					issue.Metadata.Add(CompileTypeAnnotation, compileType);
diff --git a/Engine/Source/Programs/Shared/EpicGames.Horde/Issues/Handlers/ContentIssueHandler.cs b/Engine/Source/Programs/Shared/EpicGames.Horde/Issues/Handlers/ContentIssueHandler.cs
--- a/Engine/Source/Programs/Shared/EpicGames.Horde/Issues/Handlers/ContentIssueHandler.cs
+++ b/Engine/Source/Programs/Shared/EpicGames.Horde/Issues/Handlers/ContentIssueHandler.cs
@@ -28,7 +28,7 @@ namespace EpicGames.Horde.Issues.Handlers
 				EventId eventId = issueEvent.EventId.Value;
 				if (IsMatchingEventId(eventId))
 				{
-					IssueEventGroup issue = new IssueEventGroup("Content", "{Severity} in {Files}", IssueChangeFilter.Content);
+					IssueEventGroup issue = new("Content", "{Severity} in {Files}", IssueChangeFilter.Content);
 					issue.Events.Add(issueEvent);
 					issue.Keys.AddAssets(issueEvent);
 					_issues.Add(issue);
diff --git a/Engine/Source/Programs/Shared/EpicGames.Horde/Issues/Handlers/CopyrightIssueHandler.cs b/Engine/Source/Programs/Shared/EpicGames.Horde/Issues/Handlers/CopyrightIssueHandler.cs
--- a/Engine/Source/Programs/Shared/EpicGames.Horde/Issues/Handlers/CopyrightIssueHandler.cs
+++ b/Engine/Source/Programs/Shared/EpicGames.Horde/Issues/Handlers/CopyrightIssueHandler.cs
@@ -22,7 +22,7 @@ namespace EpicGames.Horde.Issues.Handlers
 		{
 			if (issueEvent.EventId == KnownLogEvents.AutomationTool_MissingCopyright)
 			{
-				IssueEventGroup issue = new IssueEventGroup("Copyright", "Missing copyright notice in {Files}", IssueChangeFilter.Code);
+				IssueEventGroup issue = new("Copyright", "Missing copyright notice in {Files}", IssueChangeFilter.Code);
 				issue.Events.Add(issueEvent);
 				issue.Keys.AddSourceFiles(issueEvent);
 				_issues.Add(issue);
diff --git a/Engine/Source/Programs/Shared/EpicGames.Horde/Issues/Handlers/ExternalIssueHandler.cs b/Engine/Source/Programs/Shared/EpicGames.Horde/Issues/Handlers/ExternalIssueHandler.cs
--- a/Engine/Source/Programs/Shared/EpicGames.Horde/Issues/Handlers/ExternalIssueHandler.cs
+++ b/Engine/Source/Programs/Shared/EpicGames.Horde/Issues/Handlers/ExternalIssueHandler.cs
@@ -16,8 +16,8 @@ namespace EpicGames.Horde.Issues.Handlers
 	[IssueHandler]
 	public class ExternalIssueHandler : IssueHandler
 	{
-		static readonly JsonNodeOptions s_jsonNodeOptions = new JsonNodeOptions { PropertyNameCaseInsensitive = true };
-		static readonly JsonSerializerOptions? s_jsonSerializerOptions = new JsonSerializerOptions { PropertyNameCaseInsensitive = true };
+		static readonly JsonNodeOptions s_jsonNodeOptions = new() { PropertyNameCaseInsensitive = true };
+		static readonly JsonSerializerOptions? s_jsonSerializerOptions = new() { PropertyNameCaseInsensitive = true };
 
 		readonly List<IssueEventGroup> _issues = [];
 
@@ -32,7 +32,7 @@ namespace EpicGames.Horde.Issues.Handlers
 				IssueFingerprint? issueData = ParseIssueFingerprint(line);
 				if (issueData != null)
 				{
-					IssueEventGroup group = new IssueEventGroup(issueData.Type, issueData.SummaryTemplate, issueData.ChangeFilter);
+					IssueEventGroup group = new(issueData.Type, issueData.SummaryTemplate, issueData.ChangeFilter);
 					if (issueData.Keys != null)
 					{
 						group.Keys.UnionWith(issueData.Keys);
diff --git a/Engine/Source/Programs/Shared/EpicGames.Horde/Issues/Handlers/GauntletIssueHandler.cs b/Engine/Source/Programs/Shared/EpicGames.Horde/Issues/Handlers/GauntletIssueHandler.cs
--- a/Engine/Source/Programs/Shared/EpicGames.Horde/Issues/Handlers/GauntletIssueHandler.cs
+++ b/Engine/Source/Programs/Shared/EpicGames.Horde/Issues/Handlers/GauntletIssueHandler.cs
@@ -83,7 +83,7 @@ namespace EpicGames.Horde.Issues.Handlers
 		/// <summary>
 		///  Known Gauntlet events
 		/// </summary>
-		static readonly Dictionary<EventId, string> s_knownGauntletEvents = new Dictionary<EventId, string>
+		static readonly Dictionary<EventId, string> s_knownGauntletEvents = new()
 		{
 			{ KnownLogEvents.Gauntlet, FrameworkPrefix},
 			{ KnownLogEvents.Gauntlet_TestEvent, TestPrefix},
@@ -315,7 +315,7 @@ namespace EpicGames.Horde.Issues.Handlers
 				string summaryTemplate = hasCallstack ? "Automation {Meta:GauntletType} " + callstack + " {Severity} in {Meta:Node}" : "Automation {Meta:GauntletType} {Severity} in {Meta:Node}";
 
 				string gauntletType = GetEventPrefix(issueEvent.EventId!.Value);
-				IssueEventGroup issue = new IssueEventGroup($"Gauntlet:{gauntletType}", summaryTemplate, IssueChangeFilter.All);
+				IssueEventGroup issue = new($"Gauntlet:{gauntletType}", summaryTemplate, IssueChangeFilter.All);
 				issue.Events.Add(issueEvent);
 				GetHash(issueEvent, issue.Keys, issue.Metadata, hasCallstack);
 				issue.Metadata.Add("GauntletType", gauntletType);
diff --git a/Engine/Source/Programs/Shared/EpicGames.Horde/Issues/Handlers/HashedIssueHandler.cs b/Engine/Source/Programs/Shared/EpicGames.Horde/Issues/Handlers/HashedIssueHandler.cs
--- a/Engine/Source/Programs/Shared/EpicGames.Horde/Issues/Handlers/HashedIssueHandler.cs
+++ b/Engine/Source/Programs/Shared/EpicGames.Horde/Issues/Handlers/HashedIssueHandler.cs
@@ -163,7 +163,7 @@ namespace EpicGames.Horde.Issues.Handlers
 					{
 						hashes.Add(hash);
 
-						IssueEventGroup issue = new IssueEventGroup(fingerprintType, summaryTemplate, IssueChangeFilter.All);
+						IssueEventGroup issue = new(fingerprintType, summaryTemplate, IssueChangeFilter.All);
 						issue.Events.Add(stepEvent);
 						issue.Keys.AddHash(hash);
 						stepEvent.AuditLogger?.LogDebug("Fingerprint key: '{Key}' generated from event: '{Event}'", hash, stepEvent.Render());
diff --git a/Engine/Source/Programs/Shared/EpicGames.Horde/Issues/Handlers/LocalizationIssueHandler.cs b/Engine/Source/Programs/Shared/EpicGames.Horde/Issues/Handlers/LocalizationIssueHandler.cs
--- a/Engine/Source/Programs/Shared/EpicGames.Horde/Issues/Handlers/LocalizationIssueHandler.cs
+++ b/Engine/Source/Programs/Shared/EpicGames.Horde/Issues/Handlers/LocalizationIssueHandler.cs
@@ -78,7 +78,7 @@ namespace EpicGames.Horde.Issues.Handlers
 				System.String.IsNullOrEmpty(_context.NodeName) ? "" : "{Meta:Node}: ",
 				System.String.IsNullOrEmpty(_context.StreamId.ToString()) ? "" : " in {Meta:Stream}");
 
-			IssueEventGroup newIssueGroup = new IssueEventGroup(localizationIssueGroup, localizationSummaryTemplate, IssueChangeFilter.All);
+			IssueEventGroup newIssueGroup = new(localizationIssueGroup, localizationSummaryTemplate, IssueChangeFilter.All);
 
 			if (!System.String.IsNullOrEmpty(_context.NodeName))
 			{
@@ -121,7 +121,7 @@ namespace EpicGames.Horde.Issues.Handlers
 		/// <param name="issueEvent">The event data</param>
 		private static LocalizationData GetLocalizationData(IssueEvent issueEvent)
 		{
-			LocalizationData data = new LocalizationData();
+			LocalizationData data = new();
 
 			foreach (JsonLogEvent line in issueEvent.Lines)
 			{
diff --git a/Engine/Source/Programs/Shared/EpicGames.Horde/Issues/Handlers/PerforceCaseIssueHandler.cs b/Engine/Source/Programs/Shared/EpicGames.Horde/Issues/Handlers/PerforceCaseIssueHandler.cs
--- a/Engine/Source/Programs/Shared/EpicGames.Horde/Issues/Handlers/PerforceCaseIssueHandler.cs
+++ b/Engine/Source/Programs/Shared/EpicGames.Horde/Issues/Handlers/PerforceCaseIssueHandler.cs
@@ -32,7 +32,7 @@ namespace EpicGames.Horde.Issues.Handlers
 		{
 			if (issueEvent.EventId != null && IsMatchingEventId(issueEvent.EventId.Value))
 			{
-				IssueEventGroup issue = new IssueEventGroup("PerforceCase", "Inconsistent case for {Files}", IssueChangeFilter.All);
+				IssueEventGroup issue = new("PerforceCase", "Inconsistent case for {Files}", IssueChangeFilter.All);
 				issue.Events.Add(issueEvent);
 				issue.Keys.AddDepotPaths(issueEvent);
 				_issues.Add(issue);
diff --git a/Engine/Source/Programs/Shared/EpicGames.Horde/Issues/Handlers/SanitizerIssueHandler.cs b/Engine/Source/Programs/Shared/EpicGames.Horde/Issues/Handlers/SanitizerIssueHandler.cs
--- a/Engine/Source/Programs/Shared/EpicGames.Horde/Issues/Handlers/SanitizerIssueHandler.cs
+++ b/Engine/Source/Programs/Shared/EpicGames.Horde/Issues/Handlers/SanitizerIssueHandler.cs
@@ -14,9 +14,9 @@ namespace EpicGames.Horde.Issues.Handlers
 	[IssueHandler]
 	public class SanitizerIssueHandler : IssueHandler
 	{
-		static readonly Utf8String s_summaryReasonAnnotation = new Utf8String("SummaryReason");
-		static readonly Utf8String s_summarySourceFileAnnotation = new Utf8String("SummarySourceFile");
-		static readonly Utf8String s_sanitizerNameAnnotation = new Utf8String("SanitizerName");
+		static readonly Utf8String s_summaryReasonAnnotation = new("SummaryReason");
+		static readonly Utf8String s_summarySourceFileAnnotation = new("SummarySourceFile");
+		static readonly Utf8String s_sanitizerNameAnnotation = new("SanitizerName");
 		readonly List<IssueEventGroup> _issues = [];
 
 		/// <inheritdoc/>
@@ -24,17 +24,17 @@ namespace EpicGames.Horde.Issues.Handlers
 		/// <summary>
 		/// Log value describing the thread sanitizer error summary reason
 		/// </summary>
-		public static LogValue SanitizerName => new LogValue(s_sanitizerNameAnnotation, "");
+		public static LogValue SanitizerName => new(s_sanitizerNameAnnotation, "");
 
 		/// <summary>
 		/// Log value describing the thread sanitizer error summary reason
 		/// </summary>
-		public static LogValue SummaryReason => new LogValue(s_summaryReasonAnnotation, "");
+		public static LogValue SummaryReason => new(s_summaryReasonAnnotation, "");
 
 		/// <summary>
 		/// Log value describing the thread sanitizer error summary source file
 		/// </summary>
-		public static LogValue SummarySourceFile => new LogValue(s_summarySourceFileAnnotation, "");
+		public static LogValue SummarySourceFile => new(s_summarySourceFileAnnotation, "");
 
 		static bool FindAnnotation(IssueEvent issueEvent, Utf8String annotation, ref string result)
 		{
@@ -66,7 +66,7 @@ namespace EpicGames.Horde.Issues.Handlers
 				FindAnnotation(issueEvent, s_sanitizerNameAnnotation, ref sanitizerName);
 				string fingerprint = $"{sanitizerName}Sanitizer:{summaryReason}:{summaryFile}";
 
-				IssueEventGroup issue = new IssueEventGroup(fingerprint, $"{{Meta:{s_sanitizerNameAnnotation}}} found '{{Meta:{s_summaryReasonAnnotation}}}' in {{Meta:{s_summarySourceFileAnnotation}}}", IssueChangeFilter.Code);
+				IssueEventGroup issue = new(fingerprint, $"{{Meta:{s_sanitizerNameAnnotation}}} found '{{Meta:{s_summaryReasonAnnotation}}}' in {{Meta:{s_summarySourceFileAnnotation}}}", IssueChangeFilter.Code);
 				issue.Keys.Add(sanitizerName, IssueKeyType.None);
 				issue.Metadata.Add(s_sanitizerNameAnnotation.ToString(), sanitizerName);
 				issue.Keys.Add(summaryReason, IssueKeyType.Note);
diff --git a/Engine/Source/Programs/Shared/EpicGames.Horde/Issues/Handlers/ScopedIssueHandler.cs b/Engine/Source/Programs/Shared/EpicGames.Horde/Issues/Handlers/ScopedIssueHandler.cs
--- a/Engine/Source/Programs/Shared/EpicGames.Horde/Issues/Handlers/ScopedIssueHandler.cs
+++ b/Engine/Source/Programs/Shared/EpicGames.Horde/Issues/Handlers/ScopedIssueHandler.cs
@@ -69,7 +69,7 @@ namespace EpicGames.Horde.Issues.Handlers
 
 			if (TryGetHash(hashSource, out Md5Hash hash))
 			{
-				IssueEventGroup issue = new IssueEventGroup(fingerprintType, "{Severity} in {Meta:Node} - {Meta:Scope}", IssueChangeFilter.All);
+				IssueEventGroup issue = new(fingerprintType, "{Severity} in {Meta:Node} - {Meta:Scope}", IssueChangeFilter.All);
 				issue.Events.Add(logEvent);
 				issue.Keys.Add(IssueKey.FromHash(hash));
 				issue.Metadata.Add(NodeName, _context.NodeName);
diff --git a/Engine/Source/Programs/Shared/EpicGames.Horde/Issues/Handlers/SymbolIssueHandler.cs b/Engine/Source/Programs/Shared/EpicGames.Horde/Issues/Handlers/SymbolIssueHandler.cs
--- a/Engine/Source/Programs/Shared/EpicGames.Horde/Issues/Handlers/SymbolIssueHandler.cs
+++ b/Engine/Source/Programs/Shared/EpicGames.Horde/Issues/Handlers/SymbolIssueHandler.cs
@@ -47,7 +47,7 @@ namespace EpicGames.Horde.Issues.Handlers
 				EventId eventId = issueEvent.EventId.Value;
 				if (IsMatchingEventId(eventId))
 				{
-					IssueEventGroup issue = new IssueEventGroup("Symbol", "{LegacySymbolIssueHandler}", IssueChangeFilter.Code);
+					IssueEventGroup issue = new("Symbol", "{LegacySymbolIssueHandler}", IssueChangeFilter.Code);
 					issue.Events.Add(issueEvent);
 					issue.Keys.AddSymbols(issueEvent);
 					issue.Metadata.Add(EventIdName, eventId.Id.ToString());
