Shared: prefer implicit object creation when type is apparent

diff --git a/Engine/Source/Programs/Shared/EpicGames.Horde/Issues/Handlers/SystemicIssueHandler.cs b/Engine/Source/Programs/Shared/EpicGames.Horde/Issues/Handlers/SystemicIssueHandler.cs
--- a/Engine/Source/Programs/Shared/EpicGames.Horde/Issues/Handlers/SystemicIssueHandler.cs
+++ b/Engine/Source/Programs/Shared/EpicGames.Horde/Issues/Handlers/SystemicIssueHandler.cs
@@ -52,7 +52,7 @@ namespace EpicGames.Horde.Issues.Handlers
 					}
 					else
 					{
-						IssueEventGroup issue = new IssueEventGroup("Systemic", "Systemic {Severity} in {Nodes}", IssueChangeFilter.None);
+						IssueEventGroup issue = new("Systemic", "Systemic {Severity} in {Nodes}", IssueChangeFilter.None);
 						issue.Keys.Add(IssueKey.FromStep(_context.StreamId, _context.TemplateId, _context.NodeName));
 						issue.Events.Add(logEvent);
 						_issues.Add(issue);
diff --git a/Engine/Source/Programs/Shared/EpicGames.Horde/Issues/Handlers/UnacceptableWordsIssueHandler.cs b/Engine/Source/Programs/Shared/EpicGames.Horde/Issues/Handlers/UnacceptableWordsIssueHandler.cs
--- a/Engine/Source/Programs/Shared/EpicGames.Horde/Issues/Handlers/UnacceptableWordsIssueHandler.cs
+++ b/Engine/Source/Programs/Shared/EpicGames.Horde/Issues/Handlers/UnacceptableWordsIssueHandler.cs
@@ -22,7 +22,7 @@ namespace EpicGames.Horde.Issues.Handlers
 		{
 			if (issueEvent.EventId == KnownLogEvents.AutomationTool_UnacceptableWords)
 			{
-				IssueEventGroup issue = new IssueEventGroup("UnacceptableWords", "Unacceptable words in {Files}", IssueChangeFilter.Code);
+				IssueEventGroup issue = new("UnacceptableWords", "Unacceptable words in {Files}", IssueChangeFilter.Code);
 				issue.Events.Add(issueEvent);
 				issue.Keys.AddSourceFiles(issueEvent);
 				_issues.Add(issue);
diff --git a/Engine/Source/Programs/Shared/EpicGames.Horde/Issues/Handlers/WorldLeakIssueHandler.cs b/Engine/Source/Programs/Shared/EpicGames.Horde/Issues/Handlers/WorldLeakIssueHandler.cs
--- a/Engine/Source/Programs/Shared/EpicGames.Horde/Issues/Handlers/WorldLeakIssueHandler.cs
+++ b/Engine/Source/Programs/Shared/EpicGames.Horde/Issues/Handlers/WorldLeakIssueHandler.cs
@@ -14,9 +14,9 @@ namespace EpicGames.Horde.Issues.Handlers
 	[IssueHandler]
 	public class WorldLeakIssueHandler : IssueHandler
 	{
-		static readonly Utf8String s_objectFlagsAnnotation = new Utf8String("ObjectFlags");
-		static readonly Utf8String s_objectTextAnnotation = new Utf8String("ObjectText");
-		static readonly Utf8String s_lingeringReferenceAnnotation = new Utf8String("LingeringReference");
+		static readonly Utf8String s_objectFlagsAnnotation = new("ObjectFlags");
+		static readonly Utf8String s_objectTextAnnotation = new("ObjectText");
+		static readonly Utf8String s_lingeringReferenceAnnotation = new("LingeringReference");
 		readonly List<IssueEventGroup> _issues = [];
 
 		/// <inheritdoc/>
@@ -24,17 +24,17 @@ namespace EpicGames.Horde.Issues.Handlers
 		/// <summary>
 		/// Log value describing the object reference preventing garbage collection
 		/// </summary>
-		public static LogValue LingeringReference => new LogValue(s_lingeringReferenceAnnotation, "");
+		public static LogValue LingeringReference => new(s_lingeringReferenceAnnotation, "");
 
 		/// <summary>
 		/// Log value describing the flag on an object in the leak reference chain
 		/// </summary>
-		public static LogValue ObjectFlags => new LogValue(s_objectFlagsAnnotation, "");
+		public static LogValue ObjectFlags => new(s_objectFlagsAnnotation, "");
 
 		/// <summary>
 		/// Log value describing an object in the leak reference chain
 		/// </summary>
-		public static LogValue ObjectText => new LogValue(s_objectTextAnnotation, "");
+		public static LogValue ObjectText => new(s_objectTextAnnotation, "");
 
 		static bool FindAnnotation(IssueEvent issueEvent, Utf8String annotation, ref string result)
 		{
@@ -66,7 +66,7 @@ namespace EpicGames.Horde.Issues.Handlers
 
 				string fingerprint = $"WorldLeak:{lingeringReference}:{firstObjectFlags} {firstObjectText}";
 
-				IssueEventGroup issue = new IssueEventGroup(fingerprint, $"World Leak Detected: '{{Meta:{s_lingeringReferenceAnnotation}}}' is preventing GC. Reference chain begins from '{{Meta:{s_objectFlagsAnnotation}}} {{Meta:{s_objectTextAnnotation}}}'", IssueChangeFilter.All);
+				IssueEventGroup issue = new(fingerprint, $"World Leak Detected: '{{Meta:{s_lingeringReferenceAnnotation}}}' is preventing GC. Reference chain begins from '{{Meta:{s_objectFlagsAnnotation}}} {{Meta:{s_objectTextAnnotation}}}'", IssueChangeFilter.All);
 				issue.Keys.Add(lingeringReference, IssueKeyType.Hash);
 				issue.Metadata.Add(s_lingeringReferenceAnnotation.ToString(), lingeringReference);
 				issue.Keys.Add(firstObjectFlags, IssueKeyType.None);
diff --git a/Engine/Source/Programs/Shared/EpicGames.Horde/Issues/IssueEvent.cs b/Engine/Source/Programs/Shared/EpicGames.Horde/Issues/IssueEvent.cs
--- a/Engine/Source/Programs/Shared/EpicGames.Horde/Issues/IssueEvent.cs
+++ b/Engine/Source/Programs/Shared/EpicGames.Horde/Issues/IssueEvent.cs
@@ -179,7 +179,7 @@ namespace EpicGames.Horde.Issues
 				}
 			}
 
-			IssueAuditLogEntry entry = new IssueAuditLogEntry(logLevel, messageTemplate ?? "", [.. args]);
+			IssueAuditLogEntry entry = new(logLevel, messageTemplate ?? "", [.. args]);
 			_entries.Add(entry);
 		}
 	}
diff --git a/Engine/Source/Programs/Shared/EpicGames.Horde/Issues/IssueExtensions.cs b/Engine/Source/Programs/Shared/EpicGames.Horde/Issues/IssueExtensions.cs
--- a/Engine/Source/Programs/Shared/EpicGames.Horde/Issues/IssueExtensions.cs
+++ b/Engine/Source/Programs/Shared/EpicGames.Horde/Issues/IssueExtensions.cs
@@ -28,7 +28,7 @@ namespace EpicGames.Horde.Issues
 		/// <returns>Disposable object for the lifetime of this scope</returns>
 		public static IDisposable? BeginIssueScope(this ILogger logger, IssueFingerprint fingerprint)
 		{
-			Dictionary<string, object> properties = new Dictionary<string, object>(1);
+			Dictionary<string, object> properties = new(1);
 			properties.Add("@$issue", fingerprint);
 
 			return logger.BeginScope(properties);
diff --git a/Engine/Source/Programs/Shared/EpicGames.Horde/Issues/IssueKey.cs b/Engine/Source/Programs/Shared/EpicGames.Horde/Issues/IssueKey.cs
--- a/Engine/Source/Programs/Shared/EpicGames.Horde/Issues/IssueKey.cs
+++ b/Engine/Source/Programs/Shared/EpicGames.Horde/Issues/IssueKey.cs
@@ -93,27 +93,27 @@ namespace EpicGames.Horde.Issues
 		/// <summary>
 		/// Creates an issue key for a file
 		/// </summary>
-		public static IssueKey FromFile(string file, bool note = false) => new IssueKey(file, note ? IssueKeyType.Note : IssueKeyType.File);
+		public static IssueKey FromFile(string file, bool note = false) => new(file, note ? IssueKeyType.Note : IssueKeyType.File);
 
 		/// <summary>
 		/// Creates an issue key for a file
 		/// </summary>
-		public static IssueKey FromSymbol(string name) => new IssueKey(name, IssueKeyType.Symbol);
+		public static IssueKey FromSymbol(string name) => new(name, IssueKeyType.Symbol);
 
 		/// <summary>
 		/// Creates an issue key for a particular hash
 		/// </summary>
-		public static IssueKey FromHash(Md5Hash hash) => new IssueKey(hash.ToString(), IssueKeyType.Hash);
+		public static IssueKey FromHash(Md5Hash hash) => new(hash.ToString(), IssueKeyType.Hash);
 
 		/// <summary>
 		/// Creates an issue key for a particular step
 		/// </summary>
-		public static IssueKey FromStep(StreamId streamId, TemplateId templateId, string nodeName) => new IssueKey($"{streamId}:{templateId}:{nodeName}", IssueKeyType.Step);
+		public static IssueKey FromStep(StreamId streamId, TemplateId templateId, string nodeName) => new($"{streamId}:{templateId}:{nodeName}", IssueKeyType.Step);
 
 		/// <summary>
 		/// Creates an issue key for a particular step and severity
 		/// </summary>
-		public static IssueKey FromStepAndSeverity(StreamId streamId, TemplateId templateId, string nodeName, LogLevel severity) => new IssueKey($"{streamId}:{templateId}:{nodeName}:{severity}", IssueKeyType.Step);
+		public static IssueKey FromStepAndSeverity(StreamId streamId, TemplateId templateId, string nodeName, LogLevel severity) => new($"{streamId}:{templateId}:{nodeName}:{severity}", IssueKeyType.Step);
 
 		/// <inheritdoc/>
 		public bool Equals(IssueKey? other) => other is not null && other.Name.Equals(Name, StringComparison.OrdinalIgnoreCase) && other.Type == Type && String.Equals(Scope, other.Scope, StringComparison.Ordinal);
@@ -163,7 +163,7 @@ namespace EpicGames.Horde.Issues
 				{
 					int endIdx = relativePath.LastIndexOfAny(['/', '\\']) + 1;
 					string fileName = relativePath.Substring(endIdx);
-					IssueKey issueKey = new IssueKey(fileName, IssueKeyType.File);
+					IssueKey issueKey = new(fileName, IssueKeyType.File);
 					keys.Add(issueKey);
 				}
 			}
@@ -247,7 +247,7 @@ namespace EpicGames.Horde.Issues
 			int endIdx = relativePath.LastIndexOfAny(['/', '\\']) + 1;
 
 			string fileName = relativePath.Substring(endIdx);
-			IssueKey key = new IssueKey(fileName, type, scope);
+			IssueKey key = new(fileName, type, scope);
 
 			keys.Add(key);
 		}
@@ -266,7 +266,7 @@ namespace EpicGames.Horde.Issues
 				string? identifier;
 				if (document.RootElement.TryGetNestedProperty("properties.symbol.identifier", out identifier))
 				{
-					IssueKey key = new IssueKey(identifier, IssueKeyType.Symbol);
+					IssueKey key = new(identifier, IssueKeyType.Symbol);
 					keys.Add(key);
 				}
 			}
diff --git a/Engine/Source/Programs/Shared/EpicGames.Horde/Issues/WorkflowId.cs b/Engine/Source/Programs/Shared/EpicGames.Horde/Issues/WorkflowId.cs
--- a/Engine/Source/Programs/Shared/EpicGames.Horde/Issues/WorkflowId.cs
+++ b/Engine/Source/Programs/Shared/EpicGames.Horde/Issues/WorkflowId.cs
@@ -41,7 +41,7 @@ namespace EpicGames.Horde.Issues
 	class WorkflowIdConverter : StringIdConverter<WorkflowId>
 	{
 		/// <inheritdoc/>
-		public override WorkflowId FromStringId(StringId id) => new WorkflowId(id);
+		public override WorkflowId FromStringId(StringId id) => new(id);
 
 		/// <inheritdoc/>
 		public override StringId ToStringId(WorkflowId value) => value.Id;
diff --git a/Engine/Source/Programs/Shared/EpicGames.Horde/Jobs/Bisect/BisectTaskId.cs b/Engine/Source/Programs/Shared/EpicGames.Horde/Jobs/Bisect/BisectTaskId.cs
--- a/Engine/Source/Programs/Shared/EpicGames.Horde/Jobs/Bisect/BisectTaskId.cs
+++ b/Engine/Source/Programs/Shared/EpicGames.Horde/Jobs/Bisect/BisectTaskId.cs
@@ -16,7 +16,7 @@ namespace EpicGames.Horde.Jobs.Bisect
 	public readonly record struct BisectTaskId(BinaryId Id)
 	{
 		/// <inheritdoc cref="BinaryId.Parse(System.String)"/>
-		public static BisectTaskId Parse(string text) => new BisectTaskId(BinaryId.Parse(text));
+		public static BisectTaskId Parse(string text) => new(BinaryId.Parse(text));
 
 		/// <inheritdoc/>
 		public override string ToString() => Id.ToString();
@@ -28,7 +28,7 @@ namespace EpicGames.Horde.Jobs.Bisect
 	class BisectTaskIdConverter : BinaryIdConverter<BisectTaskId>
 	{
 		/// <inheritdoc/>
-		public override BisectTaskId FromBinaryId(BinaryId id) => new BisectTaskId(id);
+		public override BisectTaskId FromBinaryId(BinaryId id) => new(id);
 
 		/// <inheritdoc/>
 		public override BinaryId ToBinaryId(BisectTaskId value) => value.Id;
diff --git a/Engine/Source/Programs/Shared/EpicGames.Horde/Jobs/IJob.cs b/Engine/Source/Programs/Shared/EpicGames.Horde/Jobs/IJob.cs
--- a/Engine/Source/Programs/Shared/EpicGames.Horde/Jobs/IJob.cs
+++ b/Engine/Source/Programs/Shared/EpicGames.Horde/Jobs/IJob.cs
@@ -616,7 +616,7 @@ namespace EpicGames.Horde.Jobs
 			{
 				foreach (IJobStep batchStep in batch.Steps)
 				{
-					NodeRef batchNodeRef = new NodeRef(batch.GroupIdx, batchStep.NodeIdx);
+					NodeRef batchNodeRef = new(batch.GroupIdx, batchStep.NodeIdx);
 					stepForNode[batchNodeRef] = batchStep;
 				}
 			}
@@ -668,7 +668,7 @@ namespace EpicGames.Horde.Jobs
 				INodeGroup group = graph.Groups[batch.GroupIdx];
 
 				// Step through the batch, keeping track of the time that things finish.
-				TimingInfo timingInfo = new TimingInfo();
+				TimingInfo timingInfo = new();
 
 				// Wait for the dependencies for the batch to start
 				HashSet<INode> dependencyNodes = batch.GetStartDependencies(graph.Groups);
@@ -736,7 +736,7 @@ namespace EpicGames.Horde.Jobs
 					timingInfo.AverageTotalTimeToComplete = timingInfo.AverageTotalTimeToComplete + NullableTimeSpanFromSeconds(stepTimingInfo?.AverageDuration);
 
 					// Add it to the lookup
-					TimingInfo nodeTimingInfo = new TimingInfo(timingInfo);
+					TimingInfo nodeTimingInfo = new(timingInfo);
 					nodeTimingInfo.StepTiming = stepTimingInfo;
 					nodeToTimingInfo[node] = nodeTimingInfo;
 				}
@@ -753,7 +753,7 @@ namespace EpicGames.Horde.Jobs
 		/// <remarks>Any steps not in a terminal state (<see cref="JobStepState.Completed"/>, <see cref="JobStepState.Aborted"/>, <see cref="JobStepState.Skipped"/>) are excluded in total counts except <see cref="JobStepsCompletionInfo.StepTotalCount"/>.</remarks>
 		public static JobStepsCompletionInfo GetStepCompletionInfo(this IJob job)
 		{
-			JobStepsCompletionInfo jobStepsCompletionInfo = new JobStepsCompletionInfo();
+			JobStepsCompletionInfo jobStepsCompletionInfo = new();
 
 			foreach (IJobStepBatch batch in job.Batches)
 			{
@@ -972,7 +972,7 @@ namespace EpicGames.Horde.Jobs
 				}
 
 				// Create the response
-				GetLabelStateResponse response = new GetLabelStateResponse();
+				GetLabelStateResponse response = new();
 				response.DashboardName = label.DashboardName;
 				response.DashboardCategory = label.DashboardCategory;
 				response.UgsName = label.UgsName;
@@ -1001,7 +1001,7 @@ namespace EpicGames.Horde.Jobs
 				INode node = graph.Groups[0].Nodes[0];
 				if (node.Name == IJob.SetupNodeName)
 				{
-					NodeRef nodeRef = new NodeRef(0, 0);
+					NodeRef nodeRef = new(0, 0);
 					if (unlabeledNodes.Contains(nodeRef))
 					{
 						IJobStep step = stepForNode[nodeRef];
@@ -1016,7 +1016,7 @@ namespace EpicGames.Horde.Jobs
 			// Add a response for everything not included elsewhere.
 			GetLabelState(unlabeledNodes, stepForNode, out LabelState otherState, out LabelOutcome otherOutcome);
 
-			GetDefaultLabelStateResponse defaultResponse = new GetDefaultLabelStateResponse();
+			GetDefaultLabelStateResponse defaultResponse = new();
 			defaultResponse.State = otherState;
 			defaultResponse.Outcome = otherOutcome;
 			defaultResponse.Nodes.AddRange(unlabeledNodes.Select(x => graph.GetNode(x).Name));
@@ -1798,7 +1798,7 @@ namespace EpicGames.Horde.Jobs
 		/// <returns>New TimingInfo instance</returns>
 		public static TimingInfo Max(IEnumerable<TimingInfo> others)
 		{
-			TimingInfo timingInfo = new TimingInfo();
+			TimingInfo timingInfo = new();
 			timingInfo.WaitForAll(others);
 			return timingInfo;
 		}
diff --git a/Engine/Source/Programs/Shared/EpicGames.Horde/Jobs/JobId.cs b/Engine/Source/Programs/Shared/EpicGames.Horde/Jobs/JobId.cs
--- a/Engine/Source/Programs/Shared/EpicGames.Horde/Jobs/JobId.cs
+++ b/Engine/Source/Programs/Shared/EpicGames.Horde/Jobs/JobId.cs
@@ -21,7 +21,7 @@ namespace EpicGames.Horde.Jobs
 		public static JobId Empty { get; } = default;
 
 		/// <inheritdoc cref="BinaryId.Parse(System.String)"/>
-		public static JobId Parse(string text) => new JobId(BinaryId.Parse(text));
+		public static JobId Parse(string text) => new(BinaryId.Parse(text));
 
 		/// <inheritdoc/>
 		public override string ToString() => Id.ToString();
@@ -33,7 +33,7 @@ namespace EpicGames.Horde.Jobs
 	class JobIdConverter : BinaryIdConverter<JobId>
 	{
 		/// <inheritdoc/>
-		public override JobId FromBinaryId(BinaryId id) => new JobId(id);
+		public override JobId FromBinaryId(BinaryId id) => new(id);
 
 		/// <inheritdoc/>
 		public override BinaryId ToBinaryId(JobId value) => value.Id;
diff --git a/Engine/Source/Programs/Shared/EpicGames.Horde/Jobs/JobStepBatchId.cs b/Engine/Source/Programs/Shared/EpicGames.Horde/Jobs/JobStepBatchId.cs
--- a/Engine/Source/Programs/Shared/EpicGames.Horde/Jobs/JobStepBatchId.cs
+++ b/Engine/Source/Programs/Shared/EpicGames.Horde/Jobs/JobStepBatchId.cs
@@ -24,10 +24,10 @@ namespace EpicGames.Horde.Jobs
 		/// <summary>
 		/// Creates a new <see cref="JobId"/>
 		/// </summary>
-		public static JobStepBatchId GenerateNewId() => new JobStepBatchId(SubResourceId.GenerateNewId());
+		public static JobStepBatchId GenerateNewId() => new(SubResourceId.GenerateNewId());
 
 		/// <inheritdoc cref="SubResourceId.Parse(System.String)"/>
-		public static JobStepBatchId Parse(string text) => new JobStepBatchId(SubResourceId.Parse(text));
+		public static JobStepBatchId Parse(string text) => new(SubResourceId.Parse(text));
 
 		/// <inheritdoc/>
 		public override string ToString() => SubResourceId.ToString();
@@ -39,7 +39,7 @@ namespace EpicGames.Horde.Jobs
 	class JobStepBatchIdConverter : SubResourceIdConverter<JobStepBatchId>
 	{
 		/// <inheritdoc/>
-		public override JobStepBatchId FromSubResourceId(SubResourceId subResourceId) => new JobStepBatchId(subResourceId);
+		public override JobStepBatchId FromSubResourceId(SubResourceId subResourceId) => new(subResourceId);
 
 		/// <inheritdoc/>
 		public override SubResourceId ToSubResourceId(JobStepBatchId value) => value.SubResourceId;
diff --git a/Engine/Source/Programs/Shared/EpicGames.Horde/Jobs/JobStepId.cs b/Engine/Source/Programs/Shared/EpicGames.Horde/Jobs/JobStepId.cs
--- a/Engine/Source/Programs/Shared/EpicGames.Horde/Jobs/JobStepId.cs
+++ b/Engine/Source/Programs/Shared/EpicGames.Horde/Jobs/JobStepId.cs
@@ -24,10 +24,10 @@ namespace EpicGames.Horde.Jobs
 		/// <summary>
 		/// Creates a new <see cref="JobId"/>
 		/// </summary>
-		public static JobStepId GenerateNewId() => new JobStepId(SubResourceId.GenerateNewId());
+		public static JobStepId GenerateNewId() => new(SubResourceId.GenerateNewId());
 
 		/// <inheritdoc cref="SubResourceId.Parse(System.String)"/>
-		public static JobStepId Parse(string text) => new JobStepId(SubResourceId.Parse(text));
+		public static JobStepId Parse(string text) => new(SubResourceId.Parse(text));
 
 		/// <inheritdoc cref="SubResourceId.TryParse(System.String, out SubResourceId)"/>
 		public static bool TryParse(string text, out JobStepId stepId)
@@ -54,7 +54,7 @@ namespace EpicGames.Horde.Jobs
 	class JobStepIdConverter : SubResourceIdConverter<JobStepId>
 	{
 		/// <inheritdoc/>
-		public override JobStepId FromSubResourceId(SubResourceId id) => new JobStepId(id);
+		public override JobStepId FromSubResourceId(SubResourceId id) => new(id);
 
 		/// <inheritdoc/>
 		public override SubResourceId ToSubResourceId(JobStepId value) => value.Id;
diff --git a/Engine/Source/Programs/Shared/EpicGames.Horde/Jobs/Schedules/ISchedule.cs b/Engine/Source/Programs/Shared/EpicGames.Horde/Jobs/Schedules/ISchedule.cs
--- a/Engine/Source/Programs/Shared/EpicGames.Horde/Jobs/Schedules/ISchedule.cs
+++ b/Engine/Source/Programs/Shared/EpicGames.Horde/Jobs/Schedules/ISchedule.cs
@@ -187,7 +187,7 @@ namespace EpicGames.Horde.Jobs.Schedules
 			DateTimeOffset lastTime = TimeZoneInfo.ConvertTime((DateTimeOffset)lastTimeUtc, timeZone);
 
 			// Get the base time (ie. the start of this day) for anchoring the schedule
-			DateTimeOffset baseTime = new DateTimeOffset(lastTime.Year, lastTime.Month, lastTime.Day, 0, 0, 0, lastTime.Offset);
+			DateTimeOffset baseTime = new(lastTime.Year, lastTime.Month, lastTime.Day, 0, 0, 0, lastTime.Offset);
 			for (; ; )
 			{
 				if (pattern.DaysOfWeek == null || pattern.DaysOfWeek.Contains(baseTime.DayOfWeek))
diff --git a/Engine/Source/Programs/Shared/EpicGames.Horde/Jobs/SubResourceIdConverter.cs b/Engine/Source/Programs/Shared/EpicGames.Horde/Jobs/SubResourceIdConverter.cs
--- a/Engine/Source/Programs/Shared/EpicGames.Horde/Jobs/SubResourceIdConverter.cs
+++ b/Engine/Source/Programs/Shared/EpicGames.Horde/Jobs/SubResourceIdConverter.cs
@@ -48,7 +48,7 @@ namespace EpicGames.Horde.Jobs
 	/// </summary>
 	public sealed class SubResourceIdTypeConverter<TValue, TConverter> : TypeConverter where TValue : struct where TConverter : SubResourceIdConverter<TValue>, new()
 	{
-		readonly TConverter _converter = new TConverter();
+		readonly TConverter _converter = new();
 
 		/// <inheritdoc/>
 		public override bool CanConvertFrom(ITypeDescriptorContext? context, Type sourceType)
@@ -93,7 +93,7 @@ namespace EpicGames.Horde.Jobs
 	/// </summary>
 	public sealed class SubResourceIdJsonConverter<TValue, TConverter> : JsonConverter<TValue> where TValue : struct where TConverter : SubResourceIdConverter<TValue>, new()
 	{
-		readonly TConverter _converter = new TConverter();
+		readonly TConverter _converter = new();
 
 		/// <inheritdoc/>
 		public override TValue Read(ref Utf8JsonReader reader, Type typeToConvert, JsonSerializerOptions options) => _converter.FromSubResourceId(SubResourceId.Parse(reader.GetString() ?? String.Empty));
diff --git a/Engine/Source/Programs/Shared/EpicGames.Horde/Jobs/Templates/ParameterId.cs b/Engine/Source/Programs/Shared/EpicGames.Horde/Jobs/Templates/ParameterId.cs
--- a/Engine/Source/Programs/Shared/EpicGames.Horde/Jobs/Templates/ParameterId.cs
+++ b/Engine/Source/Programs/Shared/EpicGames.Horde/Jobs/Templates/ParameterId.cs
@@ -21,7 +21,7 @@ namespace EpicGames.Horde.Jobs.Templates
 
 		/// <inheritdoc cref="StringId.Sanitize"/>
 		public static ParameterId Sanitize(string name)
-			=> new ParameterId(StringId.Sanitize(name));
+			=> new(StringId.Sanitize(name));
 
 		/// <inheritdoc/>
 		public override string ToString() => Id.ToString();
@@ -33,7 +33,7 @@ namespace EpicGames.Horde.Jobs.Templates
 	class ParameterIdConverter : StringIdConverter<ParameterId>
 	{
 		/// <inheritdoc/>
-		public override ParameterId FromStringId(StringId id) => new ParameterId(id);
+		public override ParameterId FromStringId(StringId id) => new(id);
 
 		/// <inheritdoc/>
 		public override StringId ToStringId(ParameterId value) => value.Id;
diff --git a/Engine/Source/Programs/Shared/EpicGames.Horde/Jobs/Templates/TemplateId.cs b/Engine/Source/Programs/Shared/EpicGames.Horde/Jobs/Templates/TemplateId.cs
--- a/Engine/Source/Programs/Shared/EpicGames.Horde/Jobs/Templates/TemplateId.cs
+++ b/Engine/Source/Programs/Shared/EpicGames.Horde/Jobs/Templates/TemplateId.cs
@@ -37,7 +37,7 @@ namespace EpicGames.Horde.Jobs.Templates
 	class TemplateIdConverter : StringIdConverter<TemplateId>
 	{
 		/// <inheritdoc/>
-		public override TemplateId FromStringId(StringId id) => new TemplateId(id);
+		public override TemplateId FromStringId(StringId id) => new(id);
 
 		/// <inheritdoc/>
 		public override StringId ToStringId(TemplateId value) => value.Id;
diff --git a/Engine/Source/Programs/Shared/EpicGames.Horde/Logs/LogChunkNode.cs b/Engine/Source/Programs/Shared/EpicGames.Horde/Logs/LogChunkNode.cs
--- a/Engine/Source/Programs/Shared/EpicGames.Horde/Logs/LogChunkNode.cs
+++ b/Engine/Source/Programs/Shared/EpicGames.Horde/Logs/LogChunkNode.cs
@@ -117,14 +117,14 @@ namespace EpicGames.Horde.Logs
 		/// </summary>
 		/// <param name="idx">Index of the line to retrieve</param>
 		/// <returns>Line at the given index</returns>
-		public Utf8String GetLine(int idx) => new Utf8String(Data.Slice(LineOffsets[idx], LineOffsets[idx + 1] - LineOffsets[idx] - 1));
+		public Utf8String GetLine(int idx) => new(Data.Slice(LineOffsets[idx], LineOffsets[idx + 1] - LineOffsets[idx] - 1));
 
 		/// <summary>
 		/// Accessor for an individual line, including the trailing newline character
 		/// </summary>
 		/// <param name="idx">Index of the line to retrieve</param>
 		/// <returns>Line at the given index</returns>
-		public Utf8String GetLineWithNewline(int idx) => new Utf8String(Data.Slice(LineOffsets[idx], LineOffsets[idx + 1] - LineOffsets[idx]));
+		public Utf8String GetLineWithNewline(int idx) => new(Data.Slice(LineOffsets[idx], LineOffsets[idx + 1] - LineOffsets[idx]));
 
 		/// <summary>
 		/// Find the line index for a particular offset
@@ -342,7 +342,7 @@ namespace EpicGames.Horde.Logs
 		/// </summary>
 		/// <param name="index">Index of the line</param>
 		/// <returns>Text for the line</returns>
-		public Utf8String GetLine(int index) => new Utf8String(_data.AsMemory(_lineOffsets[index], _lineOffsets[index + 1] - _lineOffsets[index]));
+		public Utf8String GetLine(int index) => new(_data.AsMemory(_lineOffsets[index], _lineOffsets[index + 1] - _lineOffsets[index]));
 
 		/// <summary>
 		/// Create a new chunk data object with the given data appended. The internal buffers are reused, with the assumption that
@@ -440,7 +440,7 @@ namespace EpicGames.Horde.Logs
 				return outputOffset + 1;
 			}
 
-			Utf8JsonReader reader = new Utf8JsonReader(input);
+			Utf8JsonReader reader = new(input);
 			if (reader.Read() && reader.TokenType == JsonTokenType.StartObject)
 			{
 				while (reader.Read() && reader.TokenType == JsonTokenType.PropertyName)
@@ -554,7 +554,7 @@ namespace EpicGames.Horde.Logs
 		/// Create a <see cref="LogChunkNode"/> object from the current state
 		/// </summary>
 		/// <returns></returns>
-		public LogChunkNode ToLogChunk() => new LogChunkNode(_data.AsMemory(0, _length).ToArray(), [.. _lineOffsets]);
+		public LogChunkNode ToLogChunk() => new(_data.AsMemory(0, _length).ToArray(), [.. _lineOffsets]);
 	}
 
 	/// <summary>
@@ -764,7 +764,7 @@ namespace EpicGames.Horde.Logs
 		/// <returns>The plain text instance</returns>
 		public static LogChunkNode ConvertJsonToPlainText(this LogChunkNode logText, ILogger logger)
 		{
-			LogChunkBuilder other = new LogChunkBuilder();
+			LogChunkBuilder other = new();
 			other.AppendJsonAsPlainText(logText, 0, logText.LineCount, logger);
 			return other.ToLogChunk();
 		}
diff --git a/Engine/Source/Programs/Shared/EpicGames.Horde/Logs/LogId.cs b/Engine/Source/Programs/Shared/EpicGames.Horde/Logs/LogId.cs
--- a/Engine/Source/Programs/Shared/EpicGames.Horde/Logs/LogId.cs
+++ b/Engine/Source/Programs/Shared/EpicGames.Horde/Logs/LogId.cs
@@ -21,7 +21,7 @@ namespace EpicGames.Horde.Logs
 		public static LogId Empty { get; } = new LogId(default);
 
 		/// <inheritdoc cref="BinaryId.Parse(System.String)"/>
-		public static LogId Parse(string text) => new LogId(BinaryId.Parse(text));
+		public static LogId Parse(string text) => new(BinaryId.Parse(text));
 
 		/// <inheritdoc cref="BinaryId.TryParse(System.String, out BinaryId)"/>
 		public static bool TryParse(string text, out LogId id)
@@ -48,7 +48,7 @@ namespace EpicGames.Horde.Logs
 	class LogIdConverter : BinaryIdConverter<LogId>
 	{
 		/// <inheritdoc/>
-		public override LogId FromBinaryId(BinaryId id) => new LogId(id);
+		public override LogId FromBinaryId(BinaryId id) => new(id);
 
 		/// <inheritdoc/>
 		public override BinaryId ToBinaryId(LogId value) => value.Id;
diff --git a/Engine/Source/Programs/Shared/EpicGames.Horde/Logs/LogIndexNode.cs b/Engine/Source/Programs/Shared/EpicGames.Horde/Logs/LogIndexNode.cs
--- a/Engine/Source/Programs/Shared/EpicGames.Horde/Logs/LogIndexNode.cs
+++ b/Engine/Source/Programs/Shared/EpicGames.Horde/Logs/LogIndexNode.cs
@@ -116,7 +116,7 @@ namespace EpicGames.Horde.Logs
 			}
 
 			// Add the existing index data into a new ngram set
-			NgramSetBuilder newNgramSetBuilder = new NgramSetBuilder();
+			NgramSetBuilder newNgramSetBuilder = new();
 
 			ulong blockMask = ((1UL << NumChunkBits) - 1);
 			foreach (ulong value in NgramSet)
@@ -335,7 +335,7 @@ namespace EpicGames.Horde.Logs
 		{
 			byte[] textCopy = text.ToArray();
 
-			Lazy<HashSet<int>> lazyChunks = new Lazy<HashSet<int>>(() =>
+			Lazy<HashSet<int>> lazyChunks = new(() =>
 			{
 				HashSet<int> union = [];
 				for (int shift = 0; shift < Ngram.MaxBytes; shift++)
diff --git a/Engine/Source/Programs/Shared/EpicGames.Horde/Logs/LogNode.cs b/Engine/Source/Programs/Shared/EpicGames.Horde/Logs/LogNode.cs
--- a/Engine/Source/Programs/Shared/EpicGames.Horde/Logs/LogNode.cs
+++ b/Engine/Source/Programs/Shared/EpicGames.Horde/Logs/LogNode.cs
@@ -149,7 +149,7 @@ namespace EpicGames.Horde.Logs
 		readonly LogChunkSequenceBuilder _indexTextBuilder;
 
 		// Lock object for access to the fields above
-		readonly object _lockObject = new object();
+		readonly object _lockObject = new();
 
 		// Inner log device; used for messaging encoding errors.
 		readonly ILogger _logger;
@@ -294,7 +294,7 @@ namespace EpicGames.Horde.Logs
 				length += writeTextChunk.Length;
 			}
 
-			LogNode newRoot = new LogNode(_format, lineCount, length, newJsonChunkRefs, newIndexRef, complete);
+			LogNode newRoot = new(_format, lineCount, length, newJsonChunkRefs, newIndexRef, complete);
 			IHashedBlobRef<LogNode> newRootRef = await writer.WriteBlobAsync(newRoot, cancellationToken);
 
 			await writer.FlushAsync(cancellationToken);
diff --git a/Engine/Source/Programs/Shared/EpicGames.Horde/Logs/Ngram.cs b/Engine/Source/Programs/Shared/EpicGames.Horde/Logs/Ngram.cs
--- a/Engine/Source/Programs/Shared/EpicGames.Horde/Logs/Ngram.cs
+++ b/Engine/Source/Programs/Shared/EpicGames.Horde/Logs/Ngram.cs
@@ -80,7 +80,7 @@ namespace EpicGames.Horde.Logs
 		/// <returns></returns>
 		public static NgramSet Decompose(ReadOnlySpan<byte> text)
 		{
-			NgramSetBuilder builder = new NgramSetBuilder();
+			NgramSetBuilder builder = new();
 			Decompose(text, builder.Add);
 			return builder.ToNgramSet();
 		}
diff --git a/Engine/Source/Programs/Shared/EpicGames.Horde/Logs/ServerLogPacketBuilder.cs b/Engine/Source/Programs/Shared/EpicGames.Horde/Logs/ServerLogPacketBuilder.cs
--- a/Engine/Source/Programs/Shared/EpicGames.Horde/Logs/ServerLogPacketBuilder.cs
+++ b/Engine/Source/Programs/Shared/EpicGames.Horde/Logs/ServerLogPacketBuilder.cs
@@ -38,7 +38,7 @@ namespace EpicGames.Horde.Logs
 		static readonly string s_linePropertyName = LogEventPropertyName.Line.ToString();
 		static readonly string s_lineCountPropertyName = LogEventPropertyName.LineCount.ToString();
 
-		static readonly Utf8String s_escapedNewline = new Utf8String("\\n");
+		static readonly Utf8String s_escapedNewline = new("\\n");
 
 		readonly List<JsonLogEvent> _logEvents = [];
 		readonly ArrayBufferWriter<byte> _lineWriter;
@@ -111,7 +111,7 @@ namespace EpicGames.Horde.Logs
 			}
 			catch (Exception ex)
 			{
-				StringBuilder escapedLineBuilder = new StringBuilder();
+				StringBuilder escapedLineBuilder = new();
 
 				ReadOnlySpan<byte> span = jsonLogEvent.Data.Span;
 				for (int idx = 0; idx < span.Length; idx++)
@@ -128,8 +128,8 @@ namespace EpicGames.Horde.Logs
 
 				string escapedLine = escapedLineBuilder.ToString();
 				KeyValuePair<string, object?>[] properties = [new KeyValuePair<string, object?>("Text", escapedLine)];
-				LogEvent newLogEvent = new LogEvent(DateTime.UtcNow, LogLevel.Error, default, $"Invalid json log event: {escapedLineBuilder}", "Invalid json log event: {Text}", properties, LogException.FromException(ex));
-				JsonLogEvent newJsonLogEvent = new JsonLogEvent(newLogEvent);
+				LogEvent newLogEvent = new(DateTime.UtcNow, LogLevel.Error, default, $"Invalid json log event: {escapedLineBuilder}", "Invalid json log event: {Text}", properties, LogException.FromException(ex));
+				JsonLogEvent newJsonLogEvent = new(newLogEvent);
 
 				return SanitizeAndWriteEventInternal(newJsonLogEvent);
 			}
@@ -180,7 +180,7 @@ namespace EpicGames.Horde.Logs
 			else
 			{
 				// Get all the current property values
-				Dictionary<string, string> propertyValues = new Dictionary<string, string>(StringComparer.OrdinalIgnoreCase);
+				Dictionary<string, string> propertyValues = new(StringComparer.OrdinalIgnoreCase);
 				foreach ((string name, JsonNode? node) in properties)
 				{
 					string value = String.Empty;
@@ -212,7 +212,7 @@ namespace EpicGames.Horde.Logs
 
 						properties!.Remove(propertyName);
 
-						StringBuilder builder = new StringBuilder();
+						StringBuilder builder = new();
 						builder.Append(prefix);
 
 						for (int lineNum = 0; ; lineNum++)
@@ -320,12 +320,12 @@ namespace EpicGames.Horde.Logs
 			}
 
 			_lineWriter.Clear();
-			using (Utf8JsonWriter jsonWriter = new Utf8JsonWriter(_lineWriter))
+			using (Utf8JsonWriter jsonWriter = new(_lineWriter))
 			{
 				obj.WriteTo(jsonWriter);
 			}
 
-			JsonLogEvent nextEvent = new JsonLogEvent(jsonLogEvent.Level, jsonLogEvent.EventId, line, lineCount, _lineWriter.WrittenMemory.ToArray());
+			JsonLogEvent nextEvent = new(jsonLogEvent.Level, jsonLogEvent.EventId, line, lineCount, _lineWriter.WrittenMemory.ToArray());
 			WriteEventInternal(nextEvent);
 		}
 
diff --git a/Engine/Source/Programs/Shared/EpicGames.Horde/Logs/ServerLogger.cs b/Engine/Source/Programs/Shared/EpicGames.Horde/Logs/ServerLogger.cs
--- a/Engine/Source/Programs/Shared/EpicGames.Horde/Logs/ServerLogger.cs
+++ b/Engine/Source/Programs/Shared/EpicGames.Horde/Logs/ServerLogger.cs
@@ -36,7 +36,7 @@ namespace EpicGames.Horde.Logs
 		// Tailing task
 		readonly Task _tailTask;
 		AsyncEvent _tailTaskStop;
-		readonly AsyncEvent _newTailDataEvent = new AsyncEvent();
+		readonly AsyncEvent _newTailDataEvent = new();
 
 		readonly Channel<JsonLogEvent> _dataChannel;
 		Task? _dataWriter;
@@ -146,7 +146,7 @@ namespace EpicGames.Horde.Logs
 			int numWarnings = 0;
 
 			// Buffer for events read in a single iteration
-			ServerLogPacketBuilder writer = new ServerLogPacketBuilder();
+			ServerLogPacketBuilder writer = new();
 			List<RpcCreateLogEventRequest> events = [];
 
 			// Whether we've written the flush command
@@ -335,7 +335,7 @@ namespace EpicGames.Horde.Logs
 		{
 			LogRpc.LogRpcClient logRpc = await _hordeClient.CreateGrpcClientAsync<LogRpc.LogRpcClient>(cancellationToken);
 
-			RpcCreateLogEventsRequest request = new RpcCreateLogEventsRequest();
+			RpcCreateLogEventsRequest request = new();
 			request.Events.AddRange(events);
 
 			await logRpc.CreateLogEventsAsync(request, cancellationToken: cancellationToken);
@@ -363,7 +363,7 @@ namespace EpicGames.Horde.Logs
 		{
 			_internalLogger.LogInformation("Updating log {LogId} to line {LineCount}, target {Locator}", _logId, lineCount, target.GetLocator());
 
-			RpcUpdateLogRequest request = new RpcUpdateLogRequest();
+			RpcUpdateLogRequest request = new();
 			request.LogId = _logId.ToString();
 			request.LineCount = lineCount;
 			request.TargetHash = target.Hash.ToString();
@@ -383,7 +383,7 @@ namespace EpicGames.Horde.Logs
 				using AsyncDuplexStreamingCall<RpcUpdateLogTailRequest, RpcUpdateLogTailResponse> call = clientRef.UpdateLogTail(deadline: deadline, cancellationToken: cancellationToken);
 
 				// Write the request to the server
-				RpcUpdateLogTailRequest request = new RpcUpdateLogTailRequest();
+				RpcUpdateLogTailRequest request = new();
 				request.LogId = _logId.ToString();
 				request.TailNext = tailNext;
 				request.TailData = UnsafeByteOperations.UnsafeWrap(tailData);
diff --git a/Engine/Source/Programs/Shared/EpicGames.Horde/Notifications/NotificationTriggerId.cs b/Engine/Source/Programs/Shared/EpicGames.Horde/Notifications/NotificationTriggerId.cs
--- a/Engine/Source/Programs/Shared/EpicGames.Horde/Notifications/NotificationTriggerId.cs
+++ b/Engine/Source/Programs/Shared/EpicGames.Horde/Notifications/NotificationTriggerId.cs
@@ -16,7 +16,7 @@ namespace EpicGames.Horde.Notifications
 	public readonly record struct NotificationTriggerId(BinaryId Id)
 	{
 		/// <inheritdoc cref="BinaryId.Parse(System.String)"/>
-		public static NotificationTriggerId Parse(string text) => new NotificationTriggerId(BinaryId.Parse(text));
+		public static NotificationTriggerId Parse(string text) => new(BinaryId.Parse(text));
 
 		/// <inheritdoc/>
 		public override readonly string ToString() => Id.ToString();
@@ -28,7 +28,7 @@ namespace EpicGames.Horde.Notifications
 	class NotificationTriggerIdConverter : BinaryIdConverter<NotificationTriggerId>
 	{
 		/// <inheritdoc/>
-		public override NotificationTriggerId FromBinaryId(BinaryId id) => new NotificationTriggerId(id);
+		public override NotificationTriggerId FromBinaryId(BinaryId id) => new(id);
 
 		/// <inheritdoc/>
 		public override BinaryId ToBinaryId(NotificationTriggerId value) => value.Id;
diff --git a/Engine/Source/Programs/Shared/EpicGames.Horde/Projects/HttpProjectCollection.cs b/Engine/Source/Programs/Shared/EpicGames.Horde/Projects/HttpProjectCollection.cs
--- a/Engine/Source/Programs/Shared/EpicGames.Horde/Projects/HttpProjectCollection.cs
+++ b/Engine/Source/Programs/Shared/EpicGames.Horde/Projects/HttpProjectCollection.cs
@@ -32,7 +32,7 @@ namespace EpicGames.Horde.Projects
 		{
 			readonly GetProjectStreamResponse _response;
 
-			public StreamId Id => new StreamId(_response.Id);
+			public StreamId Id => new(_response.Id);
 			public string Name => _response.Name;
 
 			public ProjectStream(GetProjectStreamResponse response)
diff --git a/Engine/Source/Programs/Shared/EpicGames.Horde/Projects/ProjectId.cs b/Engine/Source/Programs/Shared/EpicGames.Horde/Projects/ProjectId.cs
--- a/Engine/Source/Programs/Shared/EpicGames.Horde/Projects/ProjectId.cs
+++ b/Engine/Source/Programs/Shared/EpicGames.Horde/Projects/ProjectId.cs
@@ -37,7 +37,7 @@ namespace EpicGames.Horde.Projects
 	public class ProjectIdConverter : StringIdConverter<ProjectId>
 	{
 		/// <inheritdoc/>
-		public override ProjectId FromStringId(StringId id) => new ProjectId(id);
+		public override ProjectId FromStringId(StringId id) => new(id);
 
 		/// <inheritdoc/>
 		public override StringId ToStringId(ProjectId value) => value.Id;
diff --git a/Engine/Source/Programs/Shared/EpicGames.Horde/Replicators/ReplicatorId.cs b/Engine/Source/Programs/Shared/EpicGames.Horde/Replicators/ReplicatorId.cs
--- a/Engine/Source/Programs/Shared/EpicGames.Horde/Replicators/ReplicatorId.cs
+++ b/Engine/Source/Programs/Shared/EpicGames.Horde/Replicators/ReplicatorId.cs
@@ -126,7 +126,7 @@ namespace EpicGames.Horde.Replicators
 	public class StreamReplicatorIdConverter : StringIdConverter<StreamReplicatorId>
 	{
 		/// <inheritdoc/>
-		public override StreamReplicatorId FromStringId(StringId id) => new StreamReplicatorId(id);
+		public override StreamReplicatorId FromStringId(StringId id) => new(id);
 
 		/// <inheritdoc/>
 		public override StringId ToStringId(StreamReplicatorId value) => value.Id;
diff --git a/Engine/Source/Programs/Shared/EpicGames.Horde/Secrets/SecretId.cs b/Engine/Source/Programs/Shared/EpicGames.Horde/Secrets/SecretId.cs
--- a/Engine/Source/Programs/Shared/EpicGames.Horde/Secrets/SecretId.cs
+++ b/Engine/Source/Programs/Shared/EpicGames.Horde/Secrets/SecretId.cs
@@ -37,7 +37,7 @@ namespace EpicGames.Horde.Secrets
 	class SecretIdConverter : StringIdConverter<SecretId>
 	{
 		/// <inheritdoc/>
-		public override SecretId FromStringId(StringId id) => new SecretId(id);
+		public override SecretId FromStringId(StringId id) => new(id);
 
 		/// <inheritdoc/>
 		public override StringId ToStringId(SecretId value) => value.Id;
diff --git a/Engine/Source/Programs/Shared/EpicGames.Horde/ServiceAccounts/ServiceAccountId.cs b/Engine/Source/Programs/Shared/EpicGames.Horde/ServiceAccounts/ServiceAccountId.cs
--- a/Engine/Source/Programs/Shared/EpicGames.Horde/ServiceAccounts/ServiceAccountId.cs
+++ b/Engine/Source/Programs/Shared/EpicGames.Horde/ServiceAccounts/ServiceAccountId.cs
@@ -16,7 +16,7 @@ namespace EpicGames.Horde.ServiceAccounts
 	public readonly record struct ServiceAccountId(BinaryId Id)
 	{
 		/// <inheritdoc cref="BinaryId.Parse(System.String)"/>
-		public static ServiceAccountId Parse(string text) => new ServiceAccountId(BinaryId.Parse(text));
+		public static ServiceAccountId Parse(string text) => new(BinaryId.Parse(text));
 
 		/// <inheritdoc/>
 		public override string ToString() => Id.ToString();
@@ -28,7 +28,7 @@ namespace EpicGames.Horde.ServiceAccounts
 	class ServiceAccountIdConverter : BinaryIdConverter<ServiceAccountId>
 	{
 		/// <inheritdoc/>
-		public override ServiceAccountId FromBinaryId(BinaryId id) => new ServiceAccountId(id);
+		public override ServiceAccountId FromBinaryId(BinaryId id) => new(id);
 
 		/// <inheritdoc/>
 		public override BinaryId ToBinaryId(ServiceAccountId value) => value.Id;
diff --git a/Engine/Source/Programs/Shared/EpicGames.Horde/Storage/Backends/FileStorageBackend.cs b/Engine/Source/Programs/Shared/EpicGames.Horde/Storage/Backends/FileStorageBackend.cs
--- a/Engine/Source/Programs/Shared/EpicGames.Horde/Storage/Backends/FileStorageBackend.cs
+++ b/Engine/Source/Programs/Shared/EpicGames.Horde/Storage/Backends/FileStorageBackend.cs
@@ -48,7 +48,7 @@ namespace EpicGames.Horde.Storage.Backends
 		/// <summary>
 		/// Gets the path for storing a file on disk
 		/// </summary>
-		static ObjectKey GetBlobFile(BlobLocator locator) => new ObjectKey($"{locator}.blob");
+		static ObjectKey GetBlobFile(BlobLocator locator) => new($"{locator}.blob");
 
 		/// <inheritdoc/>
 		public Task<Stream> OpenBlobAsync(BlobLocator locator, int offset, int? length, CancellationToken cancellationToken)
@@ -155,7 +155,7 @@ namespace EpicGames.Horde.Storage.Backends
 			string[] lines = await FileReference.ReadAllLinesAsync(file, cancellationToken);
 
 			IoHash hash = IoHash.Parse(lines[0].Trim());
-			BlobLocator locator = new BlobLocator(lines[1].Trim());
+			BlobLocator locator = new(lines[1].Trim());
 
 			return new HashedBlobRefValue(hash, locator);
 		}
diff --git a/Engine/Source/Programs/Shared/EpicGames.Horde/Storage/Backends/HttpStorageBackend.cs b/Engine/Source/Programs/Shared/EpicGames.Horde/Storage/Backends/HttpStorageBackend.cs
--- a/Engine/Source/Programs/Shared/EpicGames.Horde/Storage/Backends/HttpStorageBackend.cs
+++ b/Engine/Source/Programs/Shared/EpicGames.Horde/Storage/Backends/HttpStorageBackend.cs
@@ -39,7 +39,7 @@ namespace EpicGames.Horde.Storage.Backends
 		long _numBytes;
 		int _numActive;
 		TimeSpan _sequentialReadTime;
-		readonly Stopwatch _readTime = new Stopwatch();
+		readonly Stopwatch _readTime = new();
 
 		/// <inheritdoc/>
 		public bool SupportsRedirects => _supportsUploadRedirects;
@@ -95,7 +95,7 @@ namespace EpicGames.Horde.Storage.Backends
 
 					HttpResponseMessage response = await retryPolicy.ExecuteAsync(async () =>
 					{
-						using (HttpRequestMessage request = new HttpRequestMessage(HttpMethod.Get, $"{_basePath}/blobs/{locator}"))
+						using (HttpRequestMessage request = new(HttpMethod.Get, $"{_basePath}/blobs/{locator}"))
 						{
 							if (offset != 0 || length != null)
 							{
@@ -157,7 +157,7 @@ namespace EpicGames.Horde.Storage.Backends
 		/// <inheritdoc/>
 		public async Task<BlobLocator> WriteBlobAsync(BlobLocator? locator, Stream stream, IReadOnlyCollection<BlobLocator>? imports, string? prefix = null, CancellationToken cancellationToken = default)
 		{
-			using StreamContent streamContent = new StreamContent(stream);
+			using StreamContent streamContent = new(stream);
 
 			if (_supportsUploadRedirects)
 			{
@@ -291,7 +291,7 @@ namespace EpicGames.Horde.Storage.Backends
 					queryPath += $"&maxResults={maxResults.Value}";
 				}
 
-				using (HttpRequestMessage request = new HttpRequestMessage(HttpMethod.Get, queryPath))
+				using (HttpRequestMessage request = new(HttpMethod.Get, queryPath))
 				{
 					using (HttpResponseMessage response = await httpClient.SendAsync(request, cancellationToken))
 					{
@@ -321,7 +321,7 @@ namespace EpicGames.Horde.Storage.Backends
 		{
 			using (HttpClient httpClient = _createClient())
 			{
-				using (HttpRequestMessage request = new HttpRequestMessage(HttpMethod.Get, $"{_basePath}/refs/{name}"))
+				using (HttpRequestMessage request = new(HttpMethod.Get, $"{_basePath}/refs/{name}"))
 				{
 					if (cacheTime.IsSet())
 					{
diff --git a/Engine/Source/Programs/Shared/EpicGames.Horde/Storage/Backends/JupiterStorageBackend.cs b/Engine/Source/Programs/Shared/EpicGames.Horde/Storage/Backends/JupiterStorageBackend.cs
--- a/Engine/Source/Programs/Shared/EpicGames.Horde/Storage/Backends/JupiterStorageBackend.cs
+++ b/Engine/Source/Programs/Shared/EpicGames.Horde/Storage/Backends/JupiterStorageBackend.cs
@@ -164,7 +164,7 @@ namespace EpicGames.Horde.Storage.Backends
 		// Extract the data field from an encoded cb object
 		static ReadOnlyMemory<byte> UnwrapDataFromCbObject(ReadOnlyMemory<byte> encodedData)
 		{
-			CbObject obj = new CbObject(encodedData);
+			CbObject obj = new(encodedData);
 			return obj[ObjectFieldNames.Data].AsBinary();
 		}
 
@@ -173,7 +173,7 @@ namespace EpicGames.Horde.Storage.Backends
 		{
 			using (HttpClient httpClient = _createClient())
 			{
-				using (HttpRequestMessage request = new HttpRequestMessage(HttpMethod.Get, $"api/v1/blobs/{_namespaceId}/{hash}"))
+				using (HttpRequestMessage request = new(HttpMethod.Get, $"api/v1/blobs/{_namespaceId}/{hash}"))
 				{
 					HttpResponseMessage response = await httpClient.SendAsync(request, cancellationToken);
 					await EnsureSuccessAsync(response, cancellationToken);
@@ -196,7 +196,7 @@ namespace EpicGames.Horde.Storage.Backends
 
 		static byte[] CreateCbObject(byte[] payload, IReadOnlyCollection<BlobLocator> imports)
 		{
-			CbWriter writer = new CbWriter();
+			CbWriter writer = new();
 			writer.BeginObject();
 			writer.WriteBinaryArray(ObjectFieldNames.Data, payload);
 
@@ -260,10 +260,10 @@ namespace EpicGames.Horde.Storage.Backends
 			// Send the message
 			using (HttpClient httpClient = _createClient())
 			{
-				using ByteArrayContent content = new ByteArrayContent(data);
+				using ByteArrayContent content = new(data);
 				content.Headers.ContentType = new MediaTypeHeaderValue(MediaTypeNames.Application.Octet);
 
-				using (HttpRequestMessage request = new HttpRequestMessage(HttpMethod.Post, $"api/v1/blobs/{_namespaceId}"))
+				using (HttpRequestMessage request = new(HttpMethod.Post, $"api/v1/blobs/{_namespaceId}"))
 				{
 					request.Content = content;
 
@@ -324,7 +324,7 @@ namespace EpicGames.Horde.Storage.Backends
 			using (HttpClient httpClient = _createClient())
 			{
 				IoHash refId = IoHash.Compute(name.Text);
-				using (HttpRequestMessage request = new HttpRequestMessage(HttpMethod.Delete, $"api/v1/refs/{_namespaceId}/{_bucketId}/{refId}"))
+				using (HttpRequestMessage request = new(HttpMethod.Delete, $"api/v1/refs/{_namespaceId}/{_bucketId}/{refId}"))
 				{
 					using (HttpResponseMessage response = await httpClient.SendAsync(request, cancellationToken))
 					{
@@ -351,7 +351,7 @@ namespace EpicGames.Horde.Storage.Backends
 			{
 				IoHash refId = IoHash.Compute(name.Text.Span);
 
-				using (HttpRequestMessage request = new HttpRequestMessage(HttpMethod.Get, $"api/v1/refs/{_namespaceId}/{_bucketId}/{refId}"))
+				using (HttpRequestMessage request = new(HttpMethod.Get, $"api/v1/refs/{_namespaceId}/{_bucketId}/{refId}"))
 				{
 					if (cacheTime.IsSet())
 					{
@@ -381,7 +381,7 @@ namespace EpicGames.Horde.Storage.Backends
 
 							byte[] encodedData = await response.Content.ReadAsByteArrayAsync(cancellationToken);
 
-							CbObject objectData = new CbObject(encodedData);
+							CbObject objectData = new(encodedData);
 							CbField target = objectData[RefFieldNames.Target];
 
 							IoHash hash;
@@ -418,9 +418,9 @@ namespace EpicGames.Horde.Storage.Backends
 			using (HttpClient httpClient = _createClient())
 			{
 				IoHash refId = IoHash.Compute(name.Text.Span);
-				using (HttpRequestMessage request = new HttpRequestMessage(HttpMethod.Put, $"api/v1/refs/{_namespaceId}/{_bucketId}/{refId}"))
+				using (HttpRequestMessage request = new(HttpMethod.Put, $"api/v1/refs/{_namespaceId}/{_bucketId}/{refId}"))
 				{
-					CbWriter writer = new CbWriter();
+					CbWriter writer = new();
 					writer.BeginObject();
 
 					if (TryGetBinaryAttachment(value.Locator.Path, out CbBinaryAttachment binaryAttachment))
diff --git a/Engine/Source/Programs/Shared/EpicGames.Horde/Storage/Backends/MemoryStorageBackend.cs b/Engine/Source/Programs/Shared/EpicGames.Horde/Storage/Backends/MemoryStorageBackend.cs
--- a/Engine/Source/Programs/Shared/EpicGames.Horde/Storage/Backends/MemoryStorageBackend.cs
+++ b/Engine/Source/Programs/Shared/EpicGames.Horde/Storage/Backends/MemoryStorageBackend.cs
@@ -17,9 +17,9 @@ namespace EpicGames.Horde.Storage.Backends
 	{
 		record class AliasListNode(BlobLocator Locator, int Rank, ReadOnlyMemory<byte> Data, AliasListNode? Next);
 
-		readonly ConcurrentDictionary<BlobLocator, ReadOnlyMemory<byte>> _blobs = new ConcurrentDictionary<BlobLocator, ReadOnlyMemory<byte>>();
-		readonly ConcurrentDictionary<RefName, HashedBlobRefValue> _refs = new ConcurrentDictionary<RefName, HashedBlobRefValue>();
-		readonly ConcurrentDictionary<string, AliasListNode?> _aliases = new ConcurrentDictionary<string, AliasListNode?>(StringComparer.Ordinal);
+		readonly ConcurrentDictionary<BlobLocator, ReadOnlyMemory<byte>> _blobs = new();
+		readonly ConcurrentDictionary<RefName, HashedBlobRefValue> _refs = new();
+		readonly ConcurrentDictionary<string, AliasListNode?> _aliases = new(StringComparer.Ordinal);
 
 		/// <summary>
 		/// All data stored by the client
diff --git a/Engine/Source/Programs/Shared/EpicGames.Horde/Storage/BatchBlobReader.cs b/Engine/Source/Programs/Shared/EpicGames.Horde/Storage/BatchBlobReader.cs
--- a/Engine/Source/Programs/Shared/EpicGames.Horde/Storage/BatchBlobReader.cs
+++ b/Engine/Source/Programs/Shared/EpicGames.Horde/Storage/BatchBlobReader.cs
@@ -132,7 +132,7 @@ namespace EpicGames.Horde.Storage
 		/// Gets stats for the reader
 		/// </summary>
 		public BatchReaderStats GetStats()
-			=> new BatchReaderStats(_numRequests, _numReads, _numBundles, _numPackets, _bytesRead, _bytesDecoded);
+			=> new(_numRequests, _numReads, _numBundles, _numPackets, _bytesRead, _bytesDecoded);
 
 		/// <summary>
 		/// Adds a batch reader to the given pipeline
@@ -188,7 +188,7 @@ namespace EpicGames.Horde.Storage
 		async Task CreateBundleRequestsAsync(ChannelReader<TRequest> requestReader, CancellationToken cancellationToken)
 		{
 			int queueLength = 0;
-			Queue<BundleHandle> bundleQueue = new Queue<BundleHandle>();
+			Queue<BundleHandle> bundleQueue = new();
 			Dictionary<BundleHandle, List<PacketExport>> bundleHandleToExportBatch = [];
 
 			for (; ; )
@@ -262,7 +262,7 @@ namespace EpicGames.Horde.Storage
 			foreach (IGrouping<int, PacketExport> group in exports.GroupBy(x => x.PacketHandle.PacketOffset))
 			{
 				List<ExportRequest> exportRequests = [.. group.Select(x => new ExportRequest(x.ExportHash, x.ExportHandle, x.OriginalRequest))];
-				PacketRequest packetRequest = new PacketRequest(group.First().PacketHandle, exportRequests);
+				PacketRequest packetRequest = new(group.First().PacketHandle, exportRequests);
 				packetRequests.Add(packetRequest);
 			}
 			packetRequests.SortBy(x => x.MinOffset);
@@ -279,7 +279,7 @@ namespace EpicGames.Horde.Storage
 				int minOffset = packetRequests[minIdx].MinOffset;
 				int maxOffset = packetRequests[maxIdx].MaxOffset;
 
-				BundleRequest request = new BundleRequest(bundleHandle, minOffset, maxOffset, packetRequests.GetRange(minIdx, (maxIdx + 1) - minIdx));
+				BundleRequest request = new(bundleHandle, minOffset, maxOffset, packetRequests.GetRange(minIdx, (maxIdx + 1) - minIdx));
 				await _bundleRequests.Writer.WriteAsync(request, cancellationToken);
 			}
 		}
@@ -312,7 +312,7 @@ namespace EpicGames.Horde.Storage
 					Interlocked.Increment(ref _numReads);
 					Interlocked.Add(ref _bytesRead, data.Memory.Length);
 #pragma warning disable CA2000
-					BundleResponse response = new BundleResponse(data, request.MinOffset, request.MaxOffset, request.Packets);
+					BundleResponse response = new(data, request.MinOffset, request.MaxOffset, request.Packets);
 					await _bundleResponses.Writer.WriteAsync(response, cancellationToken);
 #pragma warning restore CA2000
 				}
diff --git a/Engine/Source/Programs/Shared/EpicGames.Horde/Storage/BatchChunkReader.cs b/Engine/Source/Programs/Shared/EpicGames.Horde/Storage/BatchChunkReader.cs
--- a/Engine/Source/Programs/Shared/EpicGames.Horde/Storage/BatchChunkReader.cs
+++ b/Engine/Source/Programs/Shared/EpicGames.Horde/Storage/BatchChunkReader.cs
@@ -60,7 +60,7 @@ namespace EpicGames.Horde.Storage
 				for (; idx < responses.Count && batch.Count < MaxBatchSize && responses[idx].Request.File == responses[firstIdx].Request.File; idx++)
 				{
 					BlobReadResponse<ChunkReadRequest> response = responses[idx];
-					OutputChunk chunk = new OutputChunk(response.Request.File, response.Request.Offset, response.Data.Data, response.Request.Handle, response.Data);
+					OutputChunk chunk = new(response.Request.File, response.Request.Offset, response.Data.Data, response.Request.Handle, response.Data);
 					batch.Add(chunk);
 				}
 
diff --git a/Engine/Source/Programs/Shared/EpicGames.Horde/Storage/BatchOutputWriter.cs b/Engine/Source/Programs/Shared/EpicGames.Horde/Storage/BatchOutputWriter.cs
--- a/Engine/Source/Programs/Shared/EpicGames.Horde/Storage/BatchOutputWriter.cs
+++ b/Engine/Source/Programs/Shared/EpicGames.Horde/Storage/BatchOutputWriter.cs
@@ -343,7 +343,7 @@ namespace EpicGames.Horde.Storage
 		{
 			try
 			{
-				MemoryMappedView view = new MemoryMappedView(memoryMappedFile, 0, file.FileEntry.Length);
+				MemoryMappedView view = new(memoryMappedFile, 0, file.FileEntry.Length);
 				return view;
 			}
 			catch (Exception ex)
diff --git a/Engine/Source/Programs/Shared/EpicGames.Horde/Storage/BlobLocator.cs b/Engine/Source/Programs/Shared/EpicGames.Horde/Storage/BlobLocator.cs
--- a/Engine/Source/Programs/Shared/EpicGames.Horde/Storage/BlobLocator.cs
+++ b/Engine/Source/Programs/Shared/EpicGames.Horde/Storage/BlobLocator.cs
@@ -191,7 +191,7 @@ namespace EpicGames.Horde.Storage
 	public sealed class BlobLocatorJsonConverter : JsonConverter<BlobLocator>
 	{
 		/// <inheritdoc/>
-		public override BlobLocator Read(ref Utf8JsonReader reader, Type typeToConvert, JsonSerializerOptions options) => new BlobLocator(new Utf8String(reader.GetUtf8String().ToArray()));
+		public override BlobLocator Read(ref Utf8JsonReader reader, Type typeToConvert, JsonSerializerOptions options) => new(new Utf8String(reader.GetUtf8String().ToArray()));
 
 		/// <inheritdoc/>
 		public override void Write(Utf8JsonWriter writer, BlobLocator value, JsonSerializerOptions options) => writer.WriteStringValue(value.ToString());
diff --git a/Engine/Source/Programs/Shared/EpicGames.Horde/Storage/BlobPipeline.cs b/Engine/Source/Programs/Shared/EpicGames.Horde/Storage/BlobPipeline.cs
--- a/Engine/Source/Programs/Shared/EpicGames.Horde/Storage/BlobPipeline.cs
+++ b/Engine/Source/Programs/Shared/EpicGames.Horde/Storage/BlobPipeline.cs
@@ -71,7 +71,7 @@ namespace EpicGames.Horde.Storage
 	public sealed class BlobPipeline<TUserData> : IAsyncDisposable
 	{
 		readonly BlobPipelineOptions _options;
-		readonly CancellationTokenSource _cancellationTokenSource = new CancellationTokenSource();
+		readonly CancellationTokenSource _cancellationTokenSource = new();
 		readonly List<Task> _tasks = [];
 		readonly Channel<BlobRequest<TUserData>> _requestChannel;
 		readonly Channel<BlobRequest<TUserData>[]> _requestBatchChannel;
@@ -197,7 +197,7 @@ namespace EpicGames.Horde.Storage
 		async Task BatchLoopAsync(CancellationToken cancellationToken)
 		{
 			int queueLength = 0;
-			Queue<BundleBatchRequest> bundleBatchQueue = new Queue<BundleBatchRequest>();
+			Queue<BundleBatchRequest> bundleBatchQueue = new();
 			Dictionary<BundleHandle, BundleBatchRequest> bundleHandleToBatch = [];
 
 			for (; ; )
@@ -269,7 +269,7 @@ namespace EpicGames.Horde.Storage
 				bundleHandleToBatch.Add(packetHandle.Bundle, bundleBatchRequest);
 			}
 
-			BundleRequest bundleRequest = new BundleRequest(packetHandle.Bundle, packetHandle.PacketOffset, exportHandle.ExportIdx, request);
+			BundleRequest bundleRequest = new(packetHandle.Bundle, packetHandle.PacketOffset, exportHandle.ExportIdx, request);
 			bundleBatchRequest.Requests.Add(bundleRequest);
 			return true;
 		}
@@ -305,7 +305,7 @@ namespace EpicGames.Horde.Storage
 					{
 						BlobRequest<TUserData> request = requestBatch[idx];
 						BlobData blobData = await requestBatch[idx].Handle.ReadBlobDataAsync(cancellationToken);
-						BlobResponse<TUserData> response = new BlobResponse<TUserData>(blobData, request.UserData);
+						BlobResponse<TUserData> response = new(blobData, request.UserData);
 						responseBatch[idx] = response;
 					}
 #pragma warning disable CA2000
diff --git a/Engine/Source/Programs/Shared/EpicGames.Horde/Storage/BlobSerializer.cs b/Engine/Source/Programs/Shared/EpicGames.Horde/Storage/BlobSerializer.cs
--- a/Engine/Source/Programs/Shared/EpicGames.Horde/Storage/BlobSerializer.cs
+++ b/Engine/Source/Programs/Shared/EpicGames.Horde/Storage/BlobSerializer.cs
@@ -25,7 +25,7 @@ namespace EpicGames.Horde.Storage
 		public static T Deserialize<T>(BlobData blobData, BlobSerializerOptions? options = null)
 		{
 			options ??= BlobSerializerOptions.Default;
-			BlobReader reader = new BlobReader(blobData, options);
+			BlobReader reader = new(blobData, options);
 			return options.GetConverter<T>().Read(reader, options);
 		}
 
@@ -111,8 +111,8 @@ namespace EpicGames.Horde.Storage
 			IEnumerator IEnumerable.GetEnumerator() => _list.GetEnumerator();
 		}
 
-		readonly ConcurrentDictionary<Type, BlobConverter> _cachedConverters = new ConcurrentDictionary<Type, BlobConverter>();
-		static readonly ConcurrentDictionary<Type, BlobConverter> s_cachedDefaultConverters = new ConcurrentDictionary<Type, BlobConverter>();
+		readonly ConcurrentDictionary<Type, BlobConverter> _cachedConverters = new();
+		static readonly ConcurrentDictionary<Type, BlobConverter> s_cachedDefaultConverters = new();
 		static readonly Func<Type, BlobConverter> s_createDefaultConverter = CreateDefaultConverter;
 
 		bool _readOnly;
@@ -146,7 +146,7 @@ namespace EpicGames.Horde.Storage
 
 		static BlobSerializerOptions CreateDefaultOptions()
 		{
-			BlobSerializerOptions options = new BlobSerializerOptions();
+			BlobSerializerOptions options = new();
 			options.MakeReadOnly();
 			return options;
 		}
@@ -198,7 +198,7 @@ namespace EpicGames.Horde.Storage
 		/// <param name="version">The server API version</param>
 		public static BlobSerializerOptions Create(HordeApiVersion version)
 		{
-			BlobSerializerOptions options = new BlobSerializerOptions();
+			BlobSerializerOptions options = new();
 			options.Converters.Add(new InteriorChunkedDataNodeConverter(version));
 			options.Converters.Add(new DirectoryNodeConverter(version));
 			return options;
diff --git a/Engine/Source/Programs/Shared/EpicGames.Horde/Storage/Bundles/BundleCompression.cs b/Engine/Source/Programs/Shared/EpicGames.Horde/Storage/Bundles/BundleCompression.cs
--- a/Engine/Source/Programs/Shared/EpicGames.Horde/Storage/Bundles/BundleCompression.cs
+++ b/Engine/Source/Programs/Shared/EpicGames.Horde/Storage/Bundles/BundleCompression.cs
@@ -61,7 +61,7 @@ namespace EpicGames.Horde.Storage.Bundles
 		/// <returns>The compressed data</returns>
 		public static byte[] Compress(BundleCompressionFormat format, ReadOnlyMemory<byte> input)
 		{
-			ArrayMemoryWriter writer = new ArrayMemoryWriter(0);
+			ArrayMemoryWriter writer = new(0);
 			Compress(format, input, writer);
 			return writer.WrittenMemory.ToArray();
 		}
@@ -94,9 +94,9 @@ namespace EpicGames.Horde.Storage.Bundles
 					}
 				case BundleCompressionFormat.Gzip:
 					{
-						using MemoryStream outputStream = new MemoryStream(input.Length);
+						using MemoryStream outputStream = new(input.Length);
 
-						using (GZipStream gzipStream = new GZipStream(outputStream, CompressionLevel.Fastest, true))
+						using (GZipStream gzipStream = new(outputStream, CompressionLevel.Fastest, true))
 						{
 							gzipStream.Write(input.Span);
 						}
@@ -131,7 +131,7 @@ namespace EpicGames.Horde.Storage.Bundles
 					{
 						int maxSize = ZstdCompressor.GetCompressBound(input.Length);
 
-						using ZstdCompressor compressor = new ZstdCompressor();
+						using ZstdCompressor compressor = new();
 
 						Span<byte> buffer = writer.GetSpan(maxSize);
 						if (!compressor.TryWrap(input.Span, buffer, out int encodedLength))
@@ -171,8 +171,8 @@ namespace EpicGames.Horde.Storage.Bundles
 					}
 				case BundleCompressionFormat.Gzip:
 					{
-						using ReadOnlyMemoryStream inputStream = new ReadOnlyMemoryStream(input);
-						using GZipStream inflatedStream = new GZipStream(inputStream, CompressionMode.Decompress, true);
+						using ReadOnlyMemoryStream inputStream = new(input);
+						using GZipStream inflatedStream = new(inputStream, CompressionMode.Decompress, true);
 
 						int length = 0;
 						for (; ; )
@@ -210,7 +210,7 @@ namespace EpicGames.Horde.Storage.Bundles
 					}
 				case BundleCompressionFormat.Zstd:
 					{
-						using ZstdDecompressor decompressor = new ZstdDecompressor();
+						using ZstdDecompressor decompressor = new();
 						if (!decompressor.TryUnwrap(input.Span, output.Span, out int bytesWritten) || bytesWritten != output.Length)
 						{
 							throw new InvalidOperationException("Unable to decompress data using Zstd");
diff --git a/Engine/Source/Programs/Shared/EpicGames.Horde/Storage/Bundles/BundleStorageNamespace.cs b/Engine/Source/Programs/Shared/EpicGames.Horde/Storage/Bundles/BundleStorageNamespace.cs
--- a/Engine/Source/Programs/Shared/EpicGames.Horde/Storage/Bundles/BundleStorageNamespace.cs
+++ b/Engine/Source/Programs/Shared/EpicGames.Horde/Storage/Bundles/BundleStorageNamespace.cs
@@ -37,7 +37,7 @@ namespace EpicGames.Horde.Storage.Bundles
 		readonly IStorageBackend _backend;
 		readonly BundleCache _cache;
 		readonly BundleOptions _options;
-		readonly PacketReaderStats _packetReaderStats = new PacketReaderStats();
+		readonly PacketReaderStats _packetReaderStats = new();
 		readonly Bundles.V1.BundleReader _bundleReader;
 		readonly ILogger _logger;
 
@@ -86,7 +86,7 @@ namespace EpicGames.Horde.Storage.Bundles
 		/// </summary>
 		public static BundleStorageNamespace CreateInMemory(BundleOptions? options, ILogger logger)
 		{
-			MemoryStorageBackend backend = new MemoryStorageBackend();
+			MemoryStorageBackend backend = new();
 			return new BundleStorageNamespace(backend, BundleCache.None, options, logger);
 		}
 
@@ -101,7 +101,7 @@ namespace EpicGames.Horde.Storage.Bundles
 		/// </summary>
 		public static BundleStorageNamespace CreateFromDirectory(DirectoryReference rootDir, BundleCache cache, MemoryMappedFileCache memoryMappedFileCache, BundleOptions? options, ILogger logger)
 		{
-			FileStorageBackend backend = new FileStorageBackend(new FileObjectStore(rootDir, memoryMappedFileCache), logger);
+			FileStorageBackend backend = new(new FileObjectStore(rootDir, memoryMappedFileCache), logger);
 			return new BundleStorageNamespace(backend, cache, options, logger);
 		}
 
@@ -176,7 +176,7 @@ namespace EpicGames.Horde.Storage.Bundles
 				throw new Exception($"{locator} is not valid");
 			}
 
-			FlushedBundleHandle bundleHandle = new FlushedBundleHandle(this, baseLocator);
+			FlushedBundleHandle bundleHandle = new(this, baseLocator);
 
 			int exportIdx;
 			if (Utf8Parser.TryParse(fragment.Span, out exportIdx, out int numBytesRead) && numBytesRead == fragment.Length)
diff --git a/Engine/Source/Programs/Shared/EpicGames.Horde/Storage/Bundles/V1/BundleHeader.cs b/Engine/Source/Programs/Shared/EpicGames.Horde/Storage/Bundles/V1/BundleHeader.cs
--- a/Engine/Source/Programs/Shared/EpicGames.Horde/Storage/Bundles/V1/BundleHeader.cs
+++ b/Engine/Source/Programs/Shared/EpicGames.Horde/Storage/Bundles/V1/BundleHeader.cs
@@ -124,7 +124,7 @@ namespace EpicGames.Horde.Storage.Bundles.V1
 		public void AppendTo(ReadOnlySequenceBuilder<byte> builder)
 		{
 			// Create a writer for the output data
-			ChunkedArrayMemoryWriter writer = new ChunkedArrayMemoryWriter();
+			ChunkedArrayMemoryWriter writer = new();
 
 			// Create the prelude
 			Span<byte> signatureData = writer.GetSpanAndAdvance(BundleSignature.NumBytes);
@@ -146,7 +146,7 @@ namespace EpicGames.Horde.Storage.Bundles.V1
 			Packets.Write(writer);
 
 			// Fill in the length of the header, and append it to the builder
-			BundleSignature signature = new BundleSignature(BundleVersion.LatestV1, writer.Length);
+			BundleSignature signature = new(BundleVersion.LatestV1, writer.Length);
 			signature.Write(signatureData);
 
 			writer.AppendTo(builder);
@@ -157,7 +157,7 @@ namespace EpicGames.Horde.Storage.Bundles.V1
 		/// </summary>
 		public byte[] ToByteArray()
 		{
-			ReadOnlySequenceBuilder<byte> builder = new ReadOnlySequenceBuilder<byte>();
+			ReadOnlySequenceBuilder<byte> builder = new();
 			AppendTo(builder);
 			return builder.Construct().ToArray();
 		}
@@ -183,7 +183,7 @@ namespace EpicGames.Horde.Storage.Bundles.V1
 		/// <returns>New header object</returns>
 		public static BundleHeader Read(ReadOnlyMemory<byte> memory)
 		{
-			using ReadOnlyMemoryStream stream = new ReadOnlyMemoryStream(memory);
+			using ReadOnlyMemoryStream stream = new(memory);
 			return ReadAsync(stream).Result;
 		}
 #pragma warning restore VSTHRD002
@@ -240,10 +240,10 @@ namespace EpicGames.Horde.Storage.Bundles.V1
 			byte[] exportHeaders = [];
 			byte[] exportReferences = [];
 
-			BundleTypeCollection types = new BundleTypeCollection();
-			BundleImportCollection imports = new BundleImportCollection();
-			BundleExportCollection exports = new BundleExportCollection();
-			BundlePacketCollection packets = new BundlePacketCollection();
+			BundleTypeCollection types = new();
+			BundleImportCollection imports = new();
+			BundleExportCollection exports = new();
+			BundlePacketCollection packets = new();
 
 			byte[] sectionHeader = new byte[SectionHeaderLength];
 			for (int offset = 0; offset < length;)
@@ -291,7 +291,7 @@ namespace EpicGames.Horde.Storage.Bundles.V1
 
 		static BundleHeader ReadLegacy(ReadOnlyMemory<byte> memory)
 		{
-			MemoryReader reader = new MemoryReader(memory);
+			MemoryReader reader = new(memory);
 
 			BundleVersion version = (BundleVersion)reader.ReadUnsignedVarInt();
 			if (version > BundleVersion.RemoveAliases)
@@ -303,7 +303,7 @@ namespace EpicGames.Horde.Storage.Bundles.V1
 
 			// Read the types
 			int numTypes = (int)reader.ReadUnsignedVarInt();
-			List<BlobType> types = new List<BlobType>(numTypes);
+			List<BlobType> types = new(numTypes);
 
 			for (int typeIdx = 0; typeIdx < numTypes; typeIdx++)
 			{
@@ -315,25 +315,25 @@ namespace EpicGames.Horde.Storage.Bundles.V1
 
 			// Read the imports
 			int numImports = (int)reader.ReadUnsignedVarInt();
-			List<BlobLocator> imports = new List<BlobLocator>(numImports);
+			List<BlobLocator> imports = new(numImports);
 			List<BundleExportRef> allExportReferences = [];
 
 			for (int importIdx = 0; importIdx < numImports; importIdx++)
 			{
-				BlobLocator locator = new BlobLocator(reader.ReadUtf8String());
+				BlobLocator locator = new(reader.ReadUtf8String());
 				imports.Add(locator);
 
 				int[] exportIndexes = reader.ReadVariableLengthArray(() => (int)reader.ReadUnsignedVarInt());
 				for (int exportIdx = 0; exportIdx < exportIndexes.Length; exportIdx++)
 				{
-					BundleExportRef exportReference = new BundleExportRef(imports.Count - 1, exportIndexes[exportIdx]);
+					BundleExportRef exportReference = new(imports.Count - 1, exportIndexes[exportIdx]);
 					allExportReferences.Add(exportReference);
 				}
 			}
 
 			// Read the exports
 			int numExports = (int)reader.ReadUnsignedVarInt();
-			List<ExportInfo> exportInfos = new List<ExportInfo>(numExports);
+			List<ExportInfo> exportInfos = new(numExports);
 
 			for (int exportIdx = 0; exportIdx < numExports; exportIdx++)
 			{
@@ -739,7 +739,7 @@ namespace EpicGames.Horde.Storage.Bundles.V1
 		/// <summary>
 		/// Hash of the node data
 		/// </summary>
-		public IoHash Hash => new IoHash(Data.Slice(0, 20).Span);
+		public IoHash Hash => new(Data.Slice(0, 20).Span);
 
 		/// <summary>
 		/// Type id of the node. Can be used to look up the type information from the bundle header.
@@ -764,7 +764,7 @@ namespace EpicGames.Horde.Storage.Bundles.V1
 		/// <summary>
 		/// References to other nodes
 		/// </summary>
-		public BundleExportRefCollection References => new BundleExportRefCollection(Data.Slice(32));
+		public BundleExportRefCollection References => new(Data.Slice(32));
 
 		/// <summary>
 		/// Constructor
@@ -857,7 +857,7 @@ namespace EpicGames.Horde.Storage.Bundles.V1
 			{
 				ReadOnlyMemory<byte> exportData = data.Slice(index * BundleExport.NumHeaderBytes, BundleExport.NumHeaderBytes);
 
-				BundleExportRefCollection exportRefs = new BundleExportRefCollection();
+				BundleExportRefCollection exportRefs = new();
 				if (refs.Length > 0)
 				{
 					int minOffset = BinaryPrimitives.ReadInt32LittleEndian(refs.Span.Slice(index * sizeof(int)));
diff --git a/Engine/Source/Programs/Shared/EpicGames.Horde/Storage/Bundles/V1/BundleReader.cs b/Engine/Source/Programs/Shared/EpicGames.Horde/Storage/Bundles/V1/BundleReader.cs
--- a/Engine/Source/Programs/Shared/EpicGames.Horde/Storage/Bundles/V1/BundleReader.cs
+++ b/Engine/Source/Programs/Shared/EpicGames.Horde/Storage/Bundles/V1/BundleReader.cs
@@ -70,7 +70,7 @@ namespace EpicGames.Horde.Storage.Bundles.V1
 
 			public QueuedPacket AddPacket(int packetIdx)
 			{
-				QueuedPacket packet = new QueuedPacket(packetIdx);
+				QueuedPacket packet = new(packetIdx);
 
 				int insertIdx = QueuedPackets.BinarySearch(packet);
 				if (insertIdx < 0)
@@ -91,7 +91,7 @@ namespace EpicGames.Horde.Storage.Bundles.V1
 		class QueuedPacket : IComparable<QueuedPacket>
 		{
 			public readonly int PacketIdx;
-			public readonly TaskCompletionSource<ReadOnlyMemory<byte>> CompletionSource = new TaskCompletionSource<ReadOnlyMemory<byte>>(TaskCreationOptions.RunContinuationsAsynchronously);
+			public readonly TaskCompletionSource<ReadOnlyMemory<byte>> CompletionSource = new(TaskCreationOptions.RunContinuationsAsynchronously);
 
 			public QueuedPacket(int packetIdx) => PacketIdx = packetIdx;
 			public int CompareTo(QueuedPacket? other) => PacketIdx - other?.PacketIdx ?? 0;
@@ -129,9 +129,9 @@ namespace EpicGames.Horde.Storage.Bundles.V1
 		readonly BundleCache _cache;
 		readonly ILogger _logger;
 
-		readonly object _queueLock = new object();
+		readonly object _queueLock = new();
 		readonly List<QueuedBundle> _queuedBundles = [];
-		readonly Dictionary<string, Task<ReadOnlyMemory<byte>>> _decodeTasks = new Dictionary<string, Task<ReadOnlyMemory<byte>>>(StringComparer.Ordinal);
+		readonly Dictionary<string, Task<ReadOnlyMemory<byte>>> _decodeTasks = new(StringComparer.Ordinal);
 
 		int _numHeaderReads;
 		int _numPacketReads;
@@ -457,7 +457,7 @@ namespace EpicGames.Horde.Storage.Bundles.V1
 					BundleHeader header = await BundleHeader.ReadAsync(signature, stream, cancellationToken);
 
 					// Construct the bundle info
-					BundleInfo bundleInfo = new BundleInfo(bundle.Locator, header, signature.HeaderLength);
+					BundleInfo bundleInfo = new(bundle.Locator, header, signature.HeaderLength);
 
 					// Also add any encoded packets we prefetched
 					int packetOffset = signature.HeaderLength;
@@ -542,7 +542,7 @@ namespace EpicGames.Horde.Storage.Bundles.V1
 
 			Dictionary<BlobLocator, BundleHandle> locatorToBundleHandle = [];
 
-			List<IBlobRef> imports = new List<IBlobRef>(export.References.Count);
+			List<IBlobRef> imports = new(export.References.Count);
 			foreach (BundleExportRef reference in export.References)
 			{
 				BlobLocator importBlob;
diff --git a/Engine/Source/Programs/Shared/EpicGames.Horde/Storage/Bundles/V1/BundleWriter.cs b/Engine/Source/Programs/Shared/EpicGames.Horde/Storage/Bundles/V1/BundleWriter.cs
--- a/Engine/Source/Programs/Shared/EpicGames.Horde/Storage/Bundles/V1/BundleWriter.cs
+++ b/Engine/Source/Programs/Shared/EpicGames.Horde/Storage/Bundles/V1/BundleWriter.cs
@@ -77,7 +77,7 @@ namespace EpicGames.Horde.Storage.Bundles.V1
 		public ValueTask<BlobData> ReadBlobDataAsync(CancellationToken cancellationToken = default) => _reader.ReadNodeDataAsync(BundleLocator, ExportIdx, cancellationToken);
 
 		/// <inheritdoc/>
-		public ValueTask FlushAsync(CancellationToken cancellationToken = default) => new ValueTask();
+		public ValueTask FlushAsync(CancellationToken cancellationToken = default) => new();
 
 		/// <inheritdoc/>
 		public override bool Equals(object? obj) => obj is FlushedNodeHandle other && BundleLocator == other.BundleLocator && ExportIdx == other.ExportIdx;
@@ -214,7 +214,7 @@ namespace EpicGames.Horde.Storage.Bundles.V1
 			readonly int _maxPacketSize;
 			readonly BundleCompressionFormat _compressionFormat;
 
-			readonly object _lockObject = new object();
+			readonly object _lockObject = new();
 
 			// The current packet being built
 			IMemoryOwner<byte>? _currentPacket;
@@ -253,7 +253,7 @@ namespace EpicGames.Horde.Storage.Bundles.V1
 			Task _compressPacketsTask = Task.CompletedTask;
 
 			// Event which is signalled after the bundle is written to storage
-			readonly TaskCompletionSource<bool> _completeEvent = new TaskCompletionSource<bool>(TaskCreationOptions.RunContinuationsAsynchronously);
+			readonly TaskCompletionSource<bool> _completeEvent = new(TaskCreationOptions.RunContinuationsAsynchronously);
 
 			// Task signalled after the write is complete
 			public Task CompleteTask => _completeEvent.Task;
@@ -323,7 +323,7 @@ namespace EpicGames.Horde.Storage.Bundles.V1
 			{
 				IoHash hash = IoHash.Compute(GetBuffer(size, size).Span);
 
-				PendingNode pendingNode = new PendingNode(_treeReader, blobType, hash, _currentPacketIdx, _currentPacketLength, (int)size, refs, aliases, this);
+				PendingNode pendingNode = new(_treeReader, blobType, hash, _currentPacketIdx, _currentPacketLength, (int)size, refs, aliases, this);
 				_currentPacketLength += pendingNode.Length;
 				UncompressedLength += pendingNode.Length;
 
@@ -394,7 +394,7 @@ namespace EpicGames.Horde.Storage.Bundles.V1
 					int encodedLength = BundleData.Compress(_compressionFormat, packetData.Memory.Slice(0, length), _encodedPacketWriter);
 
 					int encodedOffset = (_packets.Count == 0) ? 0 : _packets[^1].EncodedOffset + _packets[^1].EncodedLength;
-					BundlePacket packet = new BundlePacket(_compressionFormat, encodedOffset, encodedLength, length);
+					BundlePacket packet = new(_compressionFormat, encodedOffset, encodedLength, length);
 					_packets.Add(packet);
 					Interlocked.Add(ref _compressedLength, encodedLength);
 
@@ -441,7 +441,7 @@ namespace EpicGames.Horde.Storage.Bundles.V1
 					}
 
 					// Create the output sequence
-					ReadOnlySequenceBuilder<byte> sequence = new ReadOnlySequenceBuilder<byte>();
+					ReadOnlySequenceBuilder<byte> sequence = new();
 					header.AppendTo(sequence);
 
 					foreach (ReadOnlyMemory<byte> packet in packets)
@@ -451,19 +451,19 @@ namespace EpicGames.Horde.Storage.Bundles.V1
 
 					// Write it
 					BlobLocator locator;
-					using (ReadOnlySequenceStream stream = new ReadOnlySequenceStream(sequence.Construct()))
+					using (ReadOnlySequenceStream stream = new(sequence.Construct()))
 					{
 						locator = await store.Backend.WriteBlobAsync(stream, imports, basePath, CancellationToken.None);
 					}
 
 					traceLogger?.LogInformation("Written bundle {BundleId} as {Locator}", BundleId, locator);
 
-					FlushedBundleHandle handle = new FlushedBundleHandle(store, locator);
+					FlushedBundleHandle handle = new(store, locator);
 					for (int idx = 0; idx < _queue.Count; idx++)
 					{
 						PendingNode node = _queue[idx];
 
-						FlushedNodeHandle flushedHandle = new FlushedNodeHandle(_treeReader, locator, handle, idx);
+						FlushedNodeHandle flushedHandle = new(_treeReader, locator, handle, idx);
 						_queue[idx].MarkAsWritten(flushedHandle);
 
 						foreach (AliasInfo alias in node.Aliases)
@@ -511,7 +511,7 @@ namespace EpicGames.Horde.Storage.Bundles.V1
 				}
 
 				// Create the export list
-				List<BundleExport> exports = new List<BundleExport>(_queue.Count);
+				List<BundleExport> exports = new(_queue.Count);
 				foreach (PendingNode nodeInfo in _queue)
 				{
 					int typeIdx = FindOrAddItemIndex(nodeInfo.BlobType, types, typeToIndex);
@@ -538,12 +538,12 @@ namespace EpicGames.Horde.Storage.Bundles.V1
 						exportRefs.Add(exportRef);
 					}
 
-					BundleExport export = new BundleExport(typeIdx, nodeInfo.Packet, nodeInfo.Offset, nodeInfo.Length, exportRefs);
+					BundleExport export = new(typeIdx, nodeInfo.Packet, nodeInfo.Offset, nodeInfo.Length, exportRefs);
 					exports.Add(export);
 				}
 
 				// Get the memory for each packet
-				List<ReadOnlyMemory<byte>> packetData = new List<ReadOnlyMemory<byte>>(_packets.Count);
+				List<ReadOnlyMemory<byte>> packetData = new(_packets.Count);
 				foreach (ReadOnlyMemory<byte> segment in _encodedPacketWriter.AsSequence())
 				{
 					int offset = 0;
@@ -562,7 +562,7 @@ namespace EpicGames.Horde.Storage.Bundles.V1
 				}
 
 				// Create the bundle
-				BundleHeader header = new BundleHeader([.. types], [.. imports], [.. exports], [.. _packets]);
+				BundleHeader header = new([.. types], [.. imports], [.. exports], [.. _packets]);
 				return (header, packetData);
 			}
 
@@ -586,7 +586,7 @@ namespace EpicGames.Horde.Storage.Bundles.V1
 			readonly string? _basePath;
 			readonly long _maxMemoryFootprint;
 			readonly ILogger? _traceLogger;
-			readonly AsyncEvent _completeEvent = new AsyncEvent();
+			readonly AsyncEvent _completeEvent = new();
 			int _refCount;
 			readonly List<Task> _writeTasks = [];
 
@@ -663,7 +663,7 @@ namespace EpicGames.Horde.Storage.Bundles.V1
 			}
 		}
 
-		static readonly BundleOptions s_defaultOptions = new BundleOptions();
+		static readonly BundleOptions s_defaultOptions = new();
 
 		readonly BundleStorageNamespace _store;
 		readonly BundleReader _reader;
diff --git a/Engine/Source/Programs/Shared/EpicGames.Horde/Storage/Bundles/V2/BundleWriter.cs b/Engine/Source/Programs/Shared/EpicGames.Horde/Storage/Bundles/V2/BundleWriter.cs
--- a/Engine/Source/Programs/Shared/EpicGames.Horde/Storage/Bundles/V2/BundleWriter.cs
+++ b/Engine/Source/Programs/Shared/EpicGames.Horde/Storage/Bundles/V2/BundleWriter.cs
@@ -55,7 +55,7 @@ namespace EpicGames.Horde.Storage.Bundles.V2
 		// Bundle that has not yet been written to storage.
 		internal sealed class PendingBundleHandle : BundleHandle, IDisposable
 		{
-			public readonly object LockObject = new object();
+			public readonly object LockObject = new();
 
 			readonly BundleWriter _outer;
 			readonly BundleStorageNamespace _storageNamespace;
@@ -127,7 +127,7 @@ namespace EpicGames.Horde.Storage.Bundles.V2
 				IoHash hash = IoHash.Compute(_packetWriter.GetOutputBuffer(size, size).Span);
 
 				int exportIdx = _packetWriter.CompleteExport(size, type, imports);
-				HashedExportHandle exportHandle = new HashedExportHandle(hash, _packetHandle, exportIdx);
+				HashedExportHandle exportHandle = new(hash, _packetHandle, exportIdx);
 
 				if (aliases.Count > 0)
 				{
@@ -169,7 +169,7 @@ namespace EpicGames.Horde.Storage.Bundles.V2
 					// Point the packet handle to the encoded data
 					lock (LockObject)
 					{
-						FlushedPacketHandle flushedPacketHandle = new FlushedPacketHandle(_storageNamespace, this, packetOffset, packetLength, _bundleCache);
+						FlushedPacketHandle flushedPacketHandle = new(_storageNamespace, this, packetOffset, packetLength, _bundleCache);
 						_packetHandle.CompletePacket(flushedPacketHandle);
 					}
 				}
@@ -211,7 +211,7 @@ namespace EpicGames.Horde.Storage.Bundles.V2
 
 				// Write the bundle data
 				FlushedBundleHandle flushedHandle;
-				using (ReadOnlySequenceStream stream = new ReadOnlySequenceStream(_encodedPacketWriter.AsSequence()))
+				using (ReadOnlySequenceStream stream = new(_encodedPacketWriter.AsSequence()))
 				{
 					await _storageNamespace.Backend.WriteBlobAsync(Locator, stream, _bundleImports, cancellationToken);
 					flushedHandle = new FlushedBundleHandle(_storageNamespace, Locator);
@@ -283,7 +283,7 @@ namespace EpicGames.Horde.Storage.Bundles.V2
 
 		PendingBundleHandle _currentBundle;
 
-		readonly object _lockObject = new object();
+		readonly object _lockObject = new();
 		Task _flushTask = Task.CompletedTask;
 
 		/// <summary>
@@ -318,7 +318,7 @@ namespace EpicGames.Horde.Storage.Bundles.V2
 		}
 
 		BlobLocator GetNextLocator()
-			=> new BlobLocator($"{_locatorPrefix}_{Interlocked.Increment(ref _locatorSuffix)}");
+			=> new($"{_locatorPrefix}_{Interlocked.Increment(ref _locatorSuffix)}");
 
 		/// <inheritdoc/>
 		public override async ValueTask DisposeAsync()
diff --git a/Engine/Source/Programs/Shared/EpicGames.Horde/Storage/Bundles/V2/ExportHandle.cs b/Engine/Source/Programs/Shared/EpicGames.Horde/Storage/Bundles/V2/ExportHandle.cs
--- a/Engine/Source/Programs/Shared/EpicGames.Horde/Storage/Bundles/V2/ExportHandle.cs
+++ b/Engine/Source/Programs/Shared/EpicGames.Horde/Storage/Bundles/V2/ExportHandle.cs
@@ -14,7 +14,7 @@ namespace EpicGames.Horde.Storage.Bundles.V2
 	/// </summary>
 	public class ExportHandle : IBlobRef
 	{
-		static readonly Utf8String s_fragmentPrefix = new Utf8String("exp=");
+		static readonly Utf8String s_fragmentPrefix = new("exp=");
 
 		/// <summary>
 		/// Accessor for the packet that this export is in
@@ -87,7 +87,7 @@ namespace EpicGames.Horde.Storage.Bundles.V2
 		/// </summary>
 		public Utf8String GetIdentifier()
 		{
-			Utf8StringBuilder builder = new Utf8StringBuilder();
+			Utf8StringBuilder builder = new();
 			AppendIdentifier(builder);
 			return builder.ToUtf8String();
 		}
@@ -118,7 +118,7 @@ namespace EpicGames.Horde.Storage.Bundles.V2
 		/// <inheritdoc/>
 		public bool TryGetLocator([NotNullWhen(true)] out BlobLocator locator)
 		{
-			Utf8StringBuilder builder = new Utf8StringBuilder();
+			Utf8StringBuilder builder = new();
 			if (TryAppendLocator(builder))
 			{
 				locator = new BlobLocator(builder.ToUtf8String());
diff --git a/Engine/Source/Programs/Shared/EpicGames.Horde/Storage/Bundles/V2/Packet.cs b/Engine/Source/Programs/Shared/EpicGames.Horde/Storage/Bundles/V2/Packet.cs
--- a/Engine/Source/Programs/Shared/EpicGames.Horde/Storage/Bundles/V2/Packet.cs
+++ b/Engine/Source/Programs/Shared/EpicGames.Horde/Storage/Bundles/V2/Packet.cs
@@ -197,7 +197,7 @@ namespace EpicGames.Horde.Storage.Bundles.V2
 
 			int encodedLength = BundleData.Compress(format, _data, writer);
 
-			BundleSignature signature = new BundleSignature(BundleVersion.LatestV2, BundleSignature.NumBytes + encodedLength + sizeof(int) + 1);
+			BundleSignature signature = new(BundleVersion.LatestV2, BundleSignature.NumBytes + encodedLength + sizeof(int) + 1);
 			signature.Write(signatureSpan);
 		}
 
@@ -226,7 +226,7 @@ namespace EpicGames.Horde.Storage.Bundles.V2
 
 			BundleData.Decompress(format, data.Slice(data.Length - span.Length), memory);
 
-			Packet packet = new Packet(memory);
+			Packet packet = new(memory);
 			return RefCountedHandle.Create(packet, owner);
 		}
 	}
@@ -265,7 +265,7 @@ namespace EpicGames.Horde.Storage.Bundles.V2
 		{
 			ReadOnlySpan<byte> span = data.Span;
 			int baseIdx = (int)VarInt.ReadUnsigned(span, out int bytesRead) - Bias;
-			Utf8String fragment = new Utf8String(data.Slice(bytesRead));
+			Utf8String fragment = new(data.Slice(bytesRead));
 			return new PacketImport(baseIdx, fragment);
 		}
 	}
diff --git a/Engine/Source/Programs/Shared/EpicGames.Horde/Storage/Bundles/V2/PacketHandle.cs b/Engine/Source/Programs/Shared/EpicGames.Horde/Storage/Bundles/V2/PacketHandle.cs
--- a/Engine/Source/Programs/Shared/EpicGames.Horde/Storage/Bundles/V2/PacketHandle.cs
+++ b/Engine/Source/Programs/Shared/EpicGames.Horde/Storage/Bundles/V2/PacketHandle.cs
@@ -62,7 +62,7 @@ namespace EpicGames.Horde.Storage.Bundles.V2
 		readonly int _packetLength;
 		readonly BundleCache _cache;
 
-		static readonly Utf8String s_fragmentPrefix = new Utf8String("pkt=");
+		static readonly Utf8String s_fragmentPrefix = new("pkt=");
 
 		/// <inheritdoc/>
 		public override BundleHandle Bundle => _outer;
@@ -147,7 +147,7 @@ namespace EpicGames.Horde.Storage.Bundles.V2
 		/// <returns></returns>
 		Utf8String GetLocator()
 		{
-			Utf8StringBuilder builder = new Utf8StringBuilder();
+			Utf8StringBuilder builder = new();
 			if (!TryAppendLocator(builder))
 			{
 				throw new NotImplementedException();
@@ -216,13 +216,13 @@ namespace EpicGames.Horde.Storage.Bundles.V2
 
 		async ValueTask<Lifetime<PacketReader>> GetPacketReaderAsync(CancellationToken cancellationToken = default)
 		{
-			PacketReaderCacheKey cacheKey = new PacketReaderCacheKey(_outer, _packetOffset);
+			PacketReaderCacheKey cacheKey = new(_outer, _packetOffset);
 			return await _cache.PacketReaderCache.ScopedGetOrAddAsync(cacheKey, CreatePacketReaderAsync, cancellationToken);
 		}
 
 		internal PacketReader CreatePacketReader(ReadOnlyMemory<byte> encodedData)
 		{
-			PacketReaderCacheKey cacheKey = new PacketReaderCacheKey(_outer, _packetOffset);
+			PacketReaderCacheKey cacheKey = new(_outer, _packetOffset);
 			IRefCountedHandle<Packet> packet = Packet.Decode(encodedData, _cache.Allocator, cacheKey);
 			return new PacketReader(_storageNamespace, _cache, _outer, this, packet.Target, packet);
 		}
@@ -233,7 +233,7 @@ namespace EpicGames.Horde.Storage.Bundles.V2
 			IRefCountedHandle<Packet> packet = Packet.Decode(encodedData.Memory, _cache.Allocator, cacheKey);
 			Interlocked.Add(ref _storageNamespace.PacketReaderStats._numDecodedBytesRead, packet.Target.Length);
 #pragma warning disable CA2000
-			PacketReader reader = new PacketReader(_storageNamespace, _cache, _outer, this, packet.Target, packet);
+			PacketReader reader = new(_storageNamespace, _cache, _outer, this, packet.Target, packet);
 			return new Scoped<PacketReader>(reader);
 #pragma warning restore CA2000
 		}
@@ -268,7 +268,7 @@ namespace EpicGames.Horde.Storage.Bundles.V2
 
 		async ValueTask<Lifetime<IReadOnlyMemoryOwner<byte>>> ReadBundlePageAsync(int pageIdx, CancellationToken cancellationToken)
 		{
-			BundlePageCacheKey bundlePageCacheKey = new BundlePageCacheKey(_outer, pageIdx);
+			BundlePageCacheKey bundlePageCacheKey = new(_outer, pageIdx);
 			return await _cache.BundlePageCache.ScopedGetOrAddAsync(bundlePageCacheKey, ReadBundlePageInternalAsync, cancellationToken);
 		}
 
diff --git a/Engine/Source/Programs/Shared/EpicGames.Horde/Storage/Bundles/V2/PacketWriter.cs b/Engine/Source/Programs/Shared/EpicGames.Horde/Storage/Bundles/V2/PacketWriter.cs
--- a/Engine/Source/Programs/Shared/EpicGames.Horde/Storage/Bundles/V2/PacketWriter.cs
+++ b/Engine/Source/Programs/Shared/EpicGames.Horde/Storage/Bundles/V2/PacketWriter.cs
@@ -69,7 +69,7 @@ namespace EpicGames.Horde.Storage.Bundles.V2
 			_bufferHandle = RefCountedHandle.Create(_allocator.Alloc(1024, AllocationTag));
 			_buffer = _bufferHandle.Target;
 
-			BundleSignature signature = new BundleSignature(BundleVersion.LatestV2, 0);
+			BundleSignature signature = new(BundleVersion.LatestV2, 0);
 			signature.Write(_buffer.Span);
 
 			_length = BundleSignature.NumBytes + (sizeof(int) * 3);
@@ -154,7 +154,7 @@ namespace EpicGames.Horde.Storage.Bundles.V2
 				int offset = _exportOffsets[exportIdx];
 				int length = _exportOffsets[exportIdx + 1] - offset;
 
-				PacketExport export = new PacketExport(_buffer.Slice(offset, length));
+				PacketExport export = new(_buffer.Slice(offset, length));
 				PacketExportHeader header = export.GetHeader();
 
 				BlobType type = _types[header.TypeIdx];
@@ -189,7 +189,7 @@ namespace EpicGames.Horde.Storage.Bundles.V2
 
 		int AddImport(int baseIdx, Utf8String fragment, object handle)
 		{
-			PacketImport import = new PacketImport(baseIdx, fragment);
+			PacketImport import = new(baseIdx, fragment);
 			int idx = _imports.Count;
 			_imports.Add(import);
 			_importMap.Add(handle, idx);
@@ -221,7 +221,7 @@ namespace EpicGames.Horde.Storage.Bundles.V2
 			}
 			else if (handle is PacketHandle packetHandle)
 			{
-				Utf8StringBuilder builder = new Utf8StringBuilder();
+				Utf8StringBuilder builder = new();
 				if (!packetHandle.TryAppendIdentifier(builder))
 				{
 					throw new NotSupportedException("Referenced blob has not been flushed to storage yet");
@@ -268,7 +268,7 @@ namespace EpicGames.Horde.Storage.Bundles.V2
 		/// </summary>
 		public BlobLocator GetImportLocator(int importIdx)
 		{
-			Utf8StringBuilder builder = new Utf8StringBuilder();
+			Utf8StringBuilder builder = new();
 			BuildImportLocator(importIdx, builder);
 			return new BlobLocator(builder.ToUtf8String());
 		}
diff --git a/Engine/Source/Programs/Shared/EpicGames.Horde/Storage/ContentChunker.cs b/Engine/Source/Programs/Shared/EpicGames.Horde/Storage/ContentChunker.cs
--- a/Engine/Source/Programs/Shared/EpicGames.Horde/Storage/ContentChunker.cs
+++ b/Engine/Source/Programs/Shared/EpicGames.Horde/Storage/ContentChunker.cs
@@ -92,7 +92,7 @@ namespace EpicGames.Horde.Storage
 		{
 			using PlatformFileLock? platformFileLock = await PlatformFileLock.CreateAsync(cancellationToken);
 
-			FileStreamOptions options = new FileStreamOptions { Mode = FileMode.Open, Access = FileAccess.Read, Options = FileOptions.Asynchronous };
+			FileStreamOptions options = new() { Mode = FileMode.Open, Access = FileAccess.Read, Options = FileOptions.Asynchronous };
 			using FileStream stream = _fileInfo.Open(options);
 			stream.Seek(offset, SeekOrigin.Begin);
 			await stream.ReadExactlyAsync(memory, cancellationToken);
@@ -188,7 +188,7 @@ namespace EpicGames.Horde.Storage
 		readonly Channel<ChunkerSource> _sourceChannel;
 		readonly Channel<ChunkerOutput> _outputChannel;
 		readonly LeafChunkedDataNodeOptions _options;
-		readonly Queue<Output> _order = new Queue<Output>();
+		readonly Queue<Output> _order = new();
 
 		int _bufferOffset;
 		int _bufferLength;
@@ -224,7 +224,7 @@ namespace EpicGames.Horde.Storage
 			{
 				while (inputReader.TryRead(out ChunkerSource? input))
 				{
-					Output output = new Output(this, input);
+					Output output = new(this, input);
 					if (input.Length > 0)
 					{
 						lock (_order)
@@ -312,14 +312,14 @@ namespace EpicGames.Horde.Storage
 
 		class RingBuffer
 		{
-			readonly object _lockObject = new object();
+			readonly object _lockObject = new();
 
 			int _head;
 			int _size;
 
 			readonly Memory<byte> _data;
 			readonly int _wrapSize;
-			readonly AsyncEvent _freeEvent = new AsyncEvent();
+			readonly AsyncEvent _freeEvent = new();
 
 			public RingBuffer(Memory<byte> data, int wrapSize)
 			{
@@ -584,7 +584,7 @@ namespace EpicGames.Horde.Storage
 
 		async Task<ChunkerOutput> AddAsync(ChunkerSource input, CancellationToken cancellationToken = default)
 		{
-			InputRequest inputRequest = new InputRequest(input, new TaskCompletionSource<Block>(TaskCreationOptions.RunContinuationsAsynchronously));
+			InputRequest inputRequest = new(input, new TaskCompletionSource<Block>(TaskCreationOptions.RunContinuationsAsynchronously));
 			await _inputRequestChannel.Writer.WriteAsync(inputRequest, cancellationToken);
 			return new LatestCdcOutput(this, input.UserData, inputRequest.FirstBlockTcs.Task);
 		}
@@ -622,7 +622,7 @@ namespace EpicGames.Horde.Storage
 							nextBlockTcs = new TaskCompletionSource<Block>(TaskCreationOptions.RunContinuationsAsynchronously);
 						}
 
-						BlockRequest blockRequest = new BlockRequest(bufferOffset, readLength, readTask, blockTcs, nextBlockTcs?.Task);
+						BlockRequest blockRequest = new(bufferOffset, readLength, readTask, blockTcs, nextBlockTcs?.Task);
 						await writer.WriteAsync(blockRequest, cancellationToken);
 
 						blockTcs = nextBlockTcs;
@@ -647,7 +647,7 @@ namespace EpicGames.Horde.Storage
 					List<(int, uint)> boundaries = BuzHash.FindCandidateSplitPoints(data.Span, _options.WindowSize, _options.Threshold);
 
 					// Create the block
-					Block block = new Block(blockRequest.BufferOffset, blockRequest.Length, boundaries, blockRequest.NextBlock);
+					Block block = new(blockRequest.BufferOffset, blockRequest.Length, boundaries, blockRequest.NextBlock);
 					blockRequest.BlockTcs.SetResult(block);
 				}
 			}
diff --git a/Engine/Source/Programs/Shared/EpicGames.Horde/Storage/DedupeBlobWriter.cs b/Engine/Source/Programs/Shared/EpicGames.Horde/Storage/DedupeBlobWriter.cs
--- a/Engine/Source/Programs/Shared/EpicGames.Horde/Storage/DedupeBlobWriter.cs
+++ b/Engine/Source/Programs/Shared/EpicGames.Horde/Storage/DedupeBlobWriter.cs
@@ -20,7 +20,7 @@ namespace EpicGames.Horde.Storage
 		class DedupeCache
 		{
 			readonly int _maxKeys;
-			readonly Queue<BlobKey> _blobKeys = new Queue<BlobKey>();
+			readonly Queue<BlobKey> _blobKeys = new();
 			readonly Dictionary<BlobKey, IHashedBlobRef> _blobKeyToHandle = [];
 
 			public DedupeCache(int maxKeys)
@@ -45,7 +45,7 @@ namespace EpicGames.Horde.Storage
 
 		class WrappedHandle : IHashedBlobRef
 		{
-			public object _lockObject = new object();
+			public object _lockObject = new();
 			public IHashedBlobRef? _inner;
 
 			/// <inheritdoc/>
@@ -136,7 +136,7 @@ namespace EpicGames.Horde.Storage
 		{
 			ReadOnlyMemory<byte> data = _inner.GetOutputBuffer(size, size).Slice(0, size);
 			IoHash hash = IoHash.Compute(data.Span);
-			BlobKey key = new BlobKey(hash, type);
+			BlobKey key = new(hash, type);
 
 			WrappedHandle? wrappedHandle;
 			lock (_cache)
@@ -167,7 +167,7 @@ namespace EpicGames.Horde.Storage
 		/// </summary>
 		public StorageStats GetStats()
 		{
-			StorageStats stats = new StorageStats();
+			StorageStats stats = new();
 			stats.Add("Cache hits", _numHits);
 			stats.Add("Cache hits size", _totalHitsSize);
 			stats.Add("Cache misses", _numMisses);
@@ -185,7 +185,7 @@ namespace EpicGames.Horde.Storage
 		/// <summary>
 		/// Wraps a <see cref="IBlobWriter"/> with a <see cref="DedupeBlobWriter"/>
 		/// </summary>
-		public static DedupeBlobWriter WithDedupe(this IBlobWriter writer, int maxKeys = DedupeBlobWriter.DefaultMaxKeys) => new DedupeBlobWriter(writer, maxKeys);
+		public static DedupeBlobWriter WithDedupe(this IBlobWriter writer, int maxKeys = DedupeBlobWriter.DefaultMaxKeys) => new(writer, maxKeys);
 
 		/// <summary>
 		/// Creates a dedupe writer
diff --git a/Engine/Source/Programs/Shared/EpicGames.Horde/Storage/EncodedBlobData.cs b/Engine/Source/Programs/Shared/EpicGames.Horde/Storage/EncodedBlobData.cs
--- a/Engine/Source/Programs/Shared/EpicGames.Horde/Storage/EncodedBlobData.cs
+++ b/Engine/Source/Programs/Shared/EpicGames.Horde/Storage/EncodedBlobData.cs
@@ -27,7 +27,7 @@ namespace EpicGames.Horde.Storage
 			public readonly JaggedReadOnlyMemoryArray Array;
 
 			public LocatorCollection(JaggedReadOnlyMemoryArray array) => Array = array;
-			public BlobLocator this[int index] => new BlobLocator(new Utf8String(Array[index]));
+			public BlobLocator this[int index] => new(new Utf8String(Array[index]));
 			public int Count => Array.Count;
 			public IEnumerator<BlobLocator> GetEnumerator() => Array.Select(x => new BlobLocator(new Utf8String(x))).GetEnumerator();
 
diff --git a/Engine/Source/Programs/Shared/EpicGames.Horde/Storage/FileMetadataDb.cs b/Engine/Source/Programs/Shared/EpicGames.Horde/Storage/FileMetadataDb.cs
--- a/Engine/Source/Programs/Shared/EpicGames.Horde/Storage/FileMetadataDb.cs
+++ b/Engine/Source/Programs/Shared/EpicGames.Horde/Storage/FileMetadataDb.cs
@@ -37,7 +37,7 @@ namespace EpicGames.Horde.Storage
 			}
 		}
 
-		static readonly IoHashTypeMapper s_hashTypeMapper = new IoHashTypeMapper();
+		static readonly IoHashTypeMapper s_hashTypeMapper = new();
 
 		readonly SqliteConnection _connection;
 
@@ -72,7 +72,7 @@ namespace EpicGames.Horde.Storage
 		{
 			SqlMapper.AddTypeHandler<IoHash>(s_hashTypeMapper);
 
-			SqliteConnection? connection = new SqliteConnection(connectionString);
+			SqliteConnection? connection = new(connectionString);
 			try
 			{
 				await connection.OpenAsync(cancellationToken);
@@ -145,7 +145,7 @@ namespace EpicGames.Horde.Storage
 		/// </summary>
 		public async Task<int> AddFileAsync(int directoryId, string name, DateTime time, long length)
 		{
-			FileRow row = new FileRow(directoryId, name, time, length);
+			FileRow row = new(directoryId, name, time, length);
 			return await _connection.ExecuteScalarAsync<int>("INSERT INTO Files (DirectoryId, Name, Time, Length) VALUES (@DirectoryId, @Name, @Time, @Length); SELECT last_insert_rowid();", row);
 		}
 
@@ -186,7 +186,7 @@ namespace EpicGames.Horde.Storage
 		/// </summary>
 		public async Task<string> GetFilePathAsync(int id)
 		{
-			StringBuilder builder = new StringBuilder();
+			StringBuilder builder = new();
 			await GetFilePathAsync(id, builder);
 			return builder.ToString();
 		}
@@ -218,7 +218,7 @@ namespace EpicGames.Horde.Storage
 		/// </summary>
 		public async Task<int> AddChunkAsync(int fileId, long offset, long length, IoHash hash)
 		{
-			ChunkRow row = new ChunkRow(fileId, offset, length, hash);
+			ChunkRow row = new(fileId, offset, length, hash);
 			return await _connection.ExecuteScalarAsync<int>("INSERT INTO Chunks (FileId, Offset, Length, Hash) VALUES (@FileId, @Offset, @Length, @Hash); SELECT last_insert_rowid();", row);
 		}
 
@@ -279,7 +279,7 @@ namespace EpicGames.Horde.Storage
 		/// </summary>
 		public async Task<int> AddDirectoryAsync(int parentDirectoryId, string name)
 		{
-			DirectoryRow row = new DirectoryRow { ParentDirectoryId = parentDirectoryId, Name = name };
+			DirectoryRow row = new() { ParentDirectoryId = parentDirectoryId, Name = name };
 			return await _connection.ExecuteScalarAsync<int>("INSERT INTO Directories (ParentDirectoryId, Name) VALUES (@ParentDirectoryId, @Name); SELECT last_insert_rowid();", row);
 		}
 
@@ -304,7 +304,7 @@ namespace EpicGames.Horde.Storage
 		/// </summary>
 		public async Task<string> GetDirectoryPathAsync(int id)
 		{
-			StringBuilder builder = new StringBuilder();
+			StringBuilder builder = new();
 			await GetDirectoryPathAsync(id, builder);
 			return builder.ToString();
 		}
diff --git a/Engine/Source/Programs/Shared/EpicGames.Horde/Storage/IBlobWriter.cs b/Engine/Source/Programs/Shared/EpicGames.Horde/Storage/IBlobWriter.cs
--- a/Engine/Source/Programs/Shared/EpicGames.Horde/Storage/IBlobWriter.cs
+++ b/Engine/Source/Programs/Shared/EpicGames.Horde/Storage/IBlobWriter.cs
@@ -307,7 +307,7 @@ namespace EpicGames.Horde.Storage
 				=> default;
 
 			public ValueTask<BlobData> ReadBlobDataAsync(CancellationToken cancellationToken = default)
-				=> new ValueTask<BlobData>(_data);
+				=> new(_data);
 
 			public bool TryGetLocator([NotNullWhen(true)] out BlobLocator locator)
 				=> throw new NotSupportedException();
@@ -355,7 +355,7 @@ namespace EpicGames.Horde.Storage
 		public override ValueTask<IHashedBlobRef> WriteBlobAsync(BlobType type, int size, IReadOnlyList<IBlobRef> imports, IReadOnlyList<AliasInfo> aliases, CancellationToken cancellationToken)
 		{
 			Memory<byte> memory = _memoryWriter.GetMemoryAndAdvance(size);
-			BlobData data = new BlobData(type, memory.ToArray(), [.. imports]);
+			BlobData data = new(type, memory.ToArray(), [.. imports]);
 			return new ValueTask<IHashedBlobRef>(new BlobRef(++_nextIndex, data));
 		}
 
diff --git a/Engine/Source/Programs/Shared/EpicGames.Horde/Storage/IObjectStore.cs b/Engine/Source/Programs/Shared/EpicGames.Horde/Storage/IObjectStore.cs
--- a/Engine/Source/Programs/Shared/EpicGames.Horde/Storage/IObjectStore.cs
+++ b/Engine/Source/Programs/Shared/EpicGames.Horde/Storage/IObjectStore.cs
@@ -161,7 +161,7 @@ namespace EpicGames.Horde.Storage
 		/// <returns>Path to the uploaded object</returns>
 		public static async Task WriteAsync(this IObjectStore store, ObjectKey locator, ReadOnlyMemory<byte> data, CancellationToken cancellationToken = default)
 		{
-			using ReadOnlyMemoryStream stream = new ReadOnlyMemoryStream(data);
+			using ReadOnlyMemoryStream stream = new(data);
 			await store.WriteAsync(locator, stream, cancellationToken);
 		}
 	}
