Shared: prefer implicit object creation when type is apparent

diff --git a/Engine/Source/Programs/Shared/EpicGames.LogMatchers/SystemicEventMatcher.cs b/Engine/Source/Programs/Shared/EpicGames.LogMatchers/SystemicEventMatcher.cs
--- a/Engine/Source/Programs/Shared/EpicGames.LogMatchers/SystemicEventMatcher.cs
+++ b/Engine/Source/Programs/Shared/EpicGames.LogMatchers/SystemicEventMatcher.cs
@@ -8,18 +8,18 @@ namespace EpicGames.LogMatchers
 {
 	class SystemicEventMatcher : ILogEventMatcher
 	{
-		static readonly Regex s_ddc = new Regex(@"^\s*LogDerivedDataCache: Warning:");
+		static readonly Regex s_ddc = new(@"^\s*LogDerivedDataCache: Warning:");
 
-		static readonly Regex s_pdbUtil = new Regex(@"^\s*ERROR: Error: EC_OK");
-		static readonly Regex s_pdbUtilSuffix = new Regex(@"^\s*ERROR:\s*$");
+		static readonly Regex s_pdbUtil = new(@"^\s*ERROR: Error: EC_OK");
+		static readonly Regex s_pdbUtilSuffix = new(@"^\s*ERROR:\s*$");
 
-		static readonly Regex s_roboMerge = new Regex(@"RoboMerge\/gates.*already locked on Commit Server by buildmachine");
+		static readonly Regex s_roboMerge = new(@"RoboMerge\/gates.*already locked on Commit Server by buildmachine");
 
-		static readonly Regex s_hostDown = new Regex(@"ERROR: System\.IO\.IOException: Host is down");
+		static readonly Regex s_hostDown = new(@"ERROR: System\.IO\.IOException: Host is down");
 
-		static readonly Regex s_missingFileList = new Regex(@"^\s*ERROR: Missing local or shared file list.*\.xml");
+		static readonly Regex s_missingFileList = new(@"^\s*ERROR: Missing local or shared file list.*\.xml");
 
-		static readonly Regex s_ubaCacheResponseTimeout = new Regex(@"UbaCache (.*) - Timed out after .* waiting for message response from server");
+		static readonly Regex s_ubaCacheResponseTimeout = new(@"UbaCache (.*) - Timed out after .* waiting for message response from server");
 
 		public LogEventMatch? Match(ILogCursor cursor)
 		{
@@ -29,7 +29,7 @@ namespace EpicGames.LogMatchers
 			}
 			if (cursor.IsMatch(s_pdbUtil))
 			{
-				LogEventBuilder builder = new LogEventBuilder(cursor);
+				LogEventBuilder builder = new(cursor);
 				while (builder.Next.IsMatch(s_pdbUtilSuffix))
 				{
 					builder.MoveNext();
diff --git a/Engine/Source/Programs/Shared/EpicGames.LogMatchers/WorldLeakEventMatcher.cs b/Engine/Source/Programs/Shared/EpicGames.LogMatchers/WorldLeakEventMatcher.cs
--- a/Engine/Source/Programs/Shared/EpicGames.LogMatchers/WorldLeakEventMatcher.cs
+++ b/Engine/Source/Programs/Shared/EpicGames.LogMatchers/WorldLeakEventMatcher.cs
@@ -19,17 +19,17 @@ namespace EpicGames.LogMatchers
 		/// <summary>
 		/// Regex pattern that matches Report Start
 		/// </summary>
-		public static readonly Regex ReportStartPattern = new Regex(@"^====Fatal World Leaks====$", RegexOptions.Multiline | RegexOptions.ExplicitCapture);
+		public static readonly Regex ReportStartPattern = new(@"^====Fatal World Leaks====$", RegexOptions.Multiline | RegexOptions.ExplicitCapture);
 
 		/// <summary>
 		/// Regex pattern that matches Report End
 		/// </summary>
-		public static readonly Regex ReportEndPattern = new Regex(@"^====End Fatal World Leaks====$", RegexOptions.Multiline | RegexOptions.ExplicitCapture);
+		public static readonly Regex ReportEndPattern = new(@"^====End Fatal World Leaks====$", RegexOptions.Multiline | RegexOptions.ExplicitCapture);
 
 		/// <summary>
 		/// Regex pattern that matches callstacks reported during leak reports
 		/// </summary>
-		public static readonly Regex ReportCallstackPattern = new Regex(
+		public static readonly Regex ReportCallstackPattern = new(
 			@"^[\s]*\^\s" +
 			@"(?<Symbol>.+?)\s" +
 			@"\[(?<SourceFile>(\/|\w:).+?)" +
@@ -38,17 +38,17 @@ namespace EpicGames.LogMatchers
 		/// <summary>
 		/// Regex pattern that matches the underlining of the reference preventing the object from being GC'd
 		/// </summary>
-		public static readonly Regex ReportLingeringReferenceUnderlinePattern = new Regex(@"^\s*\^+\s*$", RegexOptions.Multiline | RegexOptions.ExplicitCapture);
+		public static readonly Regex ReportLingeringReferenceUnderlinePattern = new(@"^\s*\^+\s*$", RegexOptions.Multiline | RegexOptions.ExplicitCapture);
 
 		/// <summary>
 		/// Regex pattern that matches leaking reference chain text
 		/// </summary>
-		public static readonly Regex ReportLingeringReference = new Regex(@"^\s*\->\s*(?<ObjectText>.+)$", RegexOptions.Multiline | RegexOptions.ExplicitCapture);
+		public static readonly Regex ReportLingeringReference = new(@"^\s*\->\s*(?<ObjectText>.+)$", RegexOptions.Multiline | RegexOptions.ExplicitCapture);
 
 		/// <summary>
 		/// Regex pattern that matches reference chain text
 		/// </summary>
-		public static readonly Regex ReportBodyPattern = new Regex(@"^\s*(?<ObjectFlags>((\([\w\s]+\)\s))+)?\s*(?<ObjectText>(.+)\s/(.)+/(.)+)$", RegexOptions.Multiline | RegexOptions.ExplicitCapture);
+		public static readonly Regex ReportBodyPattern = new(@"^\s*(?<ObjectFlags>((\([\w\s]+\)\s))+)?\s*(?<ObjectText>(.+)\s/(.)+/(.)+)$", RegexOptions.Multiline | RegexOptions.ExplicitCapture);
 
 		/// <inheritdoc/> 
 		public LogEventMatch? Match(ILogCursor cursor)
@@ -61,7 +61,7 @@ namespace EpicGames.LogMatchers
 				EventId eventId = KnownLogEvents.WorldLeak;
 				LogLevel reportLevel = LogLevel.Error;
 
-				LogEventBuilder builder = new LogEventBuilder(cursor);
+				LogEventBuilder builder = new(cursor);
 
 				builder.MoveNext();
 
diff --git a/Engine/Source/Programs/Shared/EpicGames.LogMatchers/XoreaxEventMatcher.cs b/Engine/Source/Programs/Shared/EpicGames.LogMatchers/XoreaxEventMatcher.cs
--- a/Engine/Source/Programs/Shared/EpicGames.LogMatchers/XoreaxEventMatcher.cs
+++ b/Engine/Source/Programs/Shared/EpicGames.LogMatchers/XoreaxEventMatcher.cs
@@ -13,46 +13,46 @@ namespace EpicGames.LogMatchers
 		public const string PropertyDuration = "duration";
 		public const string PropertyStartTime = "startTime";
 
-		static readonly Regex s_buildService = new Regex(@"\(BuildService.exe\) is not running");
+		static readonly Regex s_buildService = new(@"\(BuildService.exe\) is not running");
 
-		static readonly Regex s_xgConsole = new Regex(@"BUILD FAILED: (.*)xgConsole\.exe(.*)");
+		static readonly Regex s_xgConsole = new(@"BUILD FAILED: (.*)xgConsole\.exe(.*)");
 
-		static readonly Regex s_buildSystemError = new Regex(@"^\s*--------------------Build System Error");
-		static readonly Regex s_buildSystemErrorEnd = new Regex(@"^(\s*)--------------------");
+		static readonly Regex s_buildSystemError = new(@"^\s*--------------------Build System Error");
+		static readonly Regex s_buildSystemErrorEnd = new(@"^(\s*)--------------------");
 
-		static readonly Regex s_buildSystemWarning = new Regex(@"^\s*--------------------Build System Warning[- ]");
-		static readonly Regex s_buildSystemWarningNext = new Regex(@"^(\s*)([^ ].*):");
+		static readonly Regex s_buildSystemWarning = new(@"^\s*--------------------Build System Warning[- ]");
+		static readonly Regex s_buildSystemWarningNext = new(@"^(\s*)([^ ].*):");
 
-		static readonly Regex s_buildSystemWarningSummary = new Regex(@"^\d+ build system warning\(s\):");
-		static readonly Regex s_buildSystemWarningSummaryNext = new Regex(@"^(\s*-\s*)([^ ].*)");
+		static readonly Regex s_buildSystemWarningSummary = new(@"^\d+ build system warning\(s\):");
+		static readonly Regex s_buildSystemWarningSummaryNext = new(@"^(\s*-\s*)([^ ].*)");
 
-		static readonly Regex s_cacheWarning = new Regex(@"^\s*WARNING: \d+ items \([^\)]*\) removed from the cache due to reaching the cache size limit");
-		static readonly Regex s_cacheWarning2 = new Regex(@"^\s*WARNING: Several items removed from the cache due to reaching the cache size limit");
-		static readonly Regex s_cacheWarning3 = new Regex(@"^\s*WARNING: The Build Cache is close to full");
+		static readonly Regex s_cacheWarning = new(@"^\s*WARNING: \d+ items \([^\)]*\) removed from the cache due to reaching the cache size limit");
+		static readonly Regex s_cacheWarning2 = new(@"^\s*WARNING: Several items removed from the cache due to reaching the cache size limit");
+		static readonly Regex s_cacheWarning3 = new(@"^\s*WARNING: The Build Cache is close to full");
 
 		public LogEventMatch? Match(ILogCursor cursor)
 		{
 			if (cursor.IsMatch(s_cacheWarning) || cursor.IsMatch(s_cacheWarning2) || cursor.IsMatch(s_cacheWarning3))
 			{
-				LogEventBuilder builder = new LogEventBuilder(cursor);
+				LogEventBuilder builder = new(cursor);
 				return builder.ToMatch(LogEventPriority.High, LogLevel.Information, KnownLogEvents.Systemic_Xge_CacheLimit);
 			}
 
 			if (cursor.IsMatch(s_buildService))
 			{
-				LogEventBuilder builder = new LogEventBuilder(cursor);
+				LogEventBuilder builder = new(cursor);
 				return builder.ToMatch(LogEventPriority.High, LogLevel.Information, KnownLogEvents.Systemic_Xge_ServiceNotRunning);
 			}
 
 			if (cursor.IsMatch(s_xgConsole))
 			{
-				LogEventBuilder builder = new LogEventBuilder(cursor);
+				LogEventBuilder builder = new(cursor);
 				return builder.ToMatch(LogEventPriority.High, LogLevel.Error, KnownLogEvents.Systemic_Xge_BuildFailed);
 			}
 
 			if (cursor.IsMatch(s_buildSystemError))
 			{
-				LogEventBuilder builder = new LogEventBuilder(cursor);
+				LogEventBuilder builder = new(cursor);
 				builder.MoveNext();
 
 				for (int idx = 0; idx < 100; idx++)
@@ -69,7 +69,7 @@ namespace EpicGames.LogMatchers
 
 			if (cursor.IsMatch(s_buildSystemWarning))
 			{
-				LogEventBuilder builder = new LogEventBuilder(cursor);
+				LogEventBuilder builder = new(cursor);
 				if (builder.Next.TryMatch(s_buildSystemWarningNext, out Match? prefix))
 				{
 					builder.MoveNext();
@@ -82,7 +82,7 @@ namespace EpicGames.LogMatchers
 						eventId = KnownLogEvents.Systemic_Xge_Standalone;
 					}
 
-					Regex prefixPattern = new Regex($"^{Regex.Escape(prefix.Groups[1].Value)}\\s");
+					Regex prefixPattern = new($"^{Regex.Escape(prefix.Groups[1].Value)}\\s");
 					while (builder.Next.IsMatch(prefixPattern))
 					{
 						builder.MoveNext();
@@ -94,7 +94,7 @@ namespace EpicGames.LogMatchers
 
 			if (cursor.IsMatch(s_buildSystemWarningSummary))
 			{
-				LogEventBuilder builder = new LogEventBuilder(cursor);
+				LogEventBuilder builder = new(cursor);
 				if (builder.Next.TryMatch(s_buildSystemWarningSummaryNext, out Match? prefix))
 				{
 					builder.MoveNext();
@@ -107,7 +107,7 @@ namespace EpicGames.LogMatchers
 						eventId = KnownLogEvents.Systemic_Xge_Standalone;
 					}
 
-					Regex prefixPattern = new Regex($"^{Regex.Escape(prefix.Groups[1].Value)}\\s");
+					Regex prefixPattern = new($"^{Regex.Escape(prefix.Groups[1].Value)}\\s");
 					while (builder.Next.IsMatch(prefixPattern))
 					{
 						builder.MoveNext();
diff --git a/Engine/Source/Programs/Shared/EpicGames.LogMatchers/ZenCliEventMatcher.cs b/Engine/Source/Programs/Shared/EpicGames.LogMatchers/ZenCliEventMatcher.cs
--- a/Engine/Source/Programs/Shared/EpicGames.LogMatchers/ZenCliEventMatcher.cs
+++ b/Engine/Source/Programs/Shared/EpicGames.LogMatchers/ZenCliEventMatcher.cs
@@ -11,13 +11,13 @@ namespace EpicGames.LogMatchers
 	/// </summary>
 	class ZenCliEventMatcher : ILogEventMatcher
 	{
-		static readonly Regex s_httpOperationErrorPattern = new Regex(
+		static readonly Regex s_httpOperationErrorPattern = new(
 			@"Error: Operation failed due to a http error: (?<message>.*)\((?<http_error>\d+)\)");
 
-		static readonly Regex s_httpRequestErrorPattern = new Regex(
+		static readonly Regex s_httpRequestErrorPattern = new(
 			@"HttpClient request failed.*?""error_code"": (?<http_error>\d+)");
 
-		static readonly Regex s_httpRequestJsonErrorPattern = new Regex(
+		static readonly Regex s_httpRequestJsonErrorPattern = new(
 			@"HttpClient request failed.*? Status: (?<http_error>\d+), Error: \'");
 
 		public LogEventMatch? Match(ILogCursor input)
@@ -27,7 +27,7 @@ namespace EpicGames.LogMatchers
 			Match? httpMatch;
 			if (input.TryMatch(s_httpOperationErrorPattern, out httpMatch) || input.TryMatch(s_httpRequestErrorPattern, out httpMatch) || input.TryMatch(s_httpRequestJsonErrorPattern, out httpMatch))
 			{
-				LogEventBuilder builder = new LogEventBuilder(input);
+				LogEventBuilder builder = new(input);
 				string httpError = httpMatch.Groups["http_error"].Value;
 
 				// http error codes that we will retry
diff --git a/Engine/Source/Programs/Shared/EpicGames.MongoDB/SingletonDocument.cs b/Engine/Source/Programs/Shared/EpicGames.MongoDB/SingletonDocument.cs
--- a/Engine/Source/Programs/Shared/EpicGames.MongoDB/SingletonDocument.cs
+++ b/Engine/Source/Programs/Shared/EpicGames.MongoDB/SingletonDocument.cs
@@ -155,7 +155,7 @@ namespace EpicGames.MongoDB
 					return @object;
 				}
 
-				T newItem = new T();
+				T newItem = new();
 				newItem.Id = _objectId;
 				await _collection.InsertOneAsync(newItem);
 			}
diff --git a/Engine/Source/Programs/Shared/EpicGames.MsBuild/CsProjBuilder.cs b/Engine/Source/Programs/Shared/EpicGames.MsBuild/CsProjBuilder.cs
--- a/Engine/Source/Programs/Shared/EpicGames.MsBuild/CsProjBuilder.cs
+++ b/Engine/Source/Programs/Shared/EpicGames.MsBuild/CsProjBuilder.cs
@@ -84,7 +84,7 @@ namespace EpicGames.MsBuild
 		{
 			Stopwatch stopwatch = Stopwatch.StartNew();
 			FileReference csharpTargetsFile = FileReference.Combine(hook.EngineDirectory, "Source", "Programs", "Shared", "UnrealEngine.CSharp.targets");
-			Dictionary<string, string> globalProperties = new Dictionary<string, string>
+			Dictionary<string, string> globalProperties = new()
 			{
 				{ "EngineDir", hook.EngineDirectory.FullName },
 				{ "EnginePath", hook.EngineDirectory.FullName },
@@ -107,7 +107,7 @@ namespace EpicGames.MsBuild
 			ConcurrentDictionary<FileReference, CsProjBuildRecordEntry> buildRecords = new(hook.ValidBuildRecords);
 			IEnumerable<FileReference> outOfDateProjects = foundProjects.Except(buildRecords.Keys).Distinct();
 
-			using ProjectCollection projectCollection = new ProjectCollection(globalProperties);
+			using ProjectCollection projectCollection = new(globalProperties);
 			ConcurrentDictionary<FileReference, Project> projects = new();
 			ConcurrentBag<FileReference> platformSpecificProjects = [];
 			logger.LogDebug("Setup {Count} projects time: {TimeSeconds:0.00} s", foundProjects.Count, stopwatch.Elapsed.TotalSeconds);
@@ -196,10 +196,10 @@ namespace EpicGames.MsBuild
 				}
 				string targetPath = Path.GetRelativePath(project.DirectoryPath, project.ExpandString(project.GetPropertyValue("TargetPath")));
 
-				FileReference projectPath = new FileReference(project.FullPath);
+				FileReference projectPath = new(project.FullPath);
 				FileReference buildRecordPath = ConstructBuildRecordPath(hook, projectPath, baseDirectories);
 
-				CsProjBuildRecord buildRecord = new CsProjBuildRecord()
+				CsProjBuildRecord buildRecord = new()
 				{
 					Version = CsProjBuildRecord.CurrentVersion,
 					TargetPath = targetPath,
@@ -258,7 +258,7 @@ namespace EpicGames.MsBuild
 					Remove = [.. glob.Removes.Select(CleanGlobString).Order()]
 				}));
 
-				CsProjBuildRecordEntry entry = new CsProjBuildRecordEntry(projectPath, buildRecordPath, buildRecord);
+				CsProjBuildRecordEntry entry = new(projectPath, buildRecordPath, buildRecord);
 				buildRecords.AddOrUpdate(entry.ProjectFile, entry, (key, oldValue) => entry);
 			});
 
@@ -340,7 +340,7 @@ namespace EpicGames.MsBuild
 			});
 
 			// write all build records
-			JsonSerializerOptions jsonOptions = new JsonSerializerOptions { WriteIndented = true };
+			JsonSerializerOptions jsonOptions = new() { WriteIndented = true };
 			Parallel.ForEach(buildRecords.Values, entry =>
 			{
 				// write all build records
@@ -392,7 +392,7 @@ namespace EpicGames.MsBuild
 			FileReference projectPath = FileReference.Combine(hook.EngineDirectory, "Intermediate", "Build", "EpicGames.ScriptBuild.props");
 
 			XNamespace ns = XNamespace.Get("http://schemas.microsoft.com/developer/msbuild/2003");
-			XDocument document = new XDocument(
+			XDocument document = new(
 				new XElement(ns + "Project",
 					new XAttribute("ToolsVersion", "Current"),
 					new XElement(ns + "ItemGroup",
@@ -407,10 +407,10 @@ namespace EpicGames.MsBuild
 				)
 			);
 
-			StringBuilder output = new StringBuilder();
+			StringBuilder output = new();
 			output.AppendLine("<?xml version=\"1.0\" encoding=\"utf-8\"?>");
 
-			XmlWriterSettings xmlSettings = new XmlWriterSettings();
+			XmlWriterSettings xmlSettings = new();
 			xmlSettings.Async = true;
 			xmlSettings.Encoding = new UTF8Encoding(false);
 			xmlSettings.Indent = true;
@@ -427,7 +427,7 @@ namespace EpicGames.MsBuild
 
 		private static async Task<bool> RunDotnetAsync(IEnumerable<string> arguments, CsProjBuildHook hook, ILogger logger, CancellationToken cancellationToken)
 		{
-			ProcessStartInfo startInfo = new ProcessStartInfo
+			ProcessStartInfo startInfo = new()
 			{
 				FileName = hook.DotnetPath.FullName,
 				Arguments = String.Join(' ', arguments),
@@ -439,11 +439,11 @@ namespace EpicGames.MsBuild
 			};
 			startInfo.EnvironmentVariables["DOTNET_MULTILEVEL_LOOKUP"] = "0"; // use only the bundled dotnet installation - ignore any other/system dotnet install
 
-			using LogEventParser parser = new LogEventParser(logger);
+			using LogEventParser parser = new(logger);
 			parser.AddMatchers(Log.EventParser.EnabledMatchers);
 
 			logger.LogDebug("Running: {Application} {Arguments}", startInfo.FileName, startInfo.Arguments);
-			using Process dotnetProcess = new Process();
+			using Process dotnetProcess = new();
 			dotnetProcess.StartInfo = startInfo;
 			dotnetProcess.OutputDataReceived += (sender, args) =>
 			{
diff --git a/Engine/Source/Programs/Shared/EpicGames.MsBuild/MsBuildSupport.cs b/Engine/Source/Programs/Shared/EpicGames.MsBuild/MsBuildSupport.cs
--- a/Engine/Source/Programs/Shared/EpicGames.MsBuild/MsBuildSupport.cs
+++ b/Engine/Source/Programs/Shared/EpicGames.MsBuild/MsBuildSupport.cs
@@ -14,7 +14,7 @@ namespace EpicGames.MsBuild;
 public static class MsBuildSupport
 {
 	private static bool s_hasRegisteredMsBuildPath = false;
-	private static readonly object s_lock = new object();
+	private static readonly object s_lock = new();
 
 	/// <summary>
 	/// Register our bundled dotnet installation to be used by Microsoft.Build
@@ -37,7 +37,7 @@ public static class MsBuildSupport
 			// Find our bundled dotnet SDK
 			List<string> listOfSdks = [];
 			{
-				ProcessStartInfo startInfo = new ProcessStartInfo
+				ProcessStartInfo startInfo = new()
 				{
 					FileName = hook.DotnetPath.FullName,
 					RedirectStandardOutput = true,
diff --git a/Engine/Source/Programs/Shared/EpicGames.OIDC/MacOSTokenStore.cs b/Engine/Source/Programs/Shared/EpicGames.OIDC/MacOSTokenStore.cs
--- a/Engine/Source/Programs/Shared/EpicGames.OIDC/MacOSTokenStore.cs
+++ b/Engine/Source/Programs/Shared/EpicGames.OIDC/MacOSTokenStore.cs
@@ -49,7 +49,7 @@ namespace EpicGames.OIDC
 
 		private static string? RunKeychainHandler(string verb, List<string> args)
 		{
-			ProcessStartInfo startInfo = new ProcessStartInfo()
+			ProcessStartInfo startInfo = new()
 			{
 				FileName = "security",
 				RedirectStandardOutput = true,
@@ -63,7 +63,7 @@ namespace EpicGames.OIDC
 
 			using Process p = Process.Start(startInfo) ?? throw new Exception("Failed to run \"security\" to interact with keychain");
 
-			StringBuilder sb = new StringBuilder();
+			StringBuilder sb = new();
 			p.OutputDataReceived += (o, e) => { sb.Append(e.Data); };
 			p.BeginOutputReadLine();
 			p.WaitForExit();
diff --git a/Engine/Source/Programs/Shared/EpicGames.OIDC/OidcAuthState.cs b/Engine/Source/Programs/Shared/EpicGames.OIDC/OidcAuthState.cs
--- a/Engine/Source/Programs/Shared/EpicGames.OIDC/OidcAuthState.cs
+++ b/Engine/Source/Programs/Shared/EpicGames.OIDC/OidcAuthState.cs
@@ -181,7 +181,7 @@ public sealed class OidcAuthState(
 		return authState?.TokenInfo?.AccessToken;
 	}
 
-	private readonly object _lockObject = new object();
+	private readonly object _lockObject = new();
 	private readonly IOidcTokenManager _oidcTokenManager = oidcTokenManager;
 	private readonly IStaticAccessTokenProvider _staticAccessTokenProvider = staticAccessTokenProvider;
 	private readonly IClock _clock = clock;
diff --git a/Engine/Source/Programs/Shared/EpicGames.OIDC/OidcTokenManager.cs b/Engine/Source/Programs/Shared/EpicGames.OIDC/OidcTokenManager.cs
--- a/Engine/Source/Programs/Shared/EpicGames.OIDC/OidcTokenManager.cs
+++ b/Engine/Source/Programs/Shared/EpicGames.OIDC/OidcTokenManager.cs
@@ -38,8 +38,8 @@ namespace EpicGames.OIDC
 
 	public class OidcTokenManager : IOidcTokenManager
 	{
-		private readonly object _lockObject = new object();
-		private readonly Dictionary<string, OidcTokenClient> _tokenClients = new Dictionary<string, OidcTokenClient>(StringComparer.OrdinalIgnoreCase);
+		private readonly object _lockObject = new();
+		private readonly Dictionary<string, OidcTokenClient> _tokenClients = new(StringComparer.OrdinalIgnoreCase);
 		private readonly ITokenStore _tokenStore;
 
 		public OidcTokenManager(ITokenStore tokenStore)
@@ -102,7 +102,7 @@ namespace EpicGames.OIDC
 		{
 			Dictionary<string, string> refreshTokens = [];
 
-			OidcTokenOptions options = new OidcTokenOptions();
+			OidcTokenOptions options = new();
 			providerConfiguration.Bind(options);
 
 			List<string> providerKeys = [.. options.Providers.Keys.Where(k => allowedProviders == null || allowedProviders.Contains(k))];
@@ -122,7 +122,7 @@ namespace EpicGames.OIDC
 					continue;
 				}
 
-				OidcTokenClient tokenClient = new OidcTokenClient(key, providerInfo, options.LoginTimeout, _tokenStore);
+				OidcTokenClient tokenClient = new(key, providerInfo, options.LoginTimeout, _tokenStore);
 
 				if (refreshTokens.TryGetValue(key, out string? refreshToken))
 				{
@@ -422,7 +422,7 @@ namespace EpicGames.OIDC
 
 		private async Task<OidcClientOptions> BuildClientOptionsAsync(Uri redirectUri, CancellationToken cancellationToken)
 		{
-			OidcClientOptions options = new OidcClientOptions
+			OidcClientOptions options = new()
 			{
 				Authority = _authorityUri.ToString(),
 				Policy = new Policy { Discovery = new DiscoveryPolicy { Authority = _authorityUri.ToString() } },
@@ -617,7 +617,7 @@ namespace EpicGames.OIDC
 							}
 
 							await using Stream body = request.Body;
-							using StreamReader reader = new StreamReader(body, request.BodyEncoding);
+							using StreamReader reader = new(body, request.BodyEncoding);
 							responseData = await reader.ReadToEndAsync(cancellationToken);
 
 							loginResult = await oidcClient!.ProcessResponseAsync(responseData, loginState, cancellationToken: cancellationToken);
@@ -740,7 +740,7 @@ namespace EpicGames.OIDC
 		{
 			// redirect uri is not used for refresh tokens so we can just pick one of them to configure the client
 			OidcClientOptions options = await BuildClientOptionsAsync(_redirectUris.First(), cancellationToken);
-			OidcClient oidcClient = new OidcClient(options);
+			OidcClient oidcClient = new(options);
 
 			// use the refresh token to acquire a new access token
 			RefreshTokenResult refreshTokenResult = await oidcClient!.RefreshTokenAsync(inRefreshToken, cancellationToken: cancellationToken);
@@ -784,8 +784,8 @@ namespace EpicGames.OIDC
 			string baseUrl = _authorityUri.ToString().TrimEnd('/');
 			string discoUrl = $"{baseUrl}/.well-known/openid-configuration";
 
-			using HttpClient client = new HttpClient();
-			using DiscoveryDocumentRequest doc = new DiscoveryDocumentRequest
+			using HttpClient client = new();
+			using DiscoveryDocumentRequest doc = new()
 			{
 				Address = discoUrl,
 				Policy =
@@ -890,7 +890,7 @@ namespace EpicGames.OIDC
 
 		public static OidcTokenOptions Bind(IConfiguration config)
 		{
-			OidcTokenOptions options = new OidcTokenOptions();
+			OidcTokenOptions options = new();
 			config.GetSection("OidcToken").Bind(options);
 			return options;
 		}
diff --git a/Engine/Source/Programs/Shared/EpicGames.OIDC/ProviderConfigurationFactory.cs b/Engine/Source/Programs/Shared/EpicGames.OIDC/ProviderConfigurationFactory.cs
--- a/Engine/Source/Programs/Shared/EpicGames.OIDC/ProviderConfigurationFactory.cs
+++ b/Engine/Source/Programs/Shared/EpicGames.OIDC/ProviderConfigurationFactory.cs
@@ -34,7 +34,7 @@ namespace EpicGames.OIDC
 				throw new Exception($"Failed to locate engine dir at {engineDir}");
 			}
 
-			ConfigurationBuilder configBuilder = new ConfigurationBuilder();
+			ConfigurationBuilder configBuilder = new();
 			configBuilder
 				.AddJsonFile($"{engineDir}/Programs/OidcToken/{ConfigFileName}", true, false)
 				.AddJsonFile($"{engineDir}/Restricted/NoRedist/Programs/OidcToken/{ConfigFileName}", true, false)
@@ -53,7 +53,7 @@ namespace EpicGames.OIDC
 
 		public static IConfiguration MergeConfiguration(IEnumerable<(DirectoryInfo, DirectoryInfo?)> configurationPaths)
 		{
-			ConfigurationBuilder builder = new ConfigurationBuilder();
+			ConfigurationBuilder builder = new();
 			foreach ((DirectoryInfo engineDir, DirectoryInfo? gameDir) in configurationPaths)
 			{
 				IConfiguration newConfiguration = ReadConfiguration(engineDir, gameDir);
@@ -65,12 +65,12 @@ namespace EpicGames.OIDC
 
 		public static async Task<ClientAuthConfigurationV1?> ReadRemoteAuthConfigurationAsync(Uri remoteUrl, string encryptionKey)
 		{
-			using HttpClient httpClient = new HttpClient();
+			using HttpClient httpClient = new();
 			const int MaxAttempts = 3;
 			byte[]? b = null;
 			for (int i = 0; i < MaxAttempts; i++)
 			{
-				using HttpRequestMessage request = new HttpRequestMessage(HttpMethod.Get, new Uri(remoteUrl, "api/v1/auth/oidc-configuration"));
+				using HttpRequestMessage request = new(HttpMethod.Get, new Uri(remoteUrl, "api/v1/auth/oidc-configuration"));
 				using HttpResponseMessage response = await httpClient.SendAsync(request);
 				if (response.StatusCode == HttpStatusCode.NotFound)
 				{
@@ -95,7 +95,7 @@ namespace EpicGames.OIDC
 			Array.Copy(b, iv, ivLength);
 
 			byte[] key = Convert.FromHexString(encryptionKey);
-			MemoryStream ms = new MemoryStream(b, ivLength, b.Length - ivLength);
+			MemoryStream ms = new(b, ivLength, b.Length - ivLength);
 			await using CryptoStream cryptoStream = new(ms, aes.CreateDecryptor(key, iv), CryptoStreamMode.Read);
 			using StreamReader decryptReader = new(cryptoStream);
 
@@ -106,9 +106,9 @@ namespace EpicGames.OIDC
 
 		public static IConfiguration BindOptions(ClientAuthConfigurationV1 remoteAuthConfig)
 		{
-			using MemoryStream ms = new MemoryStream();
+			using MemoryStream ms = new();
 			{
-				using Utf8JsonWriter jsonWriter = new Utf8JsonWriter(ms);
+				using Utf8JsonWriter jsonWriter = new(ms);
 				JsonSerializer.Serialize(jsonWriter, remoteAuthConfig, ProviderConfigurationStateContext.Default.ClientAuthConfigurationV1);
 				ms.Position = 0;
 			}
diff --git a/Engine/Source/Programs/Shared/EpicGames.OIDC/WindowsTokenStore.cs b/Engine/Source/Programs/Shared/EpicGames.OIDC/WindowsTokenStore.cs
--- a/Engine/Source/Programs/Shared/EpicGames.OIDC/WindowsTokenStore.cs
+++ b/Engine/Source/Programs/Shared/EpicGames.OIDC/WindowsTokenStore.cs
@@ -122,7 +122,7 @@ namespace EpicGames.OIDC
 					return;
 				}
 
-				using Mutex mutex = new Mutex(false, "oidcTokenStoreDat");
+				using Mutex mutex = new(false, "oidcTokenStoreDat");
 
 				try
 				{
@@ -142,8 +142,8 @@ namespace EpicGames.OIDC
 				}
 				string tempFile = Path.GetTempFileName();
 				{
-					using FileStream fs = new FileStream(tempFile, FileMode.Create, FileAccess.Write);
-					using Utf8JsonWriter writer = new Utf8JsonWriter(fs);
+					using FileStream fs = new(tempFile, FileMode.Create, FileAccess.Write);
+					using Utf8JsonWriter writer = new(fs);
 					JsonSerializer.Serialize(writer, new WindowsTokenStoreState(providers), WindowsTokenStoreStateContext.Default.WindowsTokenStoreState);
 				}
 
@@ -320,7 +320,7 @@ namespace EpicGames.OIDC
 
 		public static byte[] DoCryptProtectData(byte[] dataToProtect, string description, byte[] entropy)
 		{
-			DataBlob dataOut = new DataBlob();
+			DataBlob dataOut = new();
 
 			GCHandle dataHandle = GCHandle.Alloc(dataToProtect, GCHandleType.Pinned);
 			GCHandle entropyHandle = GCHandle.Alloc(entropy, GCHandleType.Pinned);
@@ -329,13 +329,13 @@ namespace EpicGames.OIDC
 				Marshal.Copy(dataToProtect, 0, dataHandle.AddrOfPinnedObject(), dataToProtect.Length);
 				Marshal.Copy(entropy, 0, entropyHandle.AddrOfPinnedObject(), entropy.Length);
 
-				DataBlob data = new DataBlob()
+				DataBlob data = new()
 				{
 					cbData = dataToProtect.Length,
 					pbData = dataHandle.AddrOfPinnedObject()
 				};
 
-				DataBlob entropyBlob = new DataBlob()
+				DataBlob entropyBlob = new()
 				{
 					cbData = entropy.Length,
 					pbData = entropyHandle.AddrOfPinnedObject()
@@ -361,7 +361,7 @@ namespace EpicGames.OIDC
 
 		public static byte[] DoCryptUnprotectData(byte[] dataToDecrypt, string description, byte[] entropy)
 		{
-			DataBlob dataOut = new DataBlob();
+			DataBlob dataOut = new();
 
 			GCHandle dataHandle = GCHandle.Alloc(dataToDecrypt, GCHandleType.Pinned);
 			GCHandle entropyHandle = GCHandle.Alloc(entropy, GCHandleType.Pinned);
@@ -370,13 +370,13 @@ namespace EpicGames.OIDC
 				Marshal.Copy(dataToDecrypt, 0, dataHandle.AddrOfPinnedObject(), dataToDecrypt.Length);
 				Marshal.Copy(entropy, 0, entropyHandle.AddrOfPinnedObject(), entropy.Length);
 
-				DataBlob data = new DataBlob()
+				DataBlob data = new()
 				{
 					cbData = dataToDecrypt.Length,
 					pbData = dataHandle.AddrOfPinnedObject()
 				};
 
-				DataBlob entropyBlob = new DataBlob()
+				DataBlob entropyBlob = new()
 				{
 					cbData = entropy.Length,
 					pbData = entropyHandle.AddrOfPinnedObject()
diff --git a/Engine/Source/Programs/Shared/EpicGames.Oodle/Oodle.cs b/Engine/Source/Programs/Shared/EpicGames.Oodle/Oodle.cs
--- a/Engine/Source/Programs/Shared/EpicGames.Oodle/Oodle.cs
+++ b/Engine/Source/Programs/Shared/EpicGames.Oodle/Oodle.cs
@@ -225,7 +225,7 @@ namespace EpicGames.Compression
 		[DllImport("oo2core")]
 		private static extern void Oodle_SetConfigValues(OodleConfigValues configValues);
 
-		private static readonly OodleConfigValues s_configValues = new OodleConfigValues();
+		private static readonly OodleConfigValues s_configValues = new();
 
 		/// <summary>
 		/// Static constructor. Sets the backwards compatibility version for Oodle.
diff --git a/Engine/Source/Programs/Shared/EpicGames.Perforce.Fixture/BasePerforceFixtureTest.cs b/Engine/Source/Programs/Shared/EpicGames.Perforce.Fixture/BasePerforceFixtureTest.cs
--- a/Engine/Source/Programs/Shared/EpicGames.Perforce.Fixture/BasePerforceFixtureTest.cs
+++ b/Engine/Source/Programs/Shared/EpicGames.Perforce.Fixture/BasePerforceFixtureTest.cs
@@ -96,7 +96,7 @@ public abstract class BasePerforceFixtureTest : IDisposable
 			// Remove the read-only flags set by the P4 client
 			foreach (string filePath in Directory.EnumerateFiles(dirPath, "*", SearchOption.AllDirectories))
 			{
-				FileInfo fileInfo = new FileInfo(filePath);
+				FileInfo fileInfo = new(filePath);
 				fileInfo.IsReadOnly = false;
 			}
 
diff --git a/Engine/Source/Programs/Shared/EpicGames.Perforce.Managed/CachedFileInfo.cs b/Engine/Source/Programs/Shared/EpicGames.Perforce.Managed/CachedFileInfo.cs
--- a/Engine/Source/Programs/Shared/EpicGames.Perforce.Managed/CachedFileInfo.cs
+++ b/Engine/Source/Programs/Shared/EpicGames.Perforce.Managed/CachedFileInfo.cs
@@ -38,7 +38,7 @@ namespace EpicGames.Perforce.Managed
 		{
 			FileReference location = GetLocation();
 
-			FileInfo info = new FileInfo(location.FullName);
+			FileInfo info = new(location.FullName);
 			if (!info.Exists)
 			{
 				logger.LogWarning("warning: {File} was missing from cache.", location);
@@ -65,7 +65,7 @@ namespace EpicGames.Perforce.Managed
 
 		public FileReference GetLocation()
 		{
-			StringBuilder fullName = new StringBuilder(CacheDir.FullName);
+			StringBuilder fullName = new(CacheDir.FullName);
 			fullName.Append(Path.DirectorySeparatorChar);
 			fullName.AppendFormat("{0:X}", (CacheId >> 60) & 0xf);
 			fullName.Append(Path.DirectorySeparatorChar);
diff --git a/Engine/Source/Programs/Shared/EpicGames.Perforce.Managed/ManagedWorkspace.cs b/Engine/Source/Programs/Shared/EpicGames.Perforce.Managed/ManagedWorkspace.cs
--- a/Engine/Source/Programs/Shared/EpicGames.Perforce.Managed/ManagedWorkspace.cs
+++ b/Engine/Source/Programs/Shared/EpicGames.Perforce.Managed/ManagedWorkspace.cs
@@ -381,12 +381,12 @@ namespace EpicGames.Perforce.Managed
 			DirectoryReference.CreateDirectory(baseDir);
 			FileUtils.ForceDeleteDirectoryContents(baseDir);
 
-			ManagedWorkspace repo = new ManagedWorkspace(hostName, 1, baseDir, options, tracer, logger);
+			ManagedWorkspace repo = new(hostName, 1, baseDir, options, tracer, logger);
 			await repo.SaveAsync(TransactionState.Clean, cancellationToken);
 			repo.CreateCacheHierarchy();
 
 			FileReference signatureFile = FileReference.Combine(baseDir, SignatureFileName);
-			using (BinaryWriter writer = new BinaryWriter(File.Open(signatureFile.FullName, FileMode.Create, FileAccess.Write, FileShare.Read)))
+			using (BinaryWriter writer = new(File.Open(signatureFile.FullName, FileMode.Create, FileAccess.Write, FileShare.Read)))
 			{
 				writer.Write(CurrentSignature);
 			}
@@ -404,7 +404,7 @@ namespace EpicGames.Perforce.Managed
 			FileReference signatureFile = FileReference.Combine(baseDir, SignatureFileName);
 			if (FileReference.Exists(signatureFile))
 			{
-				using (BinaryReader reader = new BinaryReader(File.Open(signatureFile.FullName, FileMode.Open, FileAccess.Read, FileShare.Read)))
+				using (BinaryReader reader = new(File.Open(signatureFile.FullName, FileMode.Open, FileAccess.Read, FileShare.Read)))
 				{
 					int signature = reader.ReadInt32();
 					if (signature == CurrentSignature)
@@ -438,7 +438,7 @@ namespace EpicGames.Perforce.Managed
 			RestoreBackup(dataFile);
 
 			byte[] data = await FileReference.ReadAllBytesAsync(dataFile, cancellationToken);
-			MemoryReader reader = new MemoryReader(data.AsMemory());
+			MemoryReader reader = new(data.AsMemory());
 
 			int version = reader.ReadInt32();
 			if (version > CurrentVersion)
@@ -478,7 +478,7 @@ namespace EpicGames.Perforce.Managed
 			byte[] buffer = new byte[serializedSize];
 
 			// Write the data to memory
-			MemoryWriter writer = new MemoryWriter(buffer.AsMemory());
+			MemoryWriter writer = new(buffer.AsMemory());
 			writer.WriteInt32(CurrentVersion);
 			writer.WriteBoolean(_requiresRepair || (state != TransactionState.Clean));
 			writer.WriteUInt32(_nextSequenceNumber);
@@ -908,7 +908,7 @@ namespace EpicGames.Perforce.Managed
 					_logger.LogInformation("");
 					for (int rowIdx = 0; rowIdx < streamNames.Count + 1; rowIdx++)
 					{
-						StringBuilder row = new StringBuilder();
+						StringBuilder row = new();
 						for (int colIdx = 0; colIdx < streamNames.Count + 1; colIdx++)
 						{
 							string cell = cells[rowIdx][colIdx];
@@ -1402,7 +1402,7 @@ namespace EpicGames.Perforce.Managed
 
 			// Update the config file with the name of the client
 			FileReference configFile = FileReference.Combine(_baseDir, "p4.ini");
-			using (StreamWriter writer = new StreamWriter(configFile.FullName))
+			using (StreamWriter writer = new(configFile.FullName))
 			{
 				await writer.WriteLineAsync($"P4PORT={perforceClient.Settings.ServerAndPort}");
 				await writer.WriteLineAsync($"P4CLIENT={perforceClient.Settings.ClientName}");
@@ -1628,7 +1628,7 @@ namespace EpicGames.Perforce.Managed
 		/// <param name="cancellationToken">Cancellation token</param>
 		private async Task<StreamSnapshotFromMemory> FindClientContentsAsync(IPerforceConnection perforceClient, int changeNumber, CancellationToken cancellationToken)
 		{
-			StreamTreeBuilder builder = new StreamTreeBuilder();
+			StreamTreeBuilder builder = new();
 
 			using TelemetrySpan span = _tracer.StartActiveSpan($"{nameof(ManagedWorkspace)}.{nameof(FindClientContentsAsync)}");
 			using (Trace("FetchMetadata"))
@@ -1637,7 +1637,7 @@ namespace EpicGames.Perforce.Managed
 				Stopwatch timer = Stopwatch.StartNew();
 
 				// Get the expected prefix for any paths in client syntax
-				Utf8String clientPrefix = new Utf8String($"//{perforceClient.Settings.ClientName}/");
+				Utf8String clientPrefix = new($"//{perforceClient.Settings.ClientName}/");
 
 				// List of the last path fragments. Since file records that are returned are typically sorted by their position in the tree, we can save quite a lot of processing by
 				// reusing as many fragemnts as possible.
@@ -1646,7 +1646,7 @@ namespace EpicGames.Perforce.Managed
 				int numFiles = 0;
 
 				// Handler for each returned record
-				FStatIndexedRecord record = new FStatIndexedRecord();
+				FStatIndexedRecord record = new();
 				void HandleRecord(PerforceRecord rawRecord)
 				{
 					// Copy into the values array
@@ -1694,7 +1694,7 @@ namespace EpicGames.Perforce.Managed
 						}
 
 						// Get the fragment text
-						Utf8String fragment = new Utf8String(clientFile.Memory.Slice(fragmentMinIdx, fragmentMaxIdx - fragmentMinIdx));
+						Utf8String fragment = new(clientFile.Memory.Slice(fragmentMinIdx, fragmentMaxIdx - fragmentMinIdx));
 
 						// If this fragment matches the same fragment from the previous iteration, take the last stream directory straight away
 						if (fragmentIdx < fragments.Count)
@@ -1730,7 +1730,7 @@ namespace EpicGames.Perforce.Managed
 					}
 
 					Md5Hash digest = Md5Hash.Parse(record.Digest);
-					FileContentId contentId = new FileContentId(digest, record.HeadType.Clone());
+					FileContentId contentId = new(digest, record.HeadType.Clone());
 					int revision = (int)Utf8String.ParseUnsignedInt(record.HaveRev);
 
 					// Add a new StreamFileInfo to the last directory object
@@ -1910,7 +1910,7 @@ namespace EpicGames.Perforce.Managed
 				Stopwatch timer = Stopwatch.StartNew();
 
 				// Handle the case where two machines may try to write to the cache file at once by writing to a temporary file
-				FileReference tempCacheFile = new FileReference(String.Format("{0}.{1}", cacheFile, Guid.NewGuid()));
+				FileReference tempCacheFile = new(String.Format("{0}.{1}", cacheFile, Guid.NewGuid()));
 				await contents.SaveAsync(tempCacheFile, basePath);
 
 				// Try to move it into place
@@ -1969,7 +1969,7 @@ namespace EpicGames.Perforce.Managed
 					foreach (KeyValuePair<FileContentId, WorkspaceFileInfo> fileToMove in filesToMove)
 					{
 						ulong cacheId = GetUniqueCacheId(fileToMove.Key);
-						CachedFileInfo newTrackingInfo = new CachedFileInfo(_cacheDir, fileToMove.Key, cacheId, fileToMove.Value._length, fileToMove.Value._lastModifiedTicks, fileToMove.Value._readOnly, _nextSequenceNumber);
+						CachedFileInfo newTrackingInfo = new(_cacheDir, fileToMove.Key, cacheId, fileToMove.Value._length, fileToMove.Value._lastModifiedTicks, fileToMove.Value._readOnly, _nextSequenceNumber);
 						_contentIdToTrackedFile.Add(fileToMove.Key, newTrackingInfo);
 						files.Add((fileToMove.Key, fileToMove.Value.GetLocation(), newTrackingInfo.GetLocation()));
 					}
@@ -2216,7 +2216,7 @@ namespace EpicGames.Perforce.Managed
 
 					// Spawn some background threads to sync them
 					Dictionary<Task, int> tasks = [];
-					Stack<IPerforceConnection> connectionPool = new Stack<IPerforceConnection>();
+					Stack<IPerforceConnection> connectionPool = new();
 					try
 					{
 						while (tasks.Count > 0 || nextBatchIdx < batches.Count)
@@ -2352,7 +2352,7 @@ namespace EpicGames.Perforce.Managed
 			else
 			{
 				FileReference syncFileName = FileReference.Combine(_baseDir, $"SyncList-{beginIdx}.txt");
-				using (StreamWriter writer = new StreamWriter(syncFileName.FullName))
+				using (StreamWriter writer = new(syncFileName.FullName))
 				{
 					for (int idx = beginIdx; idx < endIdx; idx++)
 					{
diff --git a/Engine/Source/Programs/Shared/EpicGames.Perforce.Managed/StreamSnapshotFromMemory.cs b/Engine/Source/Programs/Shared/EpicGames.Perforce.Managed/StreamSnapshotFromMemory.cs
--- a/Engine/Source/Programs/Shared/EpicGames.Perforce.Managed/StreamSnapshotFromMemory.cs
+++ b/Engine/Source/Programs/Shared/EpicGames.Perforce.Managed/StreamSnapshotFromMemory.cs
@@ -89,15 +89,15 @@ namespace EpicGames.Perforce.Managed
 				return null;
 			}
 
-			CbObject rootObj = new CbObject(data.AsMemory(s_currentSignature.Length));
+			CbObject rootObj = new(data.AsMemory(s_currentSignature.Length));
 
 			CbObject rootObj2 = rootObj["root"].AsObject();
 			Utf8String rootPath = rootObj2["path"].AsUtf8String(defaultBasePath);
-			StreamTreeRef root = new StreamTreeRef(rootPath, rootObj2);
+			StreamTreeRef root = new(rootPath, rootObj2);
 
 			CbArray array = rootObj["items"].AsArray();
 
-			Dictionary<IoHash, CbObject> hashToTree = new Dictionary<IoHash, CbObject>(array.Count);
+			Dictionary<IoHash, CbObject> hashToTree = new(array.Count);
 			foreach (CbField element in array)
 			{
 				CbObject objectElement = element.AsObject();
@@ -116,7 +116,7 @@ namespace EpicGames.Perforce.Managed
 		/// <param name="basePath"></param>
 		public async Task SaveAsync(FileReference outputFile, Utf8String basePath)
 		{
-			CbWriter writer = new CbWriter();
+			CbWriter writer = new();
 			writer.BeginObject();
 
 			writer.BeginObject("root");
diff --git a/Engine/Source/Programs/Shared/EpicGames.Perforce.Managed/StreamTree.cs b/Engine/Source/Programs/Shared/EpicGames.Perforce.Managed/StreamTree.cs
--- a/Engine/Source/Programs/Shared/EpicGames.Perforce.Managed/StreamTree.cs
+++ b/Engine/Source/Programs/Shared/EpicGames.Perforce.Managed/StreamTree.cs
@@ -36,10 +36,10 @@ namespace EpicGames.Perforce.Managed
 		public int Revision { get; }
 
 		#region Field names
-		static readonly Utf8String s_lengthField = new Utf8String("len");
-		static readonly Utf8String s_digestField = new Utf8String("dig");
-		static readonly Utf8String s_typeField = new Utf8String("type");
-		static readonly Utf8String s_revisionField = new Utf8String("rev");
+		static readonly Utf8String s_lengthField = new("len");
+		static readonly Utf8String s_digestField = new("dig");
+		static readonly Utf8String s_typeField = new("type");
+		static readonly Utf8String s_revisionField = new("rev");
 		#endregion
 
 		/// <summary>
@@ -72,7 +72,7 @@ namespace EpicGames.Perforce.Managed
 			Path = path;
 			Length = field[s_lengthField].AsInt64();
 
-			Md5Hash digest = new Md5Hash(field[s_digestField].AsBinary());
+			Md5Hash digest = new(field[s_digestField].AsBinary());
 			Utf8String type = field[s_typeField].AsUtf8String();
 			ContentId = new FileContentId(digest, type);
 
@@ -114,7 +114,7 @@ namespace EpicGames.Perforce.Managed
 		public IoHash Hash { get; set; }
 
 		#region Field names
-		static readonly Utf8String s_hashField = new Utf8String("hash");
+		static readonly Utf8String s_hashField = new("hash");
 		#endregion
 
 		/// <summary>
@@ -145,7 +145,7 @@ namespace EpicGames.Perforce.Managed
 		/// <returns></returns>
 		public IoHash ComputeHash()
 		{
-			CbWriter writer = new CbWriter();
+			CbWriter writer = new();
 			writer.BeginObject();
 			Write(writer);
 			writer.EndObject();
@@ -183,10 +183,10 @@ namespace EpicGames.Perforce.Managed
 		public Dictionary<Utf8String, StreamTreeRef> NameToTree { get; } = new Dictionary<Utf8String, StreamTreeRef>(FileUtils.PlatformPathComparerUtf8);
 
 		#region Field names
-		static readonly Utf8String s_nameField = new Utf8String("name");
-		static readonly Utf8String s_pathField = new Utf8String("path");
-		static readonly Utf8String s_filesField = new Utf8String("files");
-		static readonly Utf8String s_treesField = new Utf8String("trees");
+		static readonly Utf8String s_nameField = new("name");
+		static readonly Utf8String s_pathField = new("path");
+		static readonly Utf8String s_filesField = new("files");
+		static readonly Utf8String s_treesField = new("trees");
 		#endregion
 
 		/// <summary>
@@ -223,7 +223,7 @@ namespace EpicGames.Perforce.Managed
 
 				Utf8String name = fileObject[s_nameField].AsUtf8String();
 				Utf8String filePath = ReadPath(fileObject, path, name);
-				StreamFile file = new StreamFile(filePath, fileObject);
+				StreamFile file = new(filePath, fileObject);
 
 				NameToFile.Add(name, file);
 			}
@@ -235,7 +235,7 @@ namespace EpicGames.Perforce.Managed
 
 				Utf8String name = treeObject[s_nameField].AsUtf8String();
 				Utf8String treePath = ReadPath(treeObject, path, name);
-				StreamTreeRef tree = new StreamTreeRef(treePath, treeObject);
+				StreamTreeRef tree = new(treePath, treeObject);
 
 				NameToTree.Add(name, tree);
 			}
@@ -331,7 +331,7 @@ namespace EpicGames.Perforce.Managed
 		/// <returns></returns>
 		public CbObject ToCbObject()
 		{
-			CbWriter writer = new CbWriter();
+			CbWriter writer = new();
 			writer.BeginObject();
 			Write(writer);
 			writer.EndObject();
diff --git a/Engine/Source/Programs/Shared/EpicGames.Perforce.Managed/StreamTreeBuilder.cs b/Engine/Source/Programs/Shared/EpicGames.Perforce.Managed/StreamTreeBuilder.cs
--- a/Engine/Source/Programs/Shared/EpicGames.Perforce.Managed/StreamTreeBuilder.cs
+++ b/Engine/Source/Programs/Shared/EpicGames.Perforce.Managed/StreamTreeBuilder.cs
@@ -146,7 +146,7 @@ namespace EpicGames.Perforce.Managed
 			for (int i = 0; i < pathFragments.Length - 1; i++)
 			{
 				string pathFragment = pathFragments[i];
-				Utf8String unescapedFragment = new Utf8String(PerforceUtils.UnescapePath(pathFragment));
+				Utf8String unescapedFragment = new(PerforceUtils.UnescapePath(pathFragment));
 
 				if (!currentStreamDirectory.NameToTreeBuilder.TryGetValue(unescapedFragment, out StreamTreeBuilder? nextStreamDirectory))
 				{
@@ -156,7 +156,7 @@ namespace EpicGames.Perforce.Managed
 				currentStreamDirectory = nextStreamDirectory;
 			}
 
-			Utf8String filename = new Utf8String(PerforceUtils.UnescapePath(pathFragments[^1])); // Last fragment is filename
+			Utf8String filename = new(PerforceUtils.UnescapePath(pathFragments[^1])); // Last fragment is filename
 			currentStreamDirectory.NameToFile[filename] = depotFile;
 		}
 	}
diff --git a/Engine/Source/Programs/Shared/EpicGames.Perforce.Managed/Transactions/AddTransaction.cs b/Engine/Source/Programs/Shared/EpicGames.Perforce.Managed/Transactions/AddTransaction.cs
--- a/Engine/Source/Programs/Shared/EpicGames.Perforce.Managed/Transactions/AddTransaction.cs
+++ b/Engine/Source/Programs/Shared/EpicGames.Perforce.Managed/Transactions/AddTransaction.cs
@@ -52,9 +52,9 @@ namespace EpicGames.Perforce.Managed
 	{
 		public WorkspaceDirectoryInfo _newWorkspaceRootDir;
 		public StreamSnapshot _streamSnapshot;
-		public ConcurrentDictionary<CachedFileInfo, WorkspaceFileToMove> _filesToMove = new ConcurrentDictionary<CachedFileInfo, WorkspaceFileToMove>();
-		public ConcurrentQueue<WorkspaceFileToCopy> _filesToCopy = new ConcurrentQueue<WorkspaceFileToCopy>();
-		public ConcurrentQueue<WorkspaceFileToSync> _filesToSync = new ConcurrentQueue<WorkspaceFileToSync>();
+		public ConcurrentDictionary<CachedFileInfo, WorkspaceFileToMove> _filesToMove = new();
+		public ConcurrentQueue<WorkspaceFileToCopy> _filesToCopy = new();
+		public ConcurrentQueue<WorkspaceFileToSync> _filesToSync = new();
 		readonly Dictionary<FileContentId, CachedFileInfo> _contentIdToTrackedFile;
 		readonly Dictionary<FileContentId, WorkspaceFileInfo> _contentIdToWorkspaceFile = [];
 
@@ -136,14 +136,14 @@ namespace EpicGames.Perforce.Managed
 
 		async Task MergeSubDirectoryAsync(Utf8String name, WorkspaceDirectoryInfo workspaceSubDir, StreamTreeRef streamSubTreeRef, WorkspaceDirectoryInfo newWorkspaceDir, AsyncThreadPoolWorkQueue queue)
 		{
-			WorkspaceDirectoryInfo newWorkspaceSubDir = new WorkspaceDirectoryInfo(newWorkspaceDir, name, streamSubTreeRef);
+			WorkspaceDirectoryInfo newWorkspaceSubDir = new(newWorkspaceDir, name, streamSubTreeRef);
 			newWorkspaceDir.NameToSubDirectory.Add(name, newWorkspaceSubDir);
 			await queue.EnqueueAsync(_ => MergeDirectoryAsync(workspaceSubDir, newWorkspaceSubDir, streamSubTreeRef, queue));
 		}
 
 		async Task AddSubDirectoryAsync(Utf8String name, WorkspaceDirectoryInfo newWorkspaceDir, StreamTreeRef streamSubTreeRef, AsyncThreadPoolWorkQueue queue)
 		{
-			WorkspaceDirectoryInfo newWorkspaceSubDir = new WorkspaceDirectoryInfo(newWorkspaceDir, name, streamSubTreeRef);
+			WorkspaceDirectoryInfo newWorkspaceSubDir = new(newWorkspaceDir, name, streamSubTreeRef);
 			newWorkspaceDir.NameToSubDirectory.Add(name, newWorkspaceSubDir);
 			await queue.EnqueueAsync(_ => AddDirectoryAsync(newWorkspaceSubDir, streamSubTreeRef, queue));
 		}
@@ -151,7 +151,7 @@ namespace EpicGames.Perforce.Managed
 		void AddFile(WorkspaceDirectoryInfo newWorkspaceDir, Utf8String name, StreamFile streamFile)
 		{
 			// Create a new file for this workspace
-			WorkspaceFileInfo workspaceFile = new WorkspaceFileInfo(newWorkspaceDir, name, streamFile.ContentId);
+			WorkspaceFileInfo workspaceFile = new(newWorkspaceDir, name, streamFile.ContentId);
 			newWorkspaceDir.NameToFile.Add(name, workspaceFile);
 
 			// Try to add it to the cache
diff --git a/Engine/Source/Programs/Shared/EpicGames.Perforce.Managed/Transactions/RemoveTransaction.cs b/Engine/Source/Programs/Shared/EpicGames.Perforce.Managed/Transactions/RemoveTransaction.cs
--- a/Engine/Source/Programs/Shared/EpicGames.Perforce.Managed/Transactions/RemoveTransaction.cs
+++ b/Engine/Source/Programs/Shared/EpicGames.Perforce.Managed/Transactions/RemoveTransaction.cs
@@ -13,9 +13,9 @@ namespace EpicGames.Perforce.Managed
 
 		public WorkspaceDirectoryInfo _newWorkspaceRootDir;
 		public StreamSnapshot _streamSnapshot;
-		public ConcurrentDictionary<FileContentId, WorkspaceFileInfo> _filesToMove = new ConcurrentDictionary<FileContentId, WorkspaceFileInfo>();
-		public ConcurrentQueue<WorkspaceFileInfo> _filesToDelete = new ConcurrentQueue<WorkspaceFileInfo>();
-		public ConcurrentQueue<WorkspaceDirectoryInfo> _directoriesToDelete = new ConcurrentQueue<WorkspaceDirectoryInfo>();
+		public ConcurrentDictionary<FileContentId, WorkspaceFileInfo> _filesToMove = new();
+		public ConcurrentQueue<WorkspaceFileInfo> _filesToDelete = new();
+		public ConcurrentQueue<WorkspaceDirectoryInfo> _directoriesToDelete = new();
 
 		public RemoveTransaction(WorkspaceDirectoryInfo workspaceRootDir, StreamSnapshot streamSnapshot, Dictionary<FileContentId, CachedFileInfo> contentIdToTrackedFile)
 		{
@@ -69,7 +69,7 @@ namespace EpicGames.Perforce.Managed
 
 		async Task MergeSubDirectoryAsync(Utf8String name, WorkspaceDirectoryInfo newWorkspaceDir, WorkspaceDirectoryInfo workspaceSubDir, StreamTreeRef streamSubTreeRef, AsyncThreadPoolWorkQueue queue)
 		{
-			WorkspaceDirectoryInfo newWorkspaceSubDir = new WorkspaceDirectoryInfo(newWorkspaceDir, name, streamSubTreeRef);
+			WorkspaceDirectoryInfo newWorkspaceSubDir = new(newWorkspaceDir, name, streamSubTreeRef);
 			newWorkspaceDir.NameToSubDirectory.Add(newWorkspaceSubDir.Name, newWorkspaceSubDir);
 			await queue.EnqueueAsync(_ => MergeAsync(workspaceSubDir, newWorkspaceSubDir, streamSubTreeRef, queue));
 		}
diff --git a/Engine/Source/Programs/Shared/EpicGames.Perforce.Managed/WorkspaceDirectoryInfo.cs b/Engine/Source/Programs/Shared/EpicGames.Perforce.Managed/WorkspaceDirectoryInfo.cs
--- a/Engine/Source/Programs/Shared/EpicGames.Perforce.Managed/WorkspaceDirectoryInfo.cs
+++ b/Engine/Source/Programs/Shared/EpicGames.Perforce.Managed/WorkspaceDirectoryInfo.cs
@@ -129,8 +129,8 @@ namespace EpicGames.Perforce.Managed
 		/// <param name="numWorkers">Number of concurrent workers to use when refreshing</param>
 		public async Task<(FileInfo[] filesToDelete, DirectoryInfo[] directoriesToDelete)> RefreshAsync(bool removeUntracked, int numWorkers)
 		{
-			ConcurrentQueue<FileInfo> concurrentFilesToDelete = new ConcurrentQueue<FileInfo>();
-			ConcurrentQueue<DirectoryInfo> concurrentDirectoriesToDelete = new ConcurrentQueue<DirectoryInfo>();
+			ConcurrentQueue<FileInfo> concurrentFilesToDelete = new();
+			ConcurrentQueue<DirectoryInfo> concurrentDirectoriesToDelete = new();
 			using AsyncThreadPoolWorkQueue queue = new(numWorkers);
 			await queue.EnqueueAsync(_ => RefreshAsync(new DirectoryInfo(GetFullName()), removeUntracked, concurrentFilesToDelete, concurrentDirectoriesToDelete, queue));
 			await queue.ExecuteAsync();
@@ -148,10 +148,10 @@ namespace EpicGames.Perforce.Managed
 		async Task RefreshAsync(DirectoryInfo info, bool removeUntracked, ConcurrentQueue<FileInfo> filesToDelete, ConcurrentQueue<DirectoryInfo> directoriesToDelete, AsyncThreadPoolWorkQueue queue)
 		{
 			// Recurse through subdirectories
-			Dictionary<Utf8String, WorkspaceDirectoryInfo> newNameToSubDirectory = new Dictionary<Utf8String, WorkspaceDirectoryInfo>(NameToSubDirectory.Count, NameToSubDirectory.Comparer);
+			Dictionary<Utf8String, WorkspaceDirectoryInfo> newNameToSubDirectory = new(NameToSubDirectory.Count, NameToSubDirectory.Comparer);
 
 			// Figure out which files have changed.
-			Dictionary<Utf8String, WorkspaceFileInfo> newNameToFile = new Dictionary<Utf8String, WorkspaceFileInfo>(NameToFile.Count, NameToFile.Comparer);
+			Dictionary<Utf8String, WorkspaceFileInfo> newNameToFile = new(NameToFile.Count, NameToFile.Comparer);
 
 			foreach (FileSystemInfo fileSystemInfo in info.EnumerateFileSystemInfos())
 			{
@@ -278,7 +278,7 @@ namespace EpicGames.Perforce.Managed
 		/// <returns></returns>
 		public string GetFullName()
 		{
-			StringBuilder builder = new StringBuilder();
+			StringBuilder builder = new();
 			AppendFullPath(builder);
 			return builder.ToString();
 		}
@@ -360,7 +360,7 @@ namespace EpicGames.Perforce.Managed
 			{
 				Utf8String name = reader.ReadNullTerminatedUtf8String();
 
-				WorkspaceDirectoryInfo subDirectory = new WorkspaceDirectoryInfo(directoryInfo, name, null);
+				WorkspaceDirectoryInfo subDirectory = new(directoryInfo, name, null);
 				reader.ReadWorkspaceDirectoryInfo(subDirectory, version);
 				directoryInfo.NameToSubDirectory[subDirectory.Name] = subDirectory;
 			}
diff --git a/Engine/Source/Programs/Shared/EpicGames.Perforce.Managed/WorkspaceFileInfo.cs b/Engine/Source/Programs/Shared/EpicGames.Perforce.Managed/WorkspaceFileInfo.cs
--- a/Engine/Source/Programs/Shared/EpicGames.Perforce.Managed/WorkspaceFileInfo.cs
+++ b/Engine/Source/Programs/Shared/EpicGames.Perforce.Managed/WorkspaceFileInfo.cs
@@ -49,7 +49,7 @@ namespace EpicGames.Perforce.Managed
 
 		public void UpdateMetadata()
 		{
-			FileInfo info = new FileInfo(GetFullName());
+			FileInfo info = new(GetFullName());
 			if (info.Exists)
 			{
 				_length = info.Length;
@@ -65,7 +65,7 @@ namespace EpicGames.Perforce.Managed
 
 		public string GetClientPath()
 		{
-			StringBuilder builder = new StringBuilder();
+			StringBuilder builder = new();
 			Directory.AppendClientPath(builder);
 			builder.Append(Name);
 			return builder.ToString();
@@ -73,7 +73,7 @@ namespace EpicGames.Perforce.Managed
 
 		public string GetFullName()
 		{
-			StringBuilder builder = new StringBuilder();
+			StringBuilder builder = new();
 			Directory.AppendFullPath(builder);
 			builder.Append(Path.DirectorySeparatorChar);
 			builder.Append(Name);
diff --git a/Engine/Source/Programs/Shared/EpicGames.Perforce/FileSpecList.cs b/Engine/Source/Programs/Shared/EpicGames.Perforce/FileSpecList.cs
--- a/Engine/Source/Programs/Shared/EpicGames.Perforce/FileSpecList.cs
+++ b/Engine/Source/Programs/Shared/EpicGames.Perforce/FileSpecList.cs
@@ -41,7 +41,7 @@ namespace EpicGames.Perforce
 		/// <inheritdoc/>
 		public override int GetHashCode()
 		{
-			HashCode code = new HashCode();
+			HashCode code = new();
 			foreach (string item in List)
 			{
 				code.Add(item);
diff --git a/Engine/Source/Programs/Shared/EpicGames.Perforce/HangMonitor.cs b/Engine/Source/Programs/Shared/EpicGames.Perforce/HangMonitor.cs
--- a/Engine/Source/Programs/Shared/EpicGames.Perforce/HangMonitor.cs
+++ b/Engine/Source/Programs/Shared/EpicGames.Perforce/HangMonitor.cs
@@ -33,7 +33,7 @@ namespace EpicGames.Perforce
 			}
 		}
 
-		readonly AsyncEvent _activeEvent = new AsyncEvent();
+		readonly AsyncEvent _activeEvent = new();
 		readonly TimeSpan _interval;
 		readonly string _context;
 		readonly ILogger _logger;
@@ -70,7 +70,7 @@ namespace EpicGames.Perforce
 		/// <param name="activity">Activity to log if a hang is detected</param>
 		public IDisposable Start(string activity)
 		{
-			Scope scope = new Scope(this, activity);
+			Scope scope = new(this, activity);
 			if (Interlocked.CompareExchange(ref _scope, scope, null) != null)
 			{
 				throw new InvalidOperationException();
diff --git a/Engine/Source/Programs/Shared/EpicGames.Perforce/IPerforceOutput.cs b/Engine/Source/Programs/Shared/EpicGames.Perforce/IPerforceOutput.cs
--- a/Engine/Source/Programs/Shared/EpicGames.Perforce/IPerforceOutput.cs
+++ b/Engine/Source/Programs/Shared/EpicGames.Perforce/IPerforceOutput.cs
@@ -47,13 +47,13 @@ namespace EpicGames.Perforce
 		/// </summary>
 		static class ReadOnlyUtf8StringConstants
 		{
-			public static readonly Utf8String Code = new Utf8String("code");
-			public static readonly Utf8String Stat = new Utf8String("stat");
-			public static readonly Utf8String Info = new Utf8String("info");
-			public static readonly Utf8String Error = new Utf8String("error");
-			public static readonly Utf8String Io = new Utf8String("io");
-			public static readonly Utf8String Func = new Utf8String("func");
-			public static readonly Utf8String IsSparse = new Utf8String("isSparse");
+			public static readonly Utf8String Code = new("code");
+			public static readonly Utf8String Stat = new("stat");
+			public static readonly Utf8String Info = new("info");
+			public static readonly Utf8String Error = new("error");
+			public static readonly Utf8String Io = new("io");
+			public static readonly Utf8String Func = new("func");
+			public static readonly Utf8String IsSparse = new("isSparse");
 		}
 
 		/// <summary>
@@ -72,7 +72,7 @@ namespace EpicGames.Perforce
 
 			public Task<bool> ReadAsync(CancellationToken cancellationToken) => Task.FromResult(Data.Length > 0);
 			public void Discard(int numBytes) => Data = Data.Slice(numBytes);
-			public ValueTask DisposeAsync() => new ValueTask();
+			public ValueTask DisposeAsync() => new();
 		}
 
 		/// <summary>
@@ -96,7 +96,7 @@ namespace EpicGames.Perforce
 		/// <returns>Output object</returns>
 		public static IPerforceOutput FromResponses(IEnumerable<PerforceResponse> responses)
 		{
-			using ChunkedMemoryWriter writer = new ChunkedMemoryWriter();
+			using ChunkedMemoryWriter writer = new();
 			foreach (PerforceResponse response in responses)
 			{
 				writer.WriteFixedLengthBytes(s_recordPrefix);
@@ -120,7 +120,7 @@ namespace EpicGames.Perforce
 		/// <returns>String representation of the buffer</returns>
 		private static string FormatDataAsString(ReadOnlySpan<byte> data)
 		{
-			StringBuilder result = new StringBuilder();
+			StringBuilder result = new();
 			if (data.Length > 0)
 			{
 				for (int idx = 0; idx < data.Length && idx < 1024;)
@@ -171,7 +171,7 @@ namespace EpicGames.Perforce
 		private static string FormatDataAsHexDump(ReadOnlySpan<byte> data, int maxLength = 1024)
 		{
 			// Format the result
-			StringBuilder result = new StringBuilder();
+			StringBuilder result = new();
 
 			const int RowLength = 16;
 			for (int baseIdx = 0; baseIdx < data.Length && baseIdx < maxLength; baseIdx += RowLength)
@@ -346,7 +346,7 @@ namespace EpicGames.Perforce
 		/// <returns>Async task</returns>
 		public static async Task ReadRecordsAsync(this IPerforceOutput perforce, Action<PerforceRecord> handleRecord, CancellationToken cancellationToken)
 		{
-			PerforceRecord record = new PerforceRecord();
+			PerforceRecord record = new();
 			while (await perforce.ReadAsync(cancellationToken))
 			{
 				// Start a read to add more data
diff --git a/Engine/Source/Programs/Shared/EpicGames.Perforce/NativePerforceConnection.cs b/Engine/Source/Programs/Shared/EpicGames.Perforce/NativePerforceConnection.cs
--- a/Engine/Source/Programs/Shared/EpicGames.Perforce/NativePerforceConnection.cs
+++ b/Engine/Source/Programs/Shared/EpicGames.Perforce/NativePerforceConnection.cs
@@ -296,7 +296,7 @@ namespace EpicGames.Perforce
 		readonly BlockingCollection<PinnedBuffer> _writeBuffers = [];
 		Response? _currentResponse;
 		readonly ManualResetEvent _responseCompleteEvent;
-		readonly Stopwatch _stallTimer = new Stopwatch();
+		readonly Stopwatch _stallTimer = new();
 		readonly HangMonitor _hangMonitor;
 		bool _disposed;
 
@@ -500,7 +500,7 @@ namespace EpicGames.Perforce
 		/// <returns>Error returned when attempting to connect</returns>
 		private async Task<PerforceError?> TryConnectAsync()
 		{
-			await using Response response = new Response(this);
+			await using Response response = new(this);
 
 			NativeSettings? nativeSettings = null;
 			if (Settings != null)
@@ -517,7 +517,7 @@ namespace EpicGames.Perforce
 
 			_requests.Add((() =>
 			{
-				NativeWriteBuffer writeBuffer = new NativeWriteBuffer();
+				NativeWriteBuffer writeBuffer = new();
 				GetNextWriteBuffer(writeBuffer, 1);
 				_client = Client_Create(nativeSettings, writeBuffer, _onBufferReadyFnPtr);
 			}, response));
@@ -612,14 +612,14 @@ namespace EpicGames.Perforce
 				allArguments.AddRange(fileArguments);
 			}
 
-			Response response = new Response(this);
+			Response response = new(this);
 			_requests.Add((() => ExecCommand(command, allArguments, specData, promptResponse, interceptIo), response));
 			return response;
 		}
 
 		private void ExecCommand(string command, List<string> args, byte[]? inputData, string? promptResponse, bool interceptIo)
 		{
-			StringBuilder argList = new StringBuilder();
+			StringBuilder argList = new();
 			for (int idx = 0; idx < args.Count; idx++)
 			{
 				CommandLineArguments.Append(argList, args[idx]);
@@ -711,7 +711,7 @@ namespace EpicGames.Perforce
 				throw new PerforceException("Garbage after end of spec data");
 			}
 
-			StringBuilder result = new StringBuilder();
+			StringBuilder result = new();
 			foreach ((Utf8String key, PerforceValue value) in rows)
 			{
 				string[] valueLines = value.ToString().Split('\n');
diff --git a/Engine/Source/Programs/Shared/EpicGames.Perforce/PerforceChangeView.cs b/Engine/Source/Programs/Shared/EpicGames.Perforce/PerforceChangeView.cs
--- a/Engine/Source/Programs/Shared/EpicGames.Perforce/PerforceChangeView.cs
+++ b/Engine/Source/Programs/Shared/EpicGames.Perforce/PerforceChangeView.cs
@@ -42,7 +42,7 @@ namespace EpicGames.Perforce
 		/// <param name="ignoreCase">Whether this will be evaluated in the context of a case-insensitive server</param>
 		public static PerforceChangeView Parse(IEnumerable<string> lines, bool ignoreCase)
 		{
-			PerforceChangeView changeView = new PerforceChangeView();
+			PerforceChangeView changeView = new();
 			foreach (string line in lines)
 			{
 				if (!String.IsNullOrWhiteSpace(line))
diff --git a/Engine/Source/Programs/Shared/EpicGames.Perforce/PerforceChildProcess.cs b/Engine/Source/Programs/Shared/EpicGames.Perforce/PerforceChildProcess.cs
--- a/Engine/Source/Programs/Shared/EpicGames.Perforce/PerforceChildProcess.cs
+++ b/Engine/Source/Programs/Shared/EpicGames.Perforce/PerforceChildProcess.cs
@@ -187,7 +187,7 @@ namespace EpicGames.Perforce
 		/// <returns>Exit code and output from the process</returns>
 		public async Task<Tuple<bool, string>> TryReadToEndAsync(CancellationToken cancellationToken)
 		{
-			using MemoryStream stream = new MemoryStream();
+			using MemoryStream stream = new();
 
 			while (await ReadAsync(cancellationToken))
 			{
diff --git a/Engine/Source/Programs/Shared/EpicGames.Perforce/PerforceConnection.cs b/Engine/Source/Programs/Shared/EpicGames.Perforce/PerforceConnection.cs
--- a/Engine/Source/Programs/Shared/EpicGames.Perforce/PerforceConnection.cs
+++ b/Engine/Source/Programs/Shared/EpicGames.Perforce/PerforceConnection.cs
@@ -95,7 +95,7 @@ namespace EpicGames.Perforce
 
 		static PerforceSettings CombineSettings(string? serverAndPort, string? userName, string? clientName = null, string? appName = null, string? appVersion = null)
 		{
-			PerforceSettings settings = new PerforceSettings(PerforceSettings.Default);
+			PerforceSettings settings = new(PerforceSettings.Default);
 			if (serverAndPort != null)
 			{
 				settings.ServerAndPort = serverAndPort;
@@ -126,7 +126,7 @@ namespace EpicGames.Perforce
 		{
 			get
 			{
-				PerforceSettings settings = new PerforceSettings(PerforceEnvironment.Default);
+				PerforceSettings settings = new(PerforceEnvironment.Default);
 				if (ServerAndPort != null)
 				{
 					settings.ServerAndPort = ServerAndPort;
@@ -364,7 +364,7 @@ namespace EpicGames.Perforce
 		{
 			List<string> arguments = [$"{name}={value}"];
 
-			using (PerforceChildProcess childProcess = new PerforceChildProcess("set", arguments, null, null, GetGlobalArguments(), Logger))
+			using (PerforceChildProcess childProcess = new("set", arguments, null, null, GetGlobalArguments(), Logger))
 			{
 				return await childProcess.TryReadToEndAsync(cancellationToken);
 			}
@@ -389,7 +389,7 @@ namespace EpicGames.Perforce
 		/// <returns></returns>
 		public static Task<IPerforceConnection> WithClientAsync(this IPerforceConnection perforce, string? clientName)
 		{
-			PerforceSettings settings = new PerforceSettings(perforce.Settings) { ClientName = clientName };
+			PerforceSettings settings = new(perforce.Settings) { ClientName = clientName };
 			return PerforceConnection.CreateAsync(settings, perforce.Logger, perforce.Tracer);
 		}
 
@@ -4058,7 +4058,7 @@ namespace EpicGames.Perforce
 			// Using multiple threads is not supported through p4.exe due to threaded output not being parsable
 			if (numThreads != -1 && (connection is NativePerforceConnection))
 			{
-				StringBuilder argument = new StringBuilder($"--parallel=threads={numThreads}");
+				StringBuilder argument = new($"--parallel=threads={numThreads}");
 				if (batch != -1)
 				{
 					argument.Append($",batch={batch}");
diff --git a/Engine/Source/Programs/Shared/EpicGames.Perforce/PerforceEnvironment.cs b/Engine/Source/Programs/Shared/EpicGames.Perforce/PerforceEnvironment.cs
--- a/Engine/Source/Programs/Shared/EpicGames.Perforce/PerforceEnvironment.cs
+++ b/Engine/Source/Programs/Shared/EpicGames.Perforce/PerforceEnvironment.cs
@@ -104,7 +104,7 @@ namespace EpicGames.Perforce
 		/// </summary>
 		public FileReference Location { get; }
 
-		readonly Dictionary<string, string> _variables = new Dictionary<string, string>(StringComparer.OrdinalIgnoreCase);
+		readonly Dictionary<string, string> _variables = new(StringComparer.OrdinalIgnoreCase);
 
 		/// <summary>
 		/// Constructor
@@ -180,7 +180,7 @@ namespace EpicGames.Perforce
 
 			if (enviroValue != null)
 			{
-				FileReference location = new FileReference(enviroValue);
+				FileReference location = new(enviroValue);
 				if (FileReference.Exists(location))
 				{
 					environment = new PerforceEnvironmentFile(environment, location);
diff --git a/Engine/Source/Programs/Shared/EpicGames.Perforce/PerforceMetadataLogger.cs b/Engine/Source/Programs/Shared/EpicGames.Perforce/PerforceMetadataLogger.cs
--- a/Engine/Source/Programs/Shared/EpicGames.Perforce/PerforceMetadataLogger.cs
+++ b/Engine/Source/Programs/Shared/EpicGames.Perforce/PerforceMetadataLogger.cs
@@ -54,7 +54,7 @@ namespace EpicGames.Perforce
 		/// <param name="change">Changelist for the client</param>
 		public void AddClientView(DirectoryReference baseDir, string depotPath, int change)
 		{
-			PerforceViewMap viewMap = new PerforceViewMap();
+			PerforceViewMap viewMap = new();
 			viewMap.Entries.Add(new PerforceViewMapEntry(true, depotPath, "..."));
 			AddClientView(baseDir, viewMap, change);
 		}
@@ -150,14 +150,14 @@ namespace EpicGames.Perforce
 		static readonly Utf8String s_sourceFileType = LogValueType.SourceFile;
 		static readonly Utf8String s_assetType = LogValueType.Asset;
 
-		static readonly Utf8String s_file = new Utf8String("file");
+		static readonly Utf8String s_file = new("file");
 
 		ReadOnlyMemory<byte> Annotate(ReadOnlyMemory<byte> data)
 		{
 			ReadOnlyMemory<byte> output = data;
 			int shift = 0;
 
-			Utf8JsonReader reader = new Utf8JsonReader(data.Span);
+			Utf8JsonReader reader = new(data.Span);
 			for (bool valid = ReadFirstLogProperty(ref reader); valid; valid = ReadNextLogProperty(ref reader))
 			{
 				ReadOnlySpan<byte> type = [];
@@ -217,7 +217,7 @@ namespace EpicGames.Perforce
 		{
 			if (type.SequenceEqual(s_sourceFileType) || type.SequenceEqual(s_assetType))
 			{
-				StringBuilder annotations = new StringBuilder();
+				StringBuilder annotations = new();
 
 				file ??= Encoding.UTF8.GetString(text);
 
diff --git a/Engine/Source/Programs/Shared/EpicGames.Perforce/PerforceRecord.cs b/Engine/Source/Programs/Shared/EpicGames.Perforce/PerforceRecord.cs
--- a/Engine/Source/Programs/Shared/EpicGames.Perforce/PerforceRecord.cs
+++ b/Engine/Source/Programs/Shared/EpicGames.Perforce/PerforceRecord.cs
@@ -244,7 +244,7 @@ namespace EpicGames.Perforce
 		/// <returns>Record containing the given fields</returns>
 		public static PerforceRecord FromFields(IEnumerable<KeyValuePair<string, object>> fields, bool numberedListElements)
 		{
-			PerforceRecord record = new PerforceRecord();
+			PerforceRecord record = new();
 			foreach ((string name, object value) in fields)
 			{
 				if (value is string str)
@@ -320,8 +320,8 @@ namespace EpicGames.Perforce
 		/// <returns>Serialized data</returns>
 		public byte[] Serialize()
 		{
-			MemoryStream stream = new MemoryStream();
-			using (BinaryWriter writer = new BinaryWriter(stream))
+			MemoryStream stream = new();
+			using (BinaryWriter writer = new(stream))
 			{
 				writer.Write((byte)'{');
 				foreach ((Utf8String name, PerforceValue value) in Rows)
@@ -343,8 +343,8 @@ namespace EpicGames.Perforce
 		/// <returns>Serialized record data</returns>
 		public static byte[] Serialize(List<KeyValuePair<string, object>> keyValuePairs)
 		{
-			MemoryStream stream = new MemoryStream();
-			using (BinaryWriter writer = new BinaryWriter(stream))
+			MemoryStream stream = new();
+			using (BinaryWriter writer = new(stream))
 			{
 				writer.Write((byte)'{');
 				foreach (KeyValuePair<string, object> keyValuePair in keyValuePairs)
diff --git a/Engine/Source/Programs/Shared/EpicGames.Perforce/PerforceReflection.cs b/Engine/Source/Programs/Shared/EpicGames.Perforce/PerforceReflection.cs
--- a/Engine/Source/Programs/Shared/EpicGames.Perforce/PerforceReflection.cs
+++ b/Engine/Source/Programs/Shared/EpicGames.Perforce/PerforceReflection.cs
@@ -19,11 +19,11 @@ namespace EpicGames.Perforce
 	/// </summary>
 	static class StringConstants
 	{
-		public static readonly Utf8String True = new Utf8String("true");
-		public static readonly Utf8String False = new Utf8String("false");
-		public static readonly Utf8String New = new Utf8String("new");
-		public static readonly Utf8String None = new Utf8String("none");
-		public static readonly Utf8String Default = new Utf8String("default");
+		public static readonly Utf8String True = new("true");
+		public static readonly Utf8String False = new("false");
+		public static readonly Utf8String New = new("new");
+		public static readonly Utf8String None = new("none");
+		public static readonly Utf8String Default = new("default");
 	}
 
 	/// <summary>
@@ -110,7 +110,7 @@ namespace EpicGames.Perforce
 		/// <inheritdoc/>
 		public override string ToString()
 		{
-			StringBuilder fullName = new StringBuilder(Tag.ToString());
+			StringBuilder fullName = new(Tag.ToString());
 			if (ParentRecords.Length > 0)
 			{
 				fullName.Append('0');
@@ -214,7 +214,7 @@ namespace EpicGames.Perforce
 					object? value = field.GetValue(null);
 					if (value != null)
 					{
-						Utf8String name = new Utf8String(attribute.Name);
+						Utf8String name = new(attribute.Name);
 						_nameToValue[name] = (int)value;
 						_valueToName[(int)value] = name;
 
@@ -354,7 +354,7 @@ namespace EpicGames.Perforce
 		/// <summary>
 		/// Unix epoch; used for converting times back into C# datetime objects
 		/// </summary>
-		public static readonly DateTime UnixEpoch = new DateTime(1970, 1, 1, 0, 0, 0, DateTimeKind.Utc);
+		public static readonly DateTime UnixEpoch = new(1970, 1, 1, 0, 0, 0, DateTimeKind.Utc);
 
 		/// <summary>
 		/// Constant for the default changelist, where valid.
@@ -364,12 +364,12 @@ namespace EpicGames.Perforce
 		/// <summary>
 		/// Cached map of enum types to a lookup mapping from p4 strings to enum values.
 		/// </summary>
-		static readonly ConcurrentDictionary<Type, CachedEnumInfo> s_enumTypeToInfo = new ConcurrentDictionary<Type, CachedEnumInfo>();
+		static readonly ConcurrentDictionary<Type, CachedEnumInfo> s_enumTypeToInfo = new();
 
 		/// <summary>
 		/// Cached set of record 
 		/// </summary>
-		static readonly ConcurrentDictionary<Type, CachedRecordInfo> s_recordTypeToInfo = new ConcurrentDictionary<Type, CachedRecordInfo>();
+		static readonly ConcurrentDictionary<Type, CachedRecordInfo> s_recordTypeToInfo = new();
 
 		/// <summary>
 		/// Default type for info
@@ -508,7 +508,7 @@ namespace EpicGames.Perforce
 						rootRecord.RequiredTagsBitMask |= requiredTagBitMask;
 					}
 
-					TaggedPropertyInfo tagInfo = new TaggedPropertyInfo(new Utf8String(tagName), tagAttribute.Optional, parentRecords, property, requiredTagBitMask);
+					TaggedPropertyInfo tagInfo = new(new Utf8String(tagName), tagAttribute.Optional, parentRecords, property, requiredTagBitMask);
 
 					Type fieldType = property.PropertyType;
 
@@ -579,7 +579,7 @@ namespace EpicGames.Perforce
 				if (subElementAttribute != null)
 				{
 					Type newRecordType = property.PropertyType.GenericTypeArguments[0];
-					NestedRecordInfo newParentRecord = new NestedRecordInfo(property, GetCreateRecordDelegate(newRecordType));
+					NestedRecordInfo newParentRecord = new(property, GetCreateRecordDelegate(newRecordType));
 					AddRecordProperties(newRecordType, rootRecord, [.. parentRecords, newParentRecord]);
 				}
 			}
@@ -593,7 +593,7 @@ namespace EpicGames.Perforce
 				throw new PerforceException($"Unable to find default constructor for {type}");
 			}
 
-			DynamicMethod dynamicMethod = new DynamicMethod("_", type, null);
+			DynamicMethod dynamicMethod = new("_", type, null);
 			ILGenerator generator = dynamicMethod.GetILGenerator();
 			generator.Emit(OpCodes.Newobj, constructor);
 			generator.Emit(OpCodes.Ret);
diff --git a/Engine/Source/Programs/Shared/EpicGames.Perforce/PerforceSettings.cs b/Engine/Source/Programs/Shared/EpicGames.Perforce/PerforceSettings.cs
--- a/Engine/Source/Programs/Shared/EpicGames.Perforce/PerforceSettings.cs
+++ b/Engine/Source/Programs/Shared/EpicGames.Perforce/PerforceSettings.cs
@@ -181,7 +181,7 @@ namespace EpicGames.Perforce
 		/// </summary>
 		public static IPerforceSettings MergeWith(this IPerforceSettings settings, string? newServerAndPort = null, string? newUserName = null, string? newClientName = null)
 		{
-			PerforceSettings newSettings = new PerforceSettings(settings);
+			PerforceSettings newSettings = new(settings);
 			if (!String.IsNullOrEmpty(newServerAndPort))
 			{
 				newSettings.ServerAndPort = newServerAndPort;
@@ -205,7 +205,7 @@ namespace EpicGames.Perforce
 		/// <returns></returns>
 		public static string GetArgumentsForExternalProgram(this IPerforceSettings settings, bool includeClient)
 		{
-			StringBuilder basicCommandArgs = new StringBuilder();
+			StringBuilder basicCommandArgs = new();
 
 			basicCommandArgs.AppendFormat("-p \"{0}\"", settings.ServerAndPort);
 			basicCommandArgs.AppendFormat(" -u \"{0}\"", settings.UserName);
diff --git a/Engine/Source/Programs/Shared/EpicGames.Perforce/PerforceViewFilter.cs b/Engine/Source/Programs/Shared/EpicGames.Perforce/PerforceViewFilter.cs
--- a/Engine/Source/Programs/Shared/EpicGames.Perforce/PerforceViewFilter.cs
+++ b/Engine/Source/Programs/Shared/EpicGames.Perforce/PerforceViewFilter.cs
@@ -192,7 +192,7 @@ namespace EpicGames.Perforce
 		/// <inheritdoc/>
 		public override string ToString()
 		{
-			StringBuilder builder = new StringBuilder();
+			StringBuilder builder = new();
 			if (!Include)
 			{
 				builder.Append('-');
diff --git a/Engine/Source/Programs/Shared/EpicGames.Perforce/Records/FileType.cs b/Engine/Source/Programs/Shared/EpicGames.Perforce/Records/FileType.cs
--- a/Engine/Source/Programs/Shared/EpicGames.Perforce/Records/FileType.cs
+++ b/Engine/Source/Programs/Shared/EpicGames.Perforce/Records/FileType.cs
@@ -357,7 +357,7 @@ namespace EpicGames.Perforce.Managed
 		/// <inheritdoc/>
 		public override string ToString()
 		{
-			StringBuilder type = new StringBuilder(s_baseNames[(int)Base].ToString());
+			StringBuilder type = new(s_baseNames[(int)Base].ToString());
 
 			int remainingModifiers = (int)Modifiers;
 			if (remainingModifiers != 0 || NumRevisions > 0)
diff --git a/Engine/Source/Programs/Shared/EpicGames.ProjectStore/CloudStorage/CloudStorage.cs b/Engine/Source/Programs/Shared/EpicGames.ProjectStore/CloudStorage/CloudStorage.cs
--- a/Engine/Source/Programs/Shared/EpicGames.ProjectStore/CloudStorage/CloudStorage.cs
+++ b/Engine/Source/Programs/Shared/EpicGames.ProjectStore/CloudStorage/CloudStorage.cs
@@ -57,7 +57,7 @@ namespace EpicGames.ProjectStore
 			HttpGatewayTimeout = 93
 		}
 
-		static readonly Regex s_logLinePattern = new Regex(@"\[.*\]\s?\[(?<severity>warning|error|info)\]\s?(\[(?<file>\S+)\:(?<line>\d+)\])?(?<message>.*)");
+		static readonly Regex s_logLinePattern = new(@"\[.*\]\s?\[(?<severity>warning|error|info)\]\s?(\[(?<file>\S+)\:(?<line>\d+)\])?(?<message>.*)");
 		/// <summary>
 		/// CTOR
 		/// </summary>
@@ -110,7 +110,7 @@ namespace EpicGames.ProjectStore
 				throw new CloudStorageException("No host configured");
 			}
 
-			Uri host = new Uri(optionsOverride?.OverrideHost ?? optionsOverride?.DefaultHost ?? _options.Value.Host!);
+			Uri host = new(optionsOverride?.OverrideHost ?? optionsOverride?.DefaultHost ?? _options.Value.Host!);
 
 			// TODO: Should we support local file configuration?
 			ClientAuthConfigurationV1? authConfig = await ProviderConfigurationFactory.ReadRemoteAuthConfigurationAsync(host, ProviderConfigurationFactory.DefaultEncryptionKey);
@@ -241,7 +241,7 @@ namespace EpicGames.ProjectStore
 			try
 			{
 				// copy the query and append the bucket regex filter
-				CbWriter writer = new CbWriter();
+				CbWriter writer = new();
 				writer.BeginObject();
 				writer.WriteObject("query", query);
 				writer.WriteString("bucketRegex", bucketRegex.ToString());
@@ -277,7 +277,7 @@ namespace EpicGames.ProjectStore
 
 				EnsureZenExitCodeSuccess(exitCode);
 				byte[] b = File.ReadAllBytes(tempOutputFile);
-				CbObject o = new CbObject(b);
+				CbObject o = new(b);
 
 				bool partialResult = o["partialResult"].AsBool();
 				if (partialResult)
@@ -396,7 +396,7 @@ namespace EpicGames.ProjectStore
 		/// <inheritdoc />
 		public async Task<List<BuildContents>> ListBuildContentsAsync(StringId bucketId, CbObjectId buildId, CloudBuildOptionsOverride? optionsOverride = null, CancellationToken cancellationToken = default)
 		{
-			Uri uri = new Uri($"{_options.Value.Host}/api/v2/builds/{GetCurrentNamespace(optionsOverride)}/{bucketId}/{buildId}");
+			Uri uri = new($"{_options.Value.Host}/api/v2/builds/{GetCurrentNamespace(optionsOverride)}/{bucketId}/{buildId}");
 			string cmdline = $"builds ls --cloud-url={uri} --quiet";
 			// TODO: we need a way to ask this to get written to file
 			int exitCode = await RunZenCliAsync(cmdline, cancellationToken: cancellationToken);
@@ -606,12 +606,12 @@ namespace EpicGames.ProjectStore
 
 		private async Task<int> RunProcessAsync(string cmd, string arguments, string? workingDirectory = null, Dictionary<string, string>? env = null, IProgress<(float, string)>? progress = null, bool echoStdout = false, ILogger? echoLogger = null, CancellationToken cancellationToken = default)
 		{
-			using ManagedProcessGroup managedProcessGroup = new ManagedProcessGroup();
+			using ManagedProcessGroup managedProcessGroup = new();
 			using ManagedProcess managedProcess = new(managedProcessGroup, cmd, arguments, workingDirectory, env, ProcessPriorityClass.Normal, flags: ManagedProcessFlags.MergeOutputPipes);
 
 			_options.Value.OnProcessCreated?.Invoke(managedProcess);
 
-			ProgressValue progressValue = new ProgressValue();
+			ProgressValue progressValue = new();
 			while (managedProcess.TryReadLine(out string? line, cancellationToken))
 			{
 				if (progress != null && line.StartsWith(ProgressTextReader.DirectivePrefix, StringComparison.Ordinal))
diff --git a/Engine/Source/Programs/Shared/EpicGames.ProjectStore/CloudStorage/CloudStorageQueryBuilder.cs b/Engine/Source/Programs/Shared/EpicGames.ProjectStore/CloudStorage/CloudStorageQueryBuilder.cs
--- a/Engine/Source/Programs/Shared/EpicGames.ProjectStore/CloudStorage/CloudStorageQueryBuilder.cs
+++ b/Engine/Source/Programs/Shared/EpicGames.ProjectStore/CloudStorage/CloudStorageQueryBuilder.cs
@@ -11,7 +11,7 @@ namespace EpicGames.ProjectStore
 	/// </summary>
 	public class CloudStorageQueryBuilder
 	{
-		private readonly CbWriter _writer = new CbWriter();
+		private readonly CbWriter _writer = new();
 		private CbObject? _generatedObject = null;
 
 		/// <summary>
diff --git a/Engine/Source/Programs/Shared/EpicGames.ProjectStore/CloudStorage/ICloudStorage.cs b/Engine/Source/Programs/Shared/EpicGames.ProjectStore/CloudStorage/ICloudStorage.cs
--- a/Engine/Source/Programs/Shared/EpicGames.ProjectStore/CloudStorage/ICloudStorage.cs
+++ b/Engine/Source/Programs/Shared/EpicGames.ProjectStore/CloudStorage/ICloudStorage.cs
@@ -523,7 +523,7 @@ namespace EpicGames.ProjectStore
 			}
 
 			DirectoryReference engineBinariesDir = DirectoryReference.Combine(Unreal.RootDirectory, "Engine", "Binaries", binaryType);
-			FileReference zenExe = new FileReference(Path.Combine(engineBinariesDir.FullName, "zen" + RuntimePlatform.ExeExtension));
+			FileReference zenExe = new(Path.Combine(engineBinariesDir.FullName, "zen" + RuntimePlatform.ExeExtension));
 			ZenExePath = zenExe;
 
 			string binaryTypeDotnet = "win-x64";
@@ -537,7 +537,7 @@ namespace EpicGames.ProjectStore
 			}
 
 			DirectoryReference oidcDir = DirectoryReference.Combine(Unreal.RootDirectory, "Engine", "Binaries", "DotNET", "OidcToken", binaryTypeDotnet);
-			FileReference oidcExe = new FileReference(Path.Combine(oidcDir.FullName, "OidcToken" + RuntimePlatform.ExeExtension));
+			FileReference oidcExe = new(Path.Combine(oidcDir.FullName, "OidcToken" + RuntimePlatform.ExeExtension));
 
 			OidcTokenExePath = oidcExe;
 		}
diff --git a/Engine/Source/Programs/Shared/EpicGames.ProjectStore/ProjectStoreTypes.cs b/Engine/Source/Programs/Shared/EpicGames.ProjectStore/ProjectStoreTypes.cs
--- a/Engine/Source/Programs/Shared/EpicGames.ProjectStore/ProjectStoreTypes.cs
+++ b/Engine/Source/Programs/Shared/EpicGames.ProjectStore/ProjectStoreTypes.cs
@@ -183,7 +183,7 @@ namespace EpicGames.ProjectStore
 		/// </summary>
 		public static JsonSerializerOptions GetDefaultJsonSerializerOptions()
 		{
-			JsonSerializerOptions options = new JsonSerializerOptions();
+			JsonSerializerOptions options = new();
 			options.AllowTrailingCommas = true;
 			options.ReadCommentHandling = JsonCommentHandling.Skip;
 			options.PropertyNameCaseInsensitive = true;
@@ -235,7 +235,7 @@ namespace EpicGames.ProjectStore
 		protected void CreateProject(FileReference projectFile, DirectoryReference rootDirectory, DirectoryReference engineDirectory)
 		{
 			string projectId = GetProjectID(projectFile);
-			StringBuilder projectCreateCommandLine = new StringBuilder();
+			StringBuilder projectCreateCommandLine = new();
 			projectCreateCommandLine.AppendFormat("project-create --force-update {0} {1} {2} {3} {4}", projectId, rootDirectory.FullName, engineDirectory.FullName, projectFile.Directory.FullName, projectFile.FullName);
 
 			RunApp(GetZenExecutable(engineDirectory).FullName, projectCreateCommandLine.ToString());
@@ -265,9 +265,9 @@ namespace EpicGames.ProjectStore
 				DirectoryReference.CreateDirectory(platformCookDirectory);
 			}
 
-			FileReference projectStoreFile = new FileReference(Path.Combine(platformCookDirectory.FullName, "ue.projectstore"));
+			FileReference projectStoreFile = new(Path.Combine(platformCookDirectory.FullName, "ue.projectstore"));
 
-			ProjectStoreData data = new ProjectStoreData();
+			ProjectStoreData data = new();
 			data.ZenServer = new ZenServerStoreData();
 			data.ZenServer.IsLocalHost = true;
 			data.ZenServer.HostName = "[::1]";
@@ -383,7 +383,7 @@ namespace EpicGames.ProjectStore
 			public List<SnapshotDescriptor> Snapshots { get; set; } = [];
 		}
 
-		private readonly SnapshotDescriptorCollection _snapshotDescriptors = new SnapshotDescriptorCollection();
+		private readonly SnapshotDescriptorCollection _snapshotDescriptors = new();
 		private readonly int _changelist;
 
 		private void LoadSnapshotDescriptors(string snapshotDescriptorJSON)
@@ -448,10 +448,10 @@ namespace EpicGames.ProjectStore
 
 				string oplogName = descriptor.TargetPlatform;
 
-				StringBuilder oplogCreateCommandLine = new StringBuilder();
+				StringBuilder oplogCreateCommandLine = new();
 				oplogCreateCommandLine.AppendFormat("oplog-create --force-update {0} {1}", projectId, oplogName);
 
-				StringBuilder oplogImportCommandLine = new StringBuilder();
+				StringBuilder oplogImportCommandLine = new();
 				oplogImportCommandLine.AppendFormat("oplog-import {0} {1} {2} --ignore-missing-attachments --clean", projectId, oplogName, projectStoreFile.FullName);
 
 				if (forceImport)
@@ -546,10 +546,10 @@ namespace EpicGames.ProjectStore
 
 			string oplogName = _buildData.Metadata.CookPlatform;
 
-			StringBuilder oplogCreateCommandLine = new StringBuilder();
+			StringBuilder oplogCreateCommandLine = new();
 			oplogCreateCommandLine.AppendFormat("oplog-create --force-update {0} {1}", projectId, oplogName);
 
-			StringBuilder oplogImportCommandLine = new StringBuilder();
+			StringBuilder oplogImportCommandLine = new();
 			oplogImportCommandLine.AppendFormat("oplog-import {0} {1} {2} --ignore-missing-attachments --clean", projectId, oplogName, projectStoreFile.FullName);
 
 			if (forceImport)
@@ -887,10 +887,10 @@ namespace EpicGames.ProjectStore
 			}
 
 			DirectoryReference oidcDir = DirectoryReference.Combine(Unreal.RootDirectory, "Engine", "Binaries", "DotNET", "OidcToken", binaryType);
-			FileReference oidcExe = new FileReference(Path.Combine(oidcDir.FullName, "OidcToken" + RuntimePlatform.ExeExtension));
-			FileReference outputFile = new FileReference(Path.Combine(Path.GetTempPath(), Path.GetTempFileName()));
+			FileReference oidcExe = new(Path.Combine(oidcDir.FullName, "OidcToken" + RuntimePlatform.ExeExtension));
+			FileReference outputFile = new(Path.Combine(Path.GetTempPath(), Path.GetTempFileName()));
 
-			StringBuilder commandLine = new StringBuilder();
+			StringBuilder commandLine = new();
 
 			commandLine.AppendFormat(" --Service={0} --OutFile={1}", cloudAuthServiceName, outputFile.FullName);
 			RunApp(oidcExe.FullName, commandLine.ToString());
@@ -951,12 +951,12 @@ namespace EpicGames.ProjectStore
 		/// <exception cref="BuildIndexException">Raised on request failure or unexpected response format.</exception>
 		public List<string> GetAvailablePlatforms(string branch, string project)
 		{
-			Uri uri = new Uri(String.Format("{0}/api/v2/builds/{1}", _host, _namespace));
+			Uri uri = new(String.Format("{0}/api/v2/builds/{1}", _host, _namespace));
 
 			HttpResponseMessage response;
-			using (HttpClient client = new HttpClient())
+			using (HttpClient client = new())
 			{
-				using (HttpRequestMessage request = new HttpRequestMessage())
+				using (HttpRequestMessage request = new())
 				{
 					request.Method = HttpMethod.Get;
 					request.RequestUri = uri;
@@ -967,7 +967,7 @@ namespace EpicGames.ProjectStore
 			}
 
 			string body;
-			using (StreamReader reader = new StreamReader(response.Content.ReadAsStream()))
+			using (StreamReader reader = new(response.Content.ReadAsStream()))
 			{
 				body = reader.ReadToEnd();
 			}
@@ -1000,12 +1000,12 @@ namespace EpicGames.ProjectStore
 		private List<UnrealCloudDDCBuildData> QueryBuilds(string branch, string project, string platform, string runtime, int changelist = 0)
 		{
 			string bucket = GetBucket(project, branch, platform, runtime);
-			Uri uri = new Uri(String.Format("{0}/api/v2/builds/{1}/{2}/search", _host, _namespace, bucket));
+			Uri uri = new(String.Format("{0}/api/v2/builds/{1}/{2}/search", _host, _namespace, bucket));
 
 			HttpResponseMessage response;
-			using (HttpClient client = new HttpClient())
+			using (HttpClient client = new())
 			{
-				using (HttpRequestMessage request = new HttpRequestMessage())
+				using (HttpRequestMessage request = new())
 				{
 					request.RequestUri = uri;
 					request.Method = HttpMethod.Post;
@@ -1026,7 +1026,7 @@ namespace EpicGames.ProjectStore
 			}
 
 			string body;
-			using (StreamReader reader = new StreamReader(response.Content.ReadAsStream()))
+			using (StreamReader reader = new(response.Content.ReadAsStream()))
 			{
 				body = reader.ReadToEnd();
 			}
diff --git a/Engine/Source/Programs/Shared/EpicGames.Redis.Tests/SerializationTests.cs b/Engine/Source/Programs/Shared/EpicGames.Redis.Tests/SerializationTests.cs
--- a/Engine/Source/Programs/Shared/EpicGames.Redis.Tests/SerializationTests.cs
+++ b/Engine/Source/Programs/Shared/EpicGames.Redis.Tests/SerializationTests.cs
@@ -19,7 +19,7 @@ namespace EpicGames.Redis.Tests
 		[TestMethod]
 		public void RoundTrip()
 		{
-			TestRecord input = new TestRecord(123, "hello", new Utf8String("world"));
+			TestRecord input = new(123, "hello", new Utf8String("world"));
 
 			RedisValue value = RedisSerializer.Serialize(input);
 
@@ -34,7 +34,7 @@ namespace EpicGames.Redis.Tests
 		[TestMethod]
 		public void EscapedCharacters()
 		{
-			TestRecord input = new TestRecord(123, "|||", new Utf8String("\\"));
+			TestRecord input = new(123, "|||", new Utf8String("\\"));
 
 			RedisValue value = RedisSerializer.Serialize(input);
 
@@ -68,9 +68,9 @@ namespace EpicGames.Redis.Tests
 		[TestMethod]
 		public void CustomStringConverter()
 		{
-			TestStringRecord input = new TestStringRecord(123, new TestStringType { Number = 0 });
+			TestStringRecord input = new(123, new TestStringType { Number = 0 });
 
-			RedisClassConverter<TestStringRecord> converter = new RedisClassConverter<TestStringRecord>();
+			RedisClassConverter<TestStringRecord> converter = new();
 			RedisValue value = converter.ToRedisValue(input);
 
 			Assert.AreEqual("123|hello", (string?)value);
@@ -102,9 +102,9 @@ namespace EpicGames.Redis.Tests
 		[TestMethod]
 		public void CustomUtf8StringConverter()
 		{
-			TestUtf8StringRecord input = new TestUtf8StringRecord(123, new TestUtf8StringType { Number = 0 });
+			TestUtf8StringRecord input = new(123, new TestUtf8StringType { Number = 0 });
 
-			RedisClassConverter<TestUtf8StringRecord> converter = new RedisClassConverter<TestUtf8StringRecord>();
+			RedisClassConverter<TestUtf8StringRecord> converter = new();
 			RedisValue value = converter.ToRedisValue(input);
 
 			Assert.AreEqual("123|hello", (string?)value);
@@ -124,7 +124,7 @@ namespace EpicGames.Redis.Tests
 		[TestMethod]
 		public void ProtobufTest()
 		{
-			ProtoClass cls = new ProtoClass();
+			ProtoClass cls = new();
 			cls.Text = "hello world";
 
 			RedisValue value = RedisSerializer.Serialize(cls);
diff --git a/Engine/Source/Programs/Shared/EpicGames.Redis/Converters/RedisClassConverter.cs b/Engine/Source/Programs/Shared/EpicGames.Redis/Converters/RedisClassConverter.cs
--- a/Engine/Source/Programs/Shared/EpicGames.Redis/Converters/RedisClassConverter.cs
+++ b/Engine/Source/Programs/Shared/EpicGames.Redis/Converters/RedisClassConverter.cs
@@ -138,7 +138,7 @@ namespace EpicGames.Redis.Converters
 		/// <inheritdoc/>
 		public RedisValue ToRedisValue(T value)
 		{
-			Utf8StringBuilder builder = new Utf8StringBuilder();
+			Utf8StringBuilder builder = new();
 			if (_properties.Length > 0)
 			{
 				AppendProperty(value, 0, builder);
diff --git a/Engine/Source/Programs/Shared/EpicGames.Redis/Converters/RedisProtobufConverter.cs b/Engine/Source/Programs/Shared/EpicGames.Redis/Converters/RedisProtobufConverter.cs
--- a/Engine/Source/Programs/Shared/EpicGames.Redis/Converters/RedisProtobufConverter.cs
+++ b/Engine/Source/Programs/Shared/EpicGames.Redis/Converters/RedisProtobufConverter.cs
@@ -47,7 +47,7 @@ namespace EpicGames.Redis.Converters
 		/// <inheritdoc/>
 		public RedisValue ToRedisValue(T value)
 		{
-			using (MemoryStream stream = new MemoryStream())
+			using (MemoryStream stream = new())
 			{
 				Serializer.Serialize(stream, value);
 				return stream.ToArray();
@@ -57,7 +57,7 @@ namespace EpicGames.Redis.Converters
 		/// <inheritdoc/>
 		public T FromRedisValue(RedisValue value)
 		{
-			using (ReadOnlyMemoryStream stream = new ReadOnlyMemoryStream(value))
+			using (ReadOnlyMemoryStream stream = new(value))
 			{
 				T result = Serializer.Deserialize<T>(stream);
 				return result;
diff --git a/Engine/Source/Programs/Shared/EpicGames.Redis/RedisChannel.cs b/Engine/Source/Programs/Shared/EpicGames.Redis/RedisChannel.cs
--- a/Engine/Source/Programs/Shared/EpicGames.Redis/RedisChannel.cs
+++ b/Engine/Source/Programs/Shared/EpicGames.Redis/RedisChannel.cs
@@ -64,19 +64,19 @@ namespace EpicGames.Redis
 		/// Create a channel name from a <see cref="Utf8String"/>.
 		/// </summary>
 		/// <param name="key">The string to get a channel from.</param>
-		public static implicit operator RedisChannel<T>(Utf8String key) => new RedisChannel<T>(key, RedisChannel.PatternMode.Auto);
+		public static implicit operator RedisChannel<T>(Utf8String key) => new(key, RedisChannel.PatternMode.Auto);
 
 		/// <summary>
 		/// Create a channel name from a byte array.
 		/// </summary>
 		/// <param name="key">The byte array to get a channel from.</param>
-		public static implicit operator RedisChannel<T>(string key) => new RedisChannel<T>(key, RedisChannel.PatternMode.Auto);
+		public static implicit operator RedisChannel<T>(string key) => new(key, RedisChannel.PatternMode.Auto);
 
 		/// <summary>
 		/// Create a channel name from a byte array.
 		/// </summary>
 		/// <param name="key">The byte array to get a channel from.</param>
-		public static implicit operator RedisChannel<T>(byte[] key) => new RedisChannel<T>(key, RedisChannel.PatternMode.Auto);
+		public static implicit operator RedisChannel<T>(byte[] key) => new(key, RedisChannel.PatternMode.Auto);
 	}
 
 	/// <summary>
diff --git a/Engine/Source/Programs/Shared/EpicGames.Redis/RedisHashKey.cs b/Engine/Source/Programs/Shared/EpicGames.Redis/RedisHashKey.cs
--- a/Engine/Source/Programs/Shared/EpicGames.Redis/RedisHashKey.cs
+++ b/Engine/Source/Programs/Shared/EpicGames.Redis/RedisHashKey.cs
@@ -20,7 +20,7 @@ namespace EpicGames.Redis
 		/// Implicit conversion to typed redis key.
 		/// </summary>
 		/// <param name="key">Key to convert</param>
-		public static implicit operator RedisHashKey<T>(string key) => new RedisHashKey<T>(new RedisKey(key));
+		public static implicit operator RedisHashKey<T>(string key) => new(new RedisKey(key));
 	}
 
 	/// <summary>
@@ -34,7 +34,7 @@ namespace EpicGames.Redis
 		/// Implicit conversion to typed redis key.
 		/// </summary>
 		/// <param name="key">Key to convert</param>
-		public static implicit operator RedisHashKey<TName, TValue>(string key) => new RedisHashKey<TName, TValue>(new RedisKey(key));
+		public static implicit operator RedisHashKey<TName, TValue>(string key) => new(new RedisKey(key));
 	}
 
 	/// <inheritdoc cref="HashEntry"/>
@@ -247,7 +247,7 @@ namespace EpicGames.Redis
 
 			HashRecordInfo<TRecord> recordInfo = HashRecordInfo<TRecord>.Instance;
 
-			TRecord record = new TRecord();
+			TRecord record = new();
 			for (int idx = 0; idx < selectors.Length; idx++)
 			{
 				recordInfo.SetProperty(record, names[idx], values[idx]);
@@ -291,7 +291,7 @@ namespace EpicGames.Redis
 		{
 			HashRecordInfo<TRecord> recordInfo = HashRecordInfo<TRecord>.Instance;
 
-			TRecord value = new TRecord();
+			TRecord value = new();
 
 			HashEntry[] entries = await target.HashGetAllAsync(key.Inner, flags);
 			foreach (HashEntry entry in entries)
diff --git a/Engine/Source/Programs/Shared/EpicGames.Redis/RedisListKey.cs b/Engine/Source/Programs/Shared/EpicGames.Redis/RedisListKey.cs
--- a/Engine/Source/Programs/Shared/EpicGames.Redis/RedisListKey.cs
+++ b/Engine/Source/Programs/Shared/EpicGames.Redis/RedisListKey.cs
@@ -17,7 +17,7 @@ namespace EpicGames.Redis
 		/// Implicit conversion to typed redis key.
 		/// </summary>
 		/// <param name="key">Key to convert</param>
-		public static implicit operator RedisListKey<TElement>(string key) => new RedisListKey<TElement>(new RedisKey(key));
+		public static implicit operator RedisListKey<TElement>(string key) => new(new RedisKey(key));
 	}
 
 	/// <summary>
diff --git a/Engine/Source/Programs/Shared/EpicGames.Redis/RedisSerializer.cs b/Engine/Source/Programs/Shared/EpicGames.Redis/RedisSerializer.cs
--- a/Engine/Source/Programs/Shared/EpicGames.Redis/RedisSerializer.cs
+++ b/Engine/Source/Programs/Shared/EpicGames.Redis/RedisSerializer.cs
@@ -81,7 +81,7 @@ namespace EpicGames.Redis
 	/// </summary>
 	public sealed class RedisJsonConverter<T> : IRedisConverter<T>
 	{
-		static readonly JsonSerializerOptions s_options = new JsonSerializerOptions { PropertyNameCaseInsensitive = true, PropertyNamingPolicy = JsonNamingPolicy.CamelCase };
+		static readonly JsonSerializerOptions s_options = new() { PropertyNameCaseInsensitive = true, PropertyNamingPolicy = JsonNamingPolicy.CamelCase };
 
 		/// <inheritdoc/>
 		public RedisValue ToRedisValue(T value)
@@ -184,7 +184,7 @@ namespace EpicGames.Redis
 				=> GetConverter<T>().ToRedisValue((T)value);
 		}
 
-		static readonly ConcurrentDictionary<Type, IRedisConverter<object>> s_typeToObjectConverter = new ConcurrentDictionary<Type, IRedisConverter<object>>();
+		static readonly ConcurrentDictionary<Type, IRedisConverter<object>> s_typeToObjectConverter = new();
 
 		/// <summary>
 		/// Register a custom converter for a particular type
diff --git a/Engine/Source/Programs/Shared/EpicGames.Redis/RedisSetKey.cs b/Engine/Source/Programs/Shared/EpicGames.Redis/RedisSetKey.cs
--- a/Engine/Source/Programs/Shared/EpicGames.Redis/RedisSetKey.cs
+++ b/Engine/Source/Programs/Shared/EpicGames.Redis/RedisSetKey.cs
@@ -17,7 +17,7 @@ namespace EpicGames.Redis
 		/// Implicit conversion to typed redis key.
 		/// </summary>
 		/// <param name="key">Key to convert</param>
-		public static implicit operator RedisSetKey<TElement>(string key) => new RedisSetKey<TElement>(new RedisKey(key));
+		public static implicit operator RedisSetKey<TElement>(string key) => new(new RedisKey(key));
 	}
 
 	/// <summary>
diff --git a/Engine/Source/Programs/Shared/EpicGames.Redis/RedisSortedSetKey.cs b/Engine/Source/Programs/Shared/EpicGames.Redis/RedisSortedSetKey.cs
--- a/Engine/Source/Programs/Shared/EpicGames.Redis/RedisSortedSetKey.cs
+++ b/Engine/Source/Programs/Shared/EpicGames.Redis/RedisSortedSetKey.cs
@@ -18,7 +18,7 @@ namespace EpicGames.Redis
 		/// Implicit conversion to typed redis key.
 		/// </summary>
 		/// <param name="key">Key to convert</param>
-		public static implicit operator RedisSortedSetKey<TElement>(string key) => new RedisSortedSetKey<TElement>(new RedisKey(key));
+		public static implicit operator RedisSortedSetKey<TElement>(string key) => new(new RedisKey(key));
 	}
 
 	/// <summary>
