Shared: prefer implicit object creation when type is apparent

diff --git a/Engine/Source/Programs/Shared/EpicGames.Redis/RedisStringKey.cs b/Engine/Source/Programs/Shared/EpicGames.Redis/RedisStringKey.cs
--- a/Engine/Source/Programs/Shared/EpicGames.Redis/RedisStringKey.cs
+++ b/Engine/Source/Programs/Shared/EpicGames.Redis/RedisStringKey.cs
@@ -18,7 +18,7 @@ namespace EpicGames.Redis
 		/// Implicit conversion to typed redis key.
 		/// </summary>
 		/// <param name="key">Key to convert</param>
-		public static implicit operator RedisStringKey<TValue>(string key) => new RedisStringKey<TValue>(new RedisKey(key));
+		public static implicit operator RedisStringKey<TValue>(string key) => new(new RedisKey(key));
 	}
 
 	/// <summary>
diff --git a/Engine/Source/Programs/Shared/EpicGames.Redis/RedisSubscription.cs b/Engine/Source/Programs/Shared/EpicGames.Redis/RedisSubscription.cs
--- a/Engine/Source/Programs/Shared/EpicGames.Redis/RedisSubscription.cs
+++ b/Engine/Source/Programs/Shared/EpicGames.Redis/RedisSubscription.cs
@@ -89,7 +89,7 @@ namespace EpicGames.Redis
 		/// <inheritdoc cref="ISubscriber.SubscribeAsync(RedisChannel, Action{RedisChannel, RedisValue}, CommandFlags)"/>
 		public static async Task<RedisSubscription> SubscribeAsync(this IConnectionMultiplexer connection, RedisChannel channel, Action<RedisChannel, RedisValue> handler, CommandFlags flags = CommandFlags.None)
 		{
-			RedisSubscription subscription = new RedisSubscription(channel, handler);
+			RedisSubscription subscription = new(channel, handler);
 			await subscription.SubscribeAsync(connection.GetSubscriber(), flags);
 			return subscription;
 		}
diff --git a/Engine/Source/Programs/Shared/EpicGames.Redis/Utility/RedisEvent.cs b/Engine/Source/Programs/Shared/EpicGames.Redis/Utility/RedisEvent.cs
--- a/Engine/Source/Programs/Shared/EpicGames.Redis/Utility/RedisEvent.cs
+++ b/Engine/Source/Programs/Shared/EpicGames.Redis/Utility/RedisEvent.cs
@@ -44,7 +44,7 @@ namespace EpicGames.Redis
 		/// <returns></returns>
 		public static async Task<RedisEvent> CreateAsync(IConnectionMultiplexer multiplexer, RedisChannel channel, CancellationToken cancellationToken = default)
 		{
-			AsyncEvent asyncEvent = new AsyncEvent();
+			AsyncEvent asyncEvent = new();
 
 			RedisSubscription subscription = await multiplexer.SubscribeAsync(channel, x => asyncEvent.Pulse());
 			if (cancellationToken.IsCancellationRequested)
diff --git a/Engine/Source/Programs/Shared/EpicGames.Serialization.Tests/CbPackageTests.cs b/Engine/Source/Programs/Shared/EpicGames.Serialization.Tests/CbPackageTests.cs
--- a/Engine/Source/Programs/Shared/EpicGames.Serialization.Tests/CbPackageTests.cs
+++ b/Engine/Source/Programs/Shared/EpicGames.Serialization.Tests/CbPackageTests.cs
@@ -32,7 +32,7 @@ namespace EpicGames.Serialization.Tests
 		[TestMethod]
 		public async Task RoundTripAsync()
 		{
-			TestObject root = new TestObject()
+			TestObject root = new()
 			{
 
 			};
@@ -43,7 +43,7 @@ namespace EpicGames.Serialization.Tests
 			builder.AddAttachment(rootHash, CbPackageAttachmentFlags.IsObject, rootObject.GetView().ToArray());
 			byte[] bytes = await builder.ToByteArrayAsync();
 
-			await using MemoryStream ms = new MemoryStream(bytes);
+			await using MemoryStream ms = new(bytes);
 			CbPackageReader reader = await CbPackageReader.CreateAsync(ms);
 
 			Assert.AreEqual(rootHash, reader.RootHash);
@@ -61,11 +61,11 @@ namespace EpicGames.Serialization.Tests
 			byte[] blob = Encoding.UTF8.GetBytes("This is blob contents");
 			IoHash blobHash = IoHash.Compute(blob);
 
-			SimpleObject simple = new SimpleObject();
+			SimpleObject simple = new();
 			CbObject simpleObject = CbSerializer.Serialize(simple);
 			IoHash simpleHash = IoHash.Compute(simpleObject.GetView().Span);
 
-			TestObject root = new TestObject()
+			TestObject root = new()
 			{
 				BinaryAttachment = blobHash,
 				ObjectAttachment = simpleHash,
@@ -80,7 +80,7 @@ namespace EpicGames.Serialization.Tests
 
 			byte[] bytes = await builder.ToByteArrayAsync();
 
-			await using MemoryStream ms = new MemoryStream(bytes);
+			await using MemoryStream ms = new(bytes);
 			CbPackageReader reader = await CbPackageReader.CreateAsync(ms);
 
 			Assert.AreEqual(rootHash, reader.RootHash);
@@ -101,7 +101,7 @@ namespace EpicGames.Serialization.Tests
 				{
 					Assert.AreEqual((ulong)simpleObject.GetView().Length, entry.PayloadSize);
 
-					CbObject roundTrippedObject = new CbObject(attachmentBytes);
+					CbObject roundTrippedObject = new(attachmentBytes);
 					Assert.AreEqual(simpleObject, roundTrippedObject);
 					continue;
 				}
@@ -114,14 +114,14 @@ namespace EpicGames.Serialization.Tests
 		[TestMethod]
 		public async Task ReadNoAttachmentsAsync()
 		{
-			using CbPackageBuilder packageBuilder = new CbPackageBuilder();
-			SimpleObject simple = new SimpleObject();
+			using CbPackageBuilder packageBuilder = new();
+			SimpleObject simple = new();
 			CbObject simpleObject = CbSerializer.Serialize(simple);
 			IoHash simpleHash = IoHash.Compute(simpleObject.GetView().Span);
 			packageBuilder.AddAttachment(simpleHash, CbPackageAttachmentFlags.IsObject, simpleObject.GetView().ToArray());
 
 			byte[] buf = await packageBuilder.ToByteArrayAsync();
-			await using MemoryStream ms = new MemoryStream(buf);
+			await using MemoryStream ms = new(buf);
 
 			try
 			{
diff --git a/Engine/Source/Programs/Shared/EpicGames.Serialization.Tests/CbWriterTests.cs b/Engine/Source/Programs/Shared/EpicGames.Serialization.Tests/CbWriterTests.cs
--- a/Engine/Source/Programs/Shared/EpicGames.Serialization.Tests/CbWriterTests.cs
+++ b/Engine/Source/Programs/Shared/EpicGames.Serialization.Tests/CbWriterTests.cs
@@ -15,7 +15,7 @@ namespace EpicGames.Serialization.Tests
 		[TestMethod]
 		public void EmptyObject()
 		{
-			CbWriter writer = new CbWriter();
+			CbWriter writer = new();
 			writer.BeginObject();
 			writer.EndObject();
 
@@ -28,7 +28,7 @@ namespace EpicGames.Serialization.Tests
 		[TestMethod]
 		public void EmptyArray()
 		{
-			CbWriter writer = new CbWriter();
+			CbWriter writer = new();
 			writer.BeginArray();
 			writer.EndArray();
 
@@ -42,7 +42,7 @@ namespace EpicGames.Serialization.Tests
 		[TestMethod]
 		public void Object()
 		{
-			CbWriter writer = new CbWriter();
+			CbWriter writer = new();
 			writer.BeginObject();
 			writer.WriteInteger("a", 1);
 			writer.WriteInteger("b", 2);
@@ -73,7 +73,7 @@ namespace EpicGames.Serialization.Tests
 		[TestMethod]
 		public void Array()
 		{
-			CbWriter writer = new CbWriter();
+			CbWriter writer = new();
 			writer.BeginArray();
 			writer.WriteIntegerValue(1);
 			writer.WriteIntegerValue(2);
@@ -96,7 +96,7 @@ namespace EpicGames.Serialization.Tests
 		[TestMethod]
 		public void UniformArray()
 		{
-			CbWriter writer = new CbWriter();
+			CbWriter writer = new();
 			writer.BeginUniformArray(CbFieldType.IntegerPositive);
 			writer.WriteIntegerValue(1);
 			writer.WriteIntegerValue(2);
@@ -117,7 +117,7 @@ namespace EpicGames.Serialization.Tests
 		[TestMethod]
 		public void NestedArray()
 		{
-			CbWriter writer = new CbWriter();
+			CbWriter writer = new();
 			writer.BeginObject();
 			writer.BeginArray("a");
 			writer.WriteIntegerValue(1);
@@ -144,7 +144,7 @@ namespace EpicGames.Serialization.Tests
 		{
 			byte[] test = [1, 2, 3, 4];
 
-			CbWriter writer = new CbWriter();
+			CbWriter writer = new();
 			writer.BeginObject();
 			writer.WriteBinaryReference("a", test);
 			writer.EndObject();
@@ -168,7 +168,7 @@ namespace EpicGames.Serialization.Tests
 		[TestMethod]
 		public void UniformArrayOfArrays()
 		{
-			CbWriter writer = new CbWriter();
+			CbWriter writer = new();
 			writer.BeginUniformArray(CbFieldType.Array);
 			writer.BeginArray();
 			writer.WriteNullValue();
@@ -203,7 +203,7 @@ namespace EpicGames.Serialization.Tests
 			List<IoHash> objHashes = [.. Enumerable.Range(0, 5).Select(x => IoHash.Compute(Encoding.UTF8.GetBytes($"obj{x}")))];
 			List<IoHash> binHashes = [.. Enumerable.Range(0, 5).Select(x => IoHash.Compute(Encoding.UTF8.GetBytes($"bin{x}")))];
 
-			CbWriter writer = new CbWriter();
+			CbWriter writer = new();
 			writer.BeginObject();
 			writer.BeginUniformArray("obj", CbFieldType.ObjectAttachment);
 			for (int idx = 0; idx < 5; idx++)
@@ -232,7 +232,7 @@ namespace EpicGames.Serialization.Tests
 			// Create an empty CbObject
 			CbObject obj;
 			{
-				CbWriter writer = new CbWriter();
+				CbWriter writer = new();
 				writer.BeginObject();
 				writer.EndObject();
 				obj = writer.ToObject();
@@ -240,7 +240,7 @@ namespace EpicGames.Serialization.Tests
 
 			// Embed it in another writer at the root
 			{
-				CbWriter writer = new CbWriter();
+				CbWriter writer = new();
 				writer.WriteObject(obj);
 
 				byte[] data = writer.ToByteArray();
diff --git a/Engine/Source/Programs/Shared/EpicGames.Serialization.Tests/EngineTests.cs b/Engine/Source/Programs/Shared/EpicGames.Serialization.Tests/EngineTests.cs
--- a/Engine/Source/Programs/Shared/EpicGames.Serialization.Tests/EngineTests.cs
+++ b/Engine/Source/Programs/Shared/EpicGames.Serialization.Tests/EngineTests.cs
@@ -34,7 +34,7 @@ namespace EpicGames.Serialization.Tests
 			}
 		}
 
-		readonly Dictionary<CbFieldType, CbFieldAccessors> _typeAccessors = new Dictionary<CbFieldType, CbFieldAccessors>
+		readonly Dictionary<CbFieldType, CbFieldAccessors> _typeAccessors = new()
 		{
 			[CbFieldType.Object] = new CbFieldAccessors(CbObject.Empty, x => x.IsObject(), x => x.AsObject(), (x, y) => x.AsObject()),
 			[CbFieldType.UniformObject] = new CbFieldAccessors(CbObject.Empty, x => x.IsObject(), x => x.AsObject(), (x, y) => x.AsObject()),
@@ -78,7 +78,7 @@ namespace EpicGames.Serialization.Tests
 
 		void TestField(CbFieldType fieldType, byte[] payload, object? expectedValue = null, object? defaultValue = null, CbFieldError expectedError = CbFieldError.None, CbFieldAccessors? accessors = null)
 		{
-			CbField field = new CbField(payload, fieldType);
+			CbField field = new(payload, fieldType);
 			Assert.AreEqual(field.GetSize(), payload.Length + (CbFieldUtils.HasFieldType(fieldType) ? 0 : 1));
 			Assert.IsTrue(field.HasValue());
 			Assert.IsFalse(field.HasError());
@@ -93,7 +93,7 @@ namespace EpicGames.Serialization.Tests
 
 		void TestFieldError(CbFieldType fieldType, ReadOnlyMemory<byte> payload, CbFieldError expectedError, object? expectedValue = null, CbFieldAccessors? accessors = null)
 		{
-			CbField field = new CbField(payload, fieldType);
+			CbField field = new(payload, fieldType);
 			TestFieldError(fieldType, field, expectedError, expectedValue, accessors);
 		}
 
@@ -102,7 +102,7 @@ namespace EpicGames.Serialization.Tests
 		{
 			// Test CbField()
 			{
-				CbField defaultField = new CbField();
+				CbField defaultField = new();
 				Assert.IsFalse(defaultField.HasName());
 				Assert.IsFalse(defaultField.HasValue());
 				Assert.IsFalse(defaultField.HasError());
@@ -119,7 +119,7 @@ namespace EpicGames.Serialization.Tests
 
 			// Test CbField(None)
 			{
-				CbField noneField = new CbField(ReadOnlyMemory<byte>.Empty, CbFieldType.None);
+				CbField noneField = new(ReadOnlyMemory<byte>.Empty, CbFieldType.None);
 				Assert.AreEqual(noneField.GetSize(), 1);
 				Assert.AreEqual(noneField.GetName().Length, 0);
 				Assert.IsFalse(noneField.HasName());
@@ -134,7 +134,7 @@ namespace EpicGames.Serialization.Tests
 			{
 				CbFieldType fieldType = CbFieldType.None | CbFieldType.HasFieldName;
 				byte[] noneBytes = [(byte)fieldType, 4, (byte)'N', (byte)'a', (byte)'m', (byte)'e'];
-				CbField noneField = new CbField(noneBytes);
+				CbField noneField = new(noneBytes);
 				Assert.AreEqual(noneField.GetSize(), noneBytes.Length);
 				Assert.AreEqual(noneField.GetName(), new Utf8String("Name"));
 				Assert.IsTrue(noneField.HasName());
@@ -152,7 +152,7 @@ namespace EpicGames.Serialization.Tests
 			{
 				CbFieldType fieldType = CbFieldType.None;
 				byte[] noneBytes = [(byte)fieldType];
-				CbField noneField = new CbField(noneBytes);
+				CbField noneField = new(noneBytes);
 				Assert.AreEqual(noneField.GetSize(), noneBytes.Length);
 				Assert.AreEqual(noneField.GetName().Length, 0);
 				Assert.IsFalse(noneField.HasName());
@@ -166,7 +166,7 @@ namespace EpicGames.Serialization.Tests
 			{
 				CbFieldType fieldType = CbFieldType.None | CbFieldType.HasFieldName;
 				byte[] noneBytes = [(byte)fieldType, 4, (byte)'N', (byte)'a', (byte)'m', (byte)'e'];
-				CbField noneField = new CbField(noneBytes.AsMemory(1), fieldType);
+				CbField noneField = new(noneBytes.AsMemory(1), fieldType);
 				Assert.AreEqual(noneField.GetSize(), noneBytes.Length);
 				Assert.AreEqual(noneField.GetName(), new Utf8String("Name"));
 				Assert.IsTrue(noneField.HasName());
@@ -183,7 +183,7 @@ namespace EpicGames.Serialization.Tests
 			{
 				CbFieldType fieldType = CbFieldType.None | CbFieldType.HasFieldName;
 				byte[] noneBytes = [(byte)fieldType, 0];
-				CbField noneField = new CbField(noneBytes.AsMemory(1), fieldType);
+				CbField noneField = new(noneBytes.AsMemory(1), fieldType);
 				Assert.AreEqual(noneBytes.Length, noneField.GetSize());
 				Assert.AreEqual(Utf8String.Empty, noneField.GetName());
 				Assert.IsTrue(noneField.HasName());
@@ -198,7 +198,7 @@ namespace EpicGames.Serialization.Tests
 		{
 			// Test CbField(Null)
 			{
-				CbField nullField = new CbField(ReadOnlyMemory<byte>.Empty, CbFieldType.Null);
+				CbField nullField = new(ReadOnlyMemory<byte>.Empty, CbFieldType.Null);
 				Assert.AreEqual(nullField.GetSize(), 1);
 				Assert.IsTrue(nullField.IsNull());
 				Assert.IsTrue(nullField.HasValue());
@@ -209,7 +209,7 @@ namespace EpicGames.Serialization.Tests
 
 			// Test CbField(None) as Null
 			{
-				CbField field = new CbField();
+				CbField field = new();
 				Assert.IsFalse(field.IsNull());
 			}
 		}
@@ -261,7 +261,7 @@ namespace EpicGames.Serialization.Tests
 			{
 				byte intType = (byte)(CbFieldType.HasFieldName | CbFieldType.IntegerPositive);
 				byte[] payload = [12, intType, 1, (byte)'A', 1, intType, 1, (byte)'B', 2, intType, 1, (byte)'C', 3];
-				CbField field = new CbField(payload, CbFieldType.Object);
+				CbField field = new(payload, CbFieldType.Object);
 				TestField(CbFieldType.Object, field, new CbObject(payload, CbFieldType.Object));
 				CbObject @object = CbObject.Clone(field.AsObject());
 				TestIntObject(@object, 3, payload.Length);
@@ -279,7 +279,7 @@ namespace EpicGames.Serialization.Tests
 			{
 				byte intType = (byte)(CbFieldType.HasFieldName | CbFieldType.IntegerPositive);
 				byte[] payload = [10, intType, 1, (byte)'A', 1, 1, (byte)'B', 2, 1, (byte)'C', 3];
-				CbField field = new CbField(payload, CbFieldType.UniformObject);
+				CbField field = new(payload, CbFieldType.UniformObject);
 				TestField(CbFieldType.UniformObject, field, new CbObject(payload, CbFieldType.UniformObject));
 				CbObject @object = CbObject.Clone(field.AsObject());
 				TestIntObject(@object, 3, payload.Length);
@@ -298,7 +298,7 @@ namespace EpicGames.Serialization.Tests
 
 				// Equals
 				byte[] namedPayload = [1, (byte)'O', 10, intType, 1, (byte)'A', 1, 1, (byte)'B', 2, 1, (byte)'C', 3];
-				CbField namedField = new CbField(namedPayload, CbFieldType.UniformObject | CbFieldType.HasFieldName);
+				CbField namedField = new(namedPayload, CbFieldType.UniformObject | CbFieldType.HasFieldName);
 				Assert.IsTrue(field.AsObject().Equals(namedField.AsObject()));
 
 				// CopyTo
@@ -324,7 +324,7 @@ namespace EpicGames.Serialization.Tests
 			{
 				byte objectType = (byte)(CbFieldType.Object | CbFieldType.HasFieldName);
 				byte[] buffer = [objectType, 3, (byte)'K', (byte)'e', (byte)'y', 4, (byte)(CbFieldType.HasFieldName | CbFieldType.IntegerPositive), 1, (byte)'F', 8];
-				CbObject @object = new CbObject(buffer);
+				CbObject @object = new(buffer);
 				Assert.AreEqual(@object.GetSize(), 6);
 				CbObject objectClone = CbObject.Clone(@object);
 				Assert.AreEqual(objectClone.GetSize(), 6);
@@ -336,7 +336,7 @@ namespace EpicGames.Serialization.Tests
 					Assert.AreEqual(field.GetName(), new Utf8String("F"));
 					Assert.AreEqual(field.AsInt32(), 8);
 				}
-				for (CbFieldIterator it = objectClone.CreateIterator(), end = new CbFieldIterator(); it != end; ++it)
+				for (CbFieldIterator it = objectClone.CreateIterator(), end = new(); it != end; ++it)
 				{
 				}
 				foreach (CbField _ in objectClone)
@@ -395,7 +395,7 @@ namespace EpicGames.Serialization.Tests
 
 			// Test CbField(Array, Empty)
 			{
-				CbArray array = new CbArray();
+				CbArray array = new();
 				TestIntArray(array, 0, 2);
 
 				// Advance an iterator past the last field.
@@ -415,7 +415,7 @@ namespace EpicGames.Serialization.Tests
 			{
 				byte intType = (byte)CbFieldType.IntegerPositive;
 				byte[] payload = [7, 3, intType, 1, intType, 2, intType, 3];
-				CbField field = new CbField(payload, CbFieldType.Array);
+				CbField field = new(payload, CbFieldType.Array);
 				TestField(CbFieldType.Array, field, new CbArray(payload, CbFieldType.Array));
 				CbArray array = field.AsArray();
 				TestIntArray(array, 3, payload.Length);
@@ -427,7 +427,7 @@ namespace EpicGames.Serialization.Tests
 			{
 				byte intType = (byte)(CbFieldType.IntegerPositive);
 				byte[] payload = [5, 3, intType, 1, 2, 3];
-				CbField field = new CbField(payload, CbFieldType.UniformArray);
+				CbField field = new(payload, CbFieldType.UniformArray);
 				TestField(CbFieldType.UniformArray, field, new CbArray(payload, CbFieldType.UniformArray));
 				CbArray array = field.AsArray();
 				TestIntArray(array, 3, payload.Length);
@@ -438,7 +438,7 @@ namespace EpicGames.Serialization.Tests
 
 				// Equals
 				byte[] namedPayload = [1, (byte)'A', 5, 3, intType, 1, 2, 3];
-				CbField namedField = new CbField(namedPayload, CbFieldType.UniformArray | CbFieldType.HasFieldName);
+				CbField namedField = new(namedPayload, CbFieldType.UniformArray | CbFieldType.HasFieldName);
 				Assert.IsTrue(field.AsArray().Equals(namedField.AsArray()));
 				Assert.IsTrue(field.Equals(field.AsArray().AsField()));
 				Assert.IsTrue(namedField.Equals(namedField.AsArray().AsField()));
@@ -463,7 +463,7 @@ namespace EpicGames.Serialization.Tests
 
 			// Test CbField(None) as Array
 			{
-				CbField field = new CbField();
+				CbField field = new();
 				//				TestFieldError(CbFieldType.Array, Field, CbFieldError.TypeError);
 				field.AsArray();
 			}
@@ -472,7 +472,7 @@ namespace EpicGames.Serialization.Tests
 			{
 				byte arrayType = (byte)(CbFieldType.Array | CbFieldType.HasFieldName);
 				byte[] buffer = [arrayType, 3, (byte)'K', (byte)'e', (byte)'y', 3, 1, (byte)(CbFieldType.IntegerPositive), 8];
-				CbArray array = new CbArray(buffer);
+				CbArray array = new(buffer);
 				Assert.AreEqual(array.GetSize(), 5);
 				CbArray arrayClone = array;
 				Assert.AreEqual(arrayClone.GetSize(), 5);
@@ -484,7 +484,7 @@ namespace EpicGames.Serialization.Tests
 					Assert.AreEqual(field.AsInt32(), 8);
 					//					Assert.IsTrue(Field.IsOwned());
 				}
-				for (CbFieldIterator it = arrayClone.CreateIterator(), end = new CbFieldIterator(); it != end; ++it)
+				for (CbFieldIterator it = arrayClone.CreateIterator(), end = new(); it != end; ++it)
 				{
 				}
 				foreach (CbField _ in arrayClone)
@@ -507,7 +507,7 @@ namespace EpicGames.Serialization.Tests
 			// Test CbArray as CbFieldIterator
 			{
 				uint count = 0;
-				CbArray array = new CbArray();
+				CbArray array = new();
 				for (CbFieldIterator iter = CbFieldIterator.MakeSingle(array.AsField()); iter; ++iter)
 				{
 					CbField field = iter.Current;
@@ -520,7 +520,7 @@ namespace EpicGames.Serialization.Tests
 			// Test CbArray as CbFieldIterator
 			{
 				uint count = 0;
-				CbArray array = new CbArray();
+				CbArray array = new();
 				//				Array.MakeOwned();
 				for (CbFieldIterator iter = CbFieldIterator.MakeSingle(array.AsField()); iter; ++iter)
 				{
@@ -541,7 +541,7 @@ namespace EpicGames.Serialization.Tests
 			// Test CbField(Binary, Value)
 			{
 				byte[] payload = [3, 4, 5, 6]; // Size: 3, Data: 4/5/6
-				CbField fieldView = new CbField(payload, CbFieldType.Binary);
+				CbField fieldView = new(payload, CbFieldType.Binary);
 				TestField(CbFieldType.Binary, fieldView, (ReadOnlyMemory<byte>)payload.AsMemory(1, 3));
 
 				CbField field = fieldView;
@@ -553,7 +553,7 @@ namespace EpicGames.Serialization.Tests
 
 			// Test CbField(None) as Binary
 			{
-				CbField fieldView = new CbField();
+				CbField fieldView = new();
 				byte[] @default = [1, 2, 3];
 				TestFieldError(CbFieldType.Binary, fieldView, CbFieldError.TypeError, (ReadOnlyMemory<byte>)@default);
 
@@ -583,7 +583,7 @@ namespace EpicGames.Serialization.Tests
 
 			// Test CbField(None) as String
 			{
-				CbField field = new CbField();
+				CbField field = new();
 				TestFieldError(CbFieldType.String, field, CbFieldError.TypeError, new Utf8String("ABC"));
 			}
 		}
@@ -623,7 +623,7 @@ namespace EpicGames.Serialization.Tests
 			VarInt.WriteUnsigned(payload, magnitude - negative);
 			ulong defaultValue = 8;
 			ulong expectedValue = (negative != 0) ? (ulong)(-(long)(magnitude)) : magnitude;
-			CbField field = new CbField(payload, fieldType);
+			CbField field = new(payload, fieldType);
 
 			TestField(CbFieldType.IntegerNegative, field, (sbyte)(((expectedMask & EIntType.Int8) != 0) ? expectedValue : defaultValue),
 				(sbyte)(defaultValue), ((expectedMask & EIntType.Int8) != 0) ? CbFieldError.None : CbFieldError.RangeError, CbFieldAccessors.FromStruct<sbyte>(x => x.IsInteger(), (x, y) => x.AsInt8(y)));
@@ -684,7 +684,7 @@ namespace EpicGames.Serialization.Tests
 
 			// Test CbField(None) as Integer
 			{
-				CbField field = new CbField();
+				CbField field = new();
 				TestFieldError(CbFieldType.IntegerPositive, field, CbFieldError.TypeError, (ulong)(8));
 				TestFieldError(CbFieldType.IntegerNegative, field, CbFieldError.TypeError, (long)(8));
 			}
@@ -698,7 +698,7 @@ namespace EpicGames.Serialization.Tests
 				byte[] payload = [0xc0, 0x12, 0x34, 0x56]; // -2.28444433f
 				TestField(CbFieldType.Float32, payload, -2.28444433f);
 
-				CbField field = new CbField(payload, CbFieldType.Float32);
+				CbField field = new(payload, CbFieldType.Float32);
 				TestField(CbFieldType.Float64, field, (double)-2.28444433f);
 			}
 
@@ -707,7 +707,7 @@ namespace EpicGames.Serialization.Tests
 				byte[] payload = [0xc1, 0x23, 0x45, 0x67, 0x89, 0xab, 0xcd, 0xef]; // -631475.76888888876
 				TestField(CbFieldType.Float64, payload, -631475.76888888876);
 
-				CbField field = new CbField(payload, CbFieldType.Float64);
+				CbField field = new(payload, CbFieldType.Float64);
 				TestFieldError(CbFieldType.Float32, field, CbFieldError.RangeError, 8.0f);
 			}
 
@@ -715,7 +715,7 @@ namespace EpicGames.Serialization.Tests
 			{
 				byte[] payload = new byte[9];
 				VarInt.WriteUnsigned(payload, ((ulong)(1) << 24) - 1); // 16,777,215
-				CbField field = new CbField(payload, CbFieldType.IntegerPositive);
+				CbField field = new(payload, CbFieldType.IntegerPositive);
 				TestField(CbFieldType.Float32, field, 16_777_215.0f);
 				TestField(CbFieldType.Float64, field, 16_777_215.0);
 			}
@@ -724,7 +724,7 @@ namespace EpicGames.Serialization.Tests
 			{
 				byte[] payload = new byte[9];
 				VarInt.WriteUnsigned(payload, (ulong)(1) << 24); // 16,777,216
-				CbField field = new CbField(payload, CbFieldType.IntegerPositive);
+				CbField field = new(payload, CbFieldType.IntegerPositive);
 				TestFieldError(CbFieldType.Float32, field, CbFieldError.RangeError, 8.0f);
 				TestField(CbFieldType.Float64, field, 16_777_216.0);
 			}
@@ -733,7 +733,7 @@ namespace EpicGames.Serialization.Tests
 			{
 				byte[] payload = new byte[9];
 				VarInt.WriteUnsigned(payload, ((ulong)(1) << 53) - 1); // 9,007,199,254,740,991
-				CbField field = new CbField(payload, CbFieldType.IntegerPositive);
+				CbField field = new(payload, CbFieldType.IntegerPositive);
 				TestFieldError(CbFieldType.Float32, field, CbFieldError.RangeError, 8.0f);
 				TestField(CbFieldType.Float64, field, 9_007_199_254_740_991.0);
 			}
@@ -742,7 +742,7 @@ namespace EpicGames.Serialization.Tests
 			{
 				byte[] payload = new byte[9];
 				VarInt.WriteUnsigned(payload, (ulong)(1) << 53); // 9,007,199,254,740,992
-				CbField field = new CbField(payload, CbFieldType.IntegerPositive);
+				CbField field = new(payload, CbFieldType.IntegerPositive);
 				TestFieldError(CbFieldType.Float32, field, CbFieldError.RangeError, 8.0f);
 				TestFieldError(CbFieldType.Float64, field, CbFieldError.RangeError, 8.0);
 			}
@@ -751,7 +751,7 @@ namespace EpicGames.Serialization.Tests
 			{
 				byte[] payload = new byte[9];
 				VarInt.WriteUnsigned(payload, ~(ulong)0); // Max uint64
-				CbField field = new CbField(payload, CbFieldType.IntegerPositive);
+				CbField field = new(payload, CbFieldType.IntegerPositive);
 				TestFieldError(CbFieldType.Float32, field, CbFieldError.RangeError, 8.0f);
 				TestFieldError(CbFieldType.Float64, field, CbFieldError.RangeError, 8.0);
 			}
@@ -760,7 +760,7 @@ namespace EpicGames.Serialization.Tests
 			{
 				byte[] payload = new byte[9];
 				VarInt.WriteUnsigned(payload, ((ulong)(1) << 24) - 2); // -16,777,215
-				CbField field = new CbField(payload, CbFieldType.IntegerNegative);
+				CbField field = new(payload, CbFieldType.IntegerNegative);
 				TestField(CbFieldType.Float32, field, -16_777_215.0f);
 				TestField(CbFieldType.Float64, field, -16_777_215.0);
 			}
@@ -769,7 +769,7 @@ namespace EpicGames.Serialization.Tests
 			{
 				byte[] payload = new byte[9];
 				VarInt.WriteUnsigned(payload, ((ulong)(1) << 24) - 1); // -16,777,216
-				CbField field = new CbField(payload, CbFieldType.IntegerNegative);
+				CbField field = new(payload, CbFieldType.IntegerNegative);
 				TestFieldError(CbFieldType.Float32, field, CbFieldError.RangeError, 8.0f);
 				TestField(CbFieldType.Float64, field, -16_777_216.0);
 			}
@@ -778,7 +778,7 @@ namespace EpicGames.Serialization.Tests
 			{
 				byte[] payload = new byte[9];
 				VarInt.WriteUnsigned(payload, ((ulong)(1) << 53) - 2); // -9,007,199,254,740,991
-				CbField field = new CbField(payload, CbFieldType.IntegerNegative);
+				CbField field = new(payload, CbFieldType.IntegerNegative);
 				TestFieldError(CbFieldType.Float32, field, CbFieldError.RangeError, 8.0f);
 				TestField(CbFieldType.Float64, field, -9_007_199_254_740_991.0);
 			}
@@ -787,14 +787,14 @@ namespace EpicGames.Serialization.Tests
 			{
 				byte[] payload = new byte[9];
 				VarInt.WriteUnsigned(payload, ((ulong)(1) << 53) - 1); // -9,007,199,254,740,992
-				CbField field = new CbField(payload, CbFieldType.IntegerNegative);
+				CbField field = new(payload, CbFieldType.IntegerNegative);
 				TestFieldError(CbFieldType.Float32, field, CbFieldError.RangeError, 8.0f);
 				TestFieldError(CbFieldType.Float64, field, CbFieldError.RangeError, 8.0);
 			}
 
 			// Test CbField(None) as Float
 			{
-				CbField field = new CbField();
+				CbField field = new();
 				TestFieldError(CbFieldType.Float32, field, CbFieldError.TypeError, 8.0f);
 				TestFieldError(CbFieldType.Float64, field, CbFieldError.TypeError, 8.0);
 			}
@@ -811,7 +811,7 @@ namespace EpicGames.Serialization.Tests
 
 			// Test CbField(None) as Bool
 			{
-				CbField defaultField = new CbField();
+				CbField defaultField = new();
 				TestFieldError(CbFieldType.BoolFalse, defaultField, CbFieldError.TypeError, false);
 				TestFieldError(CbFieldType.BoolTrue, defaultField, CbFieldError.TypeError, true);
 			}
@@ -831,13 +831,13 @@ namespace EpicGames.Serialization.Tests
 
 			// Test CbField(ObjectAttachment, NonZero) AsAttachment
 			{
-				CbField field = new CbField(sequentialBytes, CbFieldType.ObjectAttachment);
+				CbField field = new(sequentialBytes, CbFieldType.ObjectAttachment);
 				TestField(CbFieldType.ObjectAttachment, field, (CbObjectAttachment)new IoHash(sequentialBytes), (CbObjectAttachment)new IoHash(), CbFieldError.None);
 			}
 
 			// Test CbField(None) as ObjectAttachment
 			{
-				CbField defaultField = new CbField();
+				CbField defaultField = new();
 				TestFieldError(CbFieldType.ObjectAttachment, defaultField, CbFieldError.TypeError, (CbObjectAttachment)new IoHash(sequentialBytes));
 			}
 		}
@@ -856,13 +856,13 @@ namespace EpicGames.Serialization.Tests
 
 			// Test CbField(BinaryAttachment, NonZero) AsAttachment
 			{
-				CbField field = new CbField(sequentialBytes, CbFieldType.BinaryAttachment);
+				CbField field = new(sequentialBytes, CbFieldType.BinaryAttachment);
 				TestField(CbFieldType.BinaryAttachment, field, (CbBinaryAttachment)new IoHash(sequentialBytes), (CbBinaryAttachment)new IoHash(), CbFieldError.None);
 			}
 
 			// Test CbField(None) as BinaryAttachment
 			{
-				CbField defaultField = new CbField();
+				CbField defaultField = new();
 				TestFieldError(CbFieldType.BinaryAttachment, defaultField, CbFieldError.TypeError, (CbBinaryAttachment)new IoHash(sequentialBytes));
 			}
 		}
@@ -881,19 +881,19 @@ namespace EpicGames.Serialization.Tests
 
 			// Test CbField(None) as Hash
 			{
-				CbField defaultField = new CbField();
+				CbField defaultField = new();
 				TestFieldError(CbFieldType.Hash, defaultField, CbFieldError.TypeError, new IoHash(sequentialBytes));
 			}
 
 			// Test CbField(ObjectAttachment) as Hash
 			{
-				CbField field = new CbField(sequentialBytes, CbFieldType.ObjectAttachment);
+				CbField field = new(sequentialBytes, CbFieldType.ObjectAttachment);
 				TestField(CbFieldType.Hash, field, new IoHash(sequentialBytes));
 			}
 
 			// Test CbField(BinaryAttachment) as Hash
 			{
-				CbField field = new CbField(sequentialBytes, CbFieldType.BinaryAttachment);
+				CbField field = new(sequentialBytes, CbFieldType.BinaryAttachment);
 				TestField(CbFieldType.Hash, field, new IoHash(sequentialBytes));
 			}
 		}
@@ -913,7 +913,7 @@ namespace EpicGames.Serialization.Tests
 
 			// Test CbField(None) as Uuid
 			{
-				CbField defaultField = new CbField();
+				CbField defaultField = new();
 				TestFieldError(CbFieldType.Uuid, defaultField, CbFieldError.TypeError, Guid.NewGuid());
 			}
 		}
@@ -930,15 +930,15 @@ namespace EpicGames.Serialization.Tests
 			// Test CbField(DateTime, Zero) as FDateTime
 			{
 				byte[] payload = [0, 0, 0, 0, 0, 0, 0, 0];
-				CbField field = new CbField(payload, CbFieldType.DateTime);
+				CbField field = new(payload, CbFieldType.DateTime);
 				Assert.AreEqual(field.AsDateTime(), new DateTime(0));
 			}
 
 			// Test CbField(None) as DateTime
 			{
-				CbField defaultField = new CbField();
+				CbField defaultField = new();
 				TestFieldError(CbFieldType.DateTime, defaultField, CbFieldError.TypeError);
-				DateTime defaultValue = new DateTime(0x1020_3040_5060_7080L, DateTimeKind.Utc);
+				DateTime defaultValue = new(0x1020_3040_5060_7080L, DateTimeKind.Utc);
 				Assert.AreEqual(defaultField.AsDateTime(defaultValue), defaultValue);
 			}
 		}
@@ -955,15 +955,15 @@ namespace EpicGames.Serialization.Tests
 			// Test CbField(TimeSpan, Zero) as FTimeSpan
 			{
 				byte[] payload = [0, 0, 0, 0, 0, 0, 0, 0];
-				CbField field = new CbField(payload, CbFieldType.TimeSpan);
+				CbField field = new(payload, CbFieldType.TimeSpan);
 				Assert.AreEqual(field.AsTimeSpan(), new TimeSpan(0));
 			}
 
 			// Test CbField(None) as TimeSpan
 			{
-				CbField defaultField = new CbField();
+				CbField defaultField = new();
 				TestFieldError(CbFieldType.TimeSpan, defaultField, CbFieldError.TypeError);
-				TimeSpan defaultValue = new TimeSpan(0x1020_3040_5060_7080L);
+				TimeSpan defaultValue = new(0x1020_3040_5060_7080L);
 				Assert.AreEqual(defaultField.AsTimeSpan(defaultValue), defaultValue);
 			}
 		}
diff --git a/Engine/Source/Programs/Shared/EpicGames.Serialization.Tests/JupiterTests.cs b/Engine/Source/Programs/Shared/EpicGames.Serialization.Tests/JupiterTests.cs
--- a/Engine/Source/Programs/Shared/EpicGames.Serialization.Tests/JupiterTests.cs
+++ b/Engine/Source/Programs/Shared/EpicGames.Serialization.Tests/JupiterTests.cs
@@ -15,7 +15,7 @@ namespace EpicGames.Serialization.Tests
 		[TestMethod]
 		public void OverflowTest()
 		{
-			CbWriter writer = new CbWriter();
+			CbWriter writer = new();
 
 			writer.BeginArray();
 			for (int idx = 0; idx < 128; idx++)
@@ -30,7 +30,7 @@ namespace EpicGames.Serialization.Tests
 		[TestMethod]
 		public void EmbeddedObject()
 		{
-			CbWriter writer1 = new CbWriter();
+			CbWriter writer1 = new();
 			writer1.BeginObject();
 			writer1.WriteInteger("a", 1);
 			writer1.WriteUtf8String("b", new Utf8String("hello"));
@@ -38,7 +38,7 @@ namespace EpicGames.Serialization.Tests
 
 			CbObject object1 = writer1.ToObject();
 
-			CbWriter writer2 = new CbWriter();
+			CbWriter writer2 = new();
 			writer2.BeginObject();
 			writer2.WriteObject("ref", object1);
 			writer2.EndObject();
@@ -52,8 +52,8 @@ namespace EpicGames.Serialization.Tests
 		{
 			byte[] bytes = File.ReadAllBytes("CompactBinaryObjects/build");
 
-			ReadOnlyMemory<byte> memory = new ReadOnlyMemory<byte>(bytes);
-			CbField o = new CbField(memory);
+			ReadOnlyMemory<byte> memory = new(bytes);
+			CbField o = new(memory);
 
 			Assert.AreEqual(new Utf8String("BuildAction"), o.Name);
 			List<CbField> buildActionFields = [.. o];
@@ -82,8 +82,8 @@ namespace EpicGames.Serialization.Tests
 		{
 			byte[] bytes = File.ReadAllBytes("CompactBinaryObjects/ReferenceOutput");
 
-			ReadOnlyMemory<byte> memory = new ReadOnlyMemory<byte>(bytes);
-			CbField o = new CbField(memory);
+			ReadOnlyMemory<byte> memory = new(bytes);
+			CbField o = new(memory);
 
 			Assert.AreEqual(new Utf8String("BuildOutput"), o.Name);
 			List<CbField> buildActionFields = [.. o];
@@ -102,7 +102,7 @@ namespace EpicGames.Serialization.Tests
 		{
 			byte[] bytes = File.ReadAllBytes("CompactBinaryObjects/compact_binary");
 
-			CbObject o = new CbObject(bytes);
+			CbObject o = new(bytes);
 			Assert.AreEqual(Utf8String.Empty, o.AsField().Name);
 			List<CbField> buildActionFields = [.. o];
 			Assert.AreEqual(3, buildActionFields.Count);
@@ -119,7 +119,7 @@ namespace EpicGames.Serialization.Tests
 			IoHash hash1 = IoHash.Parse("5d8a6dc277c968f0d027c98f879c955c1905c293");
 			IoHash hash2 = IoHash.Parse("313f0d0d334100d83aeb1ee2c42794fd087cb0ae");
 
-			CbWriter writer = new CbWriter();
+			CbWriter writer = new();
 			writer.BeginObject();
 			writer.BeginUniformArray("needs", CbFieldType.Hash);
 			writer.WriteHashValue(hash1);
@@ -128,8 +128,8 @@ namespace EpicGames.Serialization.Tests
 			writer.EndObject();
 
 			byte[] objectData = writer.ToByteArray();
-			ReadOnlyMemory<byte> memory = new ReadOnlyMemory<byte>(objectData);
-			CbField o = new CbField(memory);
+			ReadOnlyMemory<byte> memory = new(objectData);
+			CbField o = new(memory);
 
 			// the top object has no name
 			Assert.AreEqual(Utf8String.Empty, o.Name);
@@ -146,15 +146,15 @@ namespace EpicGames.Serialization.Tests
 		{
 			IoHash hash1 = IoHash.Parse("5d8a6dc277c968f0d027c98f879c955c1905c293");
 
-			CbWriter writer = new CbWriter();
+			CbWriter writer = new();
 			writer.BeginObject();
 			writer.WriteUtf8String("string", new Utf8String("test"));
 			writer.WriteBinaryAttachment("hash", hash1);
 			writer.EndObject();
 
 			byte[] objectData = writer.ToByteArray();
-			ReadOnlyMemory<byte> memory = new ReadOnlyMemory<byte>(objectData);
-			CbField o = new CbField(memory);
+			ReadOnlyMemory<byte> memory = new(objectData);
+			CbField o = new(memory);
 
 			// the object has no name and 2 fields
 			Assert.AreEqual(Utf8String.Empty, o.Name);
diff --git a/Engine/Source/Programs/Shared/EpicGames.Serialization.Tests/SerializerTests.cs b/Engine/Source/Programs/Shared/EpicGames.Serialization.Tests/SerializerTests.cs
--- a/Engine/Source/Programs/Shared/EpicGames.Serialization.Tests/SerializerTests.cs
+++ b/Engine/Source/Programs/Shared/EpicGames.Serialization.Tests/SerializerTests.cs
@@ -68,7 +68,7 @@ namespace EpicGames.Serialization.Tests
 		{
 			TestObject obj1 = CreateTestObject(new Random(123));
 
-			CbWriter writer = new CbWriter();
+			CbWriter writer = new();
 			CbSerializer.Serialize(writer, obj1);
 
 			byte[] data1 = writer.ToByteArray();
@@ -91,7 +91,7 @@ namespace EpicGames.Serialization.Tests
 
 		static TestObject CreateTestObject(Random random)
 		{
-			TestObject obj = new TestObject();
+			TestObject obj = new();
 			obj.A = random.Next();
 			obj.B = random.NextInt64().ToString();
 			obj.C = IoHash.Compute(Encoding.UTF8.GetBytes(obj.B));
@@ -119,7 +119,7 @@ namespace EpicGames.Serialization.Tests
 			Assert.ThrowsExactly<CbEmptyClassException>(() => CbSerializer.Serialize(new EmptyObject()));
 			Assert.ThrowsExactly<CbEmptyClassException>(() => CbSerializer.Deserialize<EmptyObject>(CbField.Empty));
 
-			EmptyObjectWithAttribute obj2 = new EmptyObjectWithAttribute();
+			EmptyObjectWithAttribute obj2 = new();
 			CbObject cbObj2 = CbSerializer.Serialize(obj2);
 			CbSerializer.Deserialize<EmptyObjectWithAttribute>(cbObj2);
 		}
@@ -136,7 +136,7 @@ namespace EpicGames.Serialization.Tests
 		[TestMethod]
 		public void NullFieldTest()
 		{
-			ObjectWithNullField input = new ObjectWithNullField { Value = 123 };
+			ObjectWithNullField input = new() { Value = 123 };
 			CbObject obj = CbSerializer.Serialize(input);
 
 			ObjectWithNullField output = CbSerializer.Deserialize<ObjectWithNullField>(obj.AsField());
@@ -159,7 +159,7 @@ namespace EpicGames.Serialization.Tests
 		[TestMethod]
 		public void ObjectWithReadOnlyCollectionTest()
 		{
-			ObjectWithReadOnlyCollection input = new ObjectWithReadOnlyCollection();
+			ObjectWithReadOnlyCollection input = new();
 			input.Strings.Add("Hello");
 			input.Strings.Add("World");
 			input.Integers.Add(1);
@@ -193,7 +193,7 @@ namespace EpicGames.Serialization.Tests
 		[TestMethod]
 		public void ObjectWithStringDictCollectionTest()
 		{
-			ObjectWithStringDictCollection input = new ObjectWithStringDictCollection();
+			ObjectWithStringDictCollection input = new();
 			input.Dictionary.Add("hello", 123);
 			input.Dictionary.Add("world", 456);
 			CbObject obj = CbSerializer.Serialize(input);
@@ -204,7 +204,7 @@ namespace EpicGames.Serialization.Tests
 			Assert.AreEqual(456, output.Dictionary["world"]);
 
 			// make sure string keyed dictionaries are serialized as objects
-			CbWriter writer = new CbWriter();
+			CbWriter writer = new();
 			writer.BeginObject();
 			writer.BeginObject("d");
 			writer.WriteInteger("hello", 123);
@@ -228,7 +228,7 @@ namespace EpicGames.Serialization.Tests
 		[TestMethod]
 		public void ObjectWithByteArrayTest()
 		{
-			ObjectWithByteArray input = new ObjectWithByteArray();
+			ObjectWithByteArray input = new();
 			byte[] b = Encoding.ASCII.GetBytes("test-string");
 			input.Bytes = b;
 			CbObject obj = CbSerializer.Serialize(input);
diff --git a/Engine/Source/Programs/Shared/EpicGames.Serialization/CbConverter.cs b/Engine/Source/Programs/Shared/EpicGames.Serialization/CbConverter.cs
--- a/Engine/Source/Programs/Shared/EpicGames.Serialization/CbConverter.cs
+++ b/Engine/Source/Programs/Shared/EpicGames.Serialization/CbConverter.cs
@@ -201,7 +201,7 @@ namespace EpicGames.Serialization
 		/// <summary>
 		/// Object used for locking access to shared objects
 		/// </summary>
-		static readonly object s_lockObject = new object();
+		static readonly object s_lockObject = new();
 
 		/// <summary>
 		/// Cache of property to converter
diff --git a/Engine/Source/Programs/Shared/EpicGames.Serialization/CbField.cs b/Engine/Source/Programs/Shared/EpicGames.Serialization/CbField.cs
--- a/Engine/Source/Programs/Shared/EpicGames.Serialization/CbField.cs
+++ b/Engine/Source/Programs/Shared/EpicGames.Serialization/CbField.cs
@@ -230,7 +230,7 @@ namespace EpicGames.Serialization
 		/// Convert a hash to a binary attachment 
 		/// </summary>
 		/// <param name="hash">The attachment to convert</param>
-		public static implicit operator CbBinaryAttachment(IoHash hash) => new CbBinaryAttachment(hash);
+		public static implicit operator CbBinaryAttachment(IoHash hash) => new(hash);
 
 		/// <summary>
 		/// Use a binary attachment as a hash
@@ -304,7 +304,7 @@ namespace EpicGames.Serialization
 		/// Use an object attachment as a hash
 		/// </summary>
 		/// <param name="hash">The attachment to convert</param>
-		public static implicit operator CbObjectAttachment(IoHash hash) => new CbObjectAttachment(hash);
+		public static implicit operator CbObjectAttachment(IoHash hash) => new(hash);
 
 		/// <summary>
 		/// Use an object attachment as a hash
@@ -768,7 +768,7 @@ namespace EpicGames.Serialization
 		/// <summary>
 		/// Returns the name of the field if it has a name, otherwise an empty view.
 		/// </summary>
-		public Utf8String Name => new Utf8String(Memory.Slice(_payloadOffset - _nameLen, _nameLen));
+		public Utf8String Name => new(Memory.Slice(_payloadOffset - _nameLen, _nameLen));
 
 		/// <summary>
 		/// Gets the value of this field
@@ -1387,7 +1387,7 @@ namespace EpicGames.Serialization
 		/// </summary>
 		/// <param name="defaultValue"></param>
 		/// <returns></returns>
-		public TimeSpan AsTimeSpan(TimeSpan defaultValue = default) => new TimeSpan(AsTimeSpanTicks(defaultValue.Ticks));
+		public TimeSpan AsTimeSpan(TimeSpan defaultValue = default) => new(AsTimeSpanTicks(defaultValue.Ticks));
 
 		/// <summary>
 		/// Access the field as a object id
@@ -1943,7 +1943,7 @@ namespace EpicGames.Serialization
 		/// <inheritdoc/>
 		public override CbObject Read(ref Utf8JsonReader reader, Type typeToConvert, JsonSerializerOptions options)
 		{
-			CbWriter writer = new CbWriter();
+			CbWriter writer = new();
 
 			using JsonDocument doc = JsonDocument.ParseValue(ref reader);
 			CbJsonReader.ReadJsonField(doc.RootElement, null, writer);
@@ -2122,7 +2122,7 @@ namespace EpicGames.Serialization
 			// Always iterate over non-uniform ranges because we do not know if they contain an attachment.
 			if (CbFieldUtils.HasFieldType(Current.GetTypeWithFlags()))
 			{
-				for (CbFieldIterator it = new CbFieldIterator(this); it; ++it)
+				for (CbFieldIterator it = new(this); it; ++it)
 				{
 					if (CbFieldUtils.MayContainAttachments(it.Current.GetTypeWithFlags()))
 					{
@@ -2135,7 +2135,7 @@ namespace EpicGames.Serialization
 			{
 				if (CbFieldUtils.MayContainAttachments(Current.GetTypeWithFlags()))
 				{
-					for (CbFieldIterator it = new CbFieldIterator(this); it; ++it)
+					for (CbFieldIterator it = new(this); it; ++it)
 					{
 						it.Current.IterateAttachments(visitor);
 					}
@@ -2286,7 +2286,7 @@ namespace EpicGames.Serialization
 		/// </summary>
 		/// <param name="field"></param>
 		/// <returns></returns>
-		public static CbArray FromFieldNoCheck(CbField field) => new CbArray(field);
+		public static CbArray FromFieldNoCheck(CbField field) => new(field);
 
 		/// <summary>
 		/// Returns the size of the array in bytes if serialized by itself with no name.
@@ -2459,7 +2459,7 @@ namespace EpicGames.Serialization
 		/// <returns></returns>
 		public static CbObject Build(Action<CbWriter> build)
 		{
-			CbWriter writer = new CbWriter();
+			CbWriter writer = new();
 			writer.BeginObject();
 			build(writer);
 			writer.EndObject();
@@ -2503,7 +2503,7 @@ namespace EpicGames.Serialization
 		/// </summary>
 		/// <param name="field"></param>
 		/// <returns></returns>
-		public static CbObject FromFieldNoCheck(CbField field) => new CbObject(field);
+		public static CbObject FromFieldNoCheck(CbField field) => new(field);
 
 		/// <summary>
 		/// Returns the size of the object in bytes if serialized by itself with no name.
@@ -2637,8 +2637,8 @@ namespace EpicGames.Serialization
 		/// <returns></returns>
 		public string ToJson()
 		{
-			ArrayBufferWriter<byte> buffer = new ArrayBufferWriter<byte>();
-			using (Utf8JsonWriter jsonWriter = new Utf8JsonWriter(buffer))
+			ArrayBufferWriter<byte> buffer = new();
+			using (Utf8JsonWriter jsonWriter = new(buffer))
 			{
 				ToJson(jsonWriter, JsonSerializerOptions.Default);
 			}
diff --git a/Engine/Source/Programs/Shared/EpicGames.Serialization/CbPackage.cs b/Engine/Source/Programs/Shared/EpicGames.Serialization/CbPackage.cs
--- a/Engine/Source/Programs/Shared/EpicGames.Serialization/CbPackage.cs
+++ b/Engine/Source/Programs/Shared/EpicGames.Serialization/CbPackage.cs
@@ -52,8 +52,8 @@ namespace EpicGames.Serialization
 				throw new Exception($"Was not able to read package header, not enough bytes. Bytes read: {readBytes}");
 			}
 
-			using MemoryStream ms = new MemoryStream(buf);
-			using BinaryReader reader = new BinaryReader(ms);
+			using MemoryStream ms = new(buf);
+			using BinaryReader reader = new(ms);
 			uint headerMagic = reader.ReadUInt32();
 			if (headerMagic != CbPackageConstants.CbPackageHeaderMagic)
 			{
@@ -69,7 +69,7 @@ namespace EpicGames.Serialization
 
 		public readonly void Write(Stream stream)
 		{
-			using BinaryWriter writer = new BinaryWriter(stream, Encoding.ASCII, leaveOpen: true);
+			using BinaryWriter writer = new(stream, Encoding.ASCII, leaveOpen: true);
 			writer.Write(HeaderMagic);
 			writer.Write(AttachmentCount);
 			writer.Write(Reserved1);
@@ -132,8 +132,8 @@ namespace EpicGames.Serialization
 			{
 				throw new Exception($"Was not able to read package attachment entry, not enough bytes. Bytes read: {readBytes}");
 			}
-			using MemoryStream ms = new MemoryStream(buf);
-			using BinaryReader reader = new BinaryReader(ms);
+			using MemoryStream ms = new(buf);
+			using BinaryReader reader = new(ms);
 			ulong payloadSize = reader.ReadUInt64();
 
 			uint flags = reader.ReadUInt32();
@@ -155,7 +155,7 @@ namespace EpicGames.Serialization
 		/// <param name="stream"></param>
 		internal readonly void Write(Stream stream)
 		{
-			using BinaryWriter writer = new BinaryWriter(stream, Encoding.ASCII, leaveOpen: true);
+			using BinaryWriter writer = new(stream, Encoding.ASCII, leaveOpen: true);
 			writer.Write(PayloadSize);
 			writer.Write((uint)Flags);
 			writer.Write(AttachmentHash.ToByteArray());
@@ -314,7 +314,7 @@ namespace EpicGames.Serialization
 		/// <returns></returns>
 		public void AddAttachment(IoHash attachmentHash, CbPackageAttachmentFlags flags, byte[] blobMemory)
 		{
-			CbPackageAttachmentEntry entry = new CbPackageAttachmentEntry()
+			CbPackageAttachmentEntry entry = new()
 			{
 				AttachmentHash = attachmentHash,
 				Flags = flags,
@@ -334,7 +334,7 @@ namespace EpicGames.Serialization
 		/// <returns></returns>
 		public void AddAttachment(IoHash attachmentHash, CbPackageAttachmentFlags flags, Stream stream, ulong length)
 		{
-			CbPackageAttachmentEntry entry = new CbPackageAttachmentEntry()
+			CbPackageAttachmentEntry entry = new()
 			{
 				AttachmentHash = attachmentHash,
 				Flags = flags,
@@ -350,7 +350,7 @@ namespace EpicGames.Serialization
 		/// <returns></returns>
 		public async Task<byte[]> ToByteArrayAsync()
 		{
-			MemoryStream packageBuffer = new MemoryStream();
+			MemoryStream packageBuffer = new();
 
 			if (_streamAttachments.Count == 0)
 			{
@@ -358,7 +358,7 @@ namespace EpicGames.Serialization
 			}
 
 			// we will overwrite this header again after building the package builder
-			CbPackageHeader header = new CbPackageHeader
+			CbPackageHeader header = new()
 			{
 				HeaderMagic = CbPackageConstants.CbPackageHeaderMagic,
 				AttachmentCount = (uint)_streamAttachments.Count - 1, // the root object does not count as a attachment
diff --git a/Engine/Source/Programs/Shared/EpicGames.Serialization/CbSerializer.cs b/Engine/Source/Programs/Shared/EpicGames.Serialization/CbSerializer.cs
--- a/Engine/Source/Programs/Shared/EpicGames.Serialization/CbSerializer.cs
+++ b/Engine/Source/Programs/Shared/EpicGames.Serialization/CbSerializer.cs
@@ -115,7 +115,7 @@ namespace EpicGames.Serialization
 		/// <returns></returns>
 		public static CbObject Serialize(Type type, object value)
 		{
-			CbWriter writer = new CbWriter();
+			CbWriter writer = new();
 			CbConverter.GetConverter(type).WriteObject(writer, value);
 			return writer.ToObject();
 		}
@@ -128,7 +128,7 @@ namespace EpicGames.Serialization
 		/// <returns></returns>
 		public static CbObject Serialize<T>(T value)
 		{
-			CbWriter writer = new CbWriter();
+			CbWriter writer = new();
 			CbConverter.GetConverter<T>().Write(writer, value);
 			return writer.ToObject();
 		}
@@ -141,7 +141,7 @@ namespace EpicGames.Serialization
 		/// <returns></returns>
 		public static byte[] SerializeToByteArray<T>(T value)
 		{
-			CbWriter writer = new CbWriter();
+			CbWriter writer = new();
 			CbConverter.GetConverter<T>().Write(writer, value);
 			return writer.ToByteArray();
 		}
@@ -220,7 +220,7 @@ namespace EpicGames.Serialization
 		public static CbObject FromJson(string json)
 		{
 			using JsonDocument document = JsonDocument.Parse(json);
-			CbWriter writer = new CbWriter();
+			CbWriter writer = new();
 			ReadJsonField(document.RootElement, null, writer);
 
 			return writer.ToObject();
diff --git a/Engine/Source/Programs/Shared/EpicGames.Serialization/CbWriter.cs b/Engine/Source/Programs/Shared/EpicGames.Serialization/CbWriter.cs
--- a/Engine/Source/Programs/Shared/EpicGames.Serialization/CbWriter.cs
+++ b/Engine/Source/Programs/Shared/EpicGames.Serialization/CbWriter.cs
@@ -86,12 +86,12 @@ namespace EpicGames.Serialization
 		/// <summary>
 		/// Implicit conversion from a string
 		/// </summary>
-		public static implicit operator CbFieldName(string text) => new CbFieldName(new Utf8String(text));
+		public static implicit operator CbFieldName(string text) => new(new Utf8String(text));
 
 		/// <summary>
 		/// Implicit conversion from a utf8 string
 		/// </summary>
-		public static implicit operator CbFieldName(Utf8String text) => new CbFieldName(text);
+		public static implicit operator CbFieldName(Utf8String text) => new(text);
 	}
 
 	/// <summary>
@@ -144,9 +144,9 @@ namespace EpicGames.Serialization
 			}
 		}
 
-		readonly Scope _rootScope = new Scope { _fieldType = CbFieldType.Array };
-		readonly Stack<Scope> _openScopes = new Stack<Scope>();
-		readonly Stack<Scope> _freeScopes = new Stack<Scope>();
+		readonly Scope _rootScope = new() { _fieldType = CbFieldType.Array };
+		readonly Stack<Scope> _openScopes = new();
+		readonly Stack<Scope> _freeScopes = new();
 
 		Memory<byte> _buffer = Memory<byte>.Empty;
 		int _bufferPos = 0; // Offset of the first field in the current buffer
diff --git a/Engine/Source/Programs/Shared/EpicGames.Serialization/CompressedBuffer.cs b/Engine/Source/Programs/Shared/EpicGames.Serialization/CompressedBuffer.cs
--- a/Engine/Source/Programs/Shared/EpicGames.Serialization/CompressedBuffer.cs
+++ b/Engine/Source/Programs/Shared/EpicGames.Serialization/CompressedBuffer.cs
@@ -117,13 +117,13 @@ namespace EpicGames.Serialization
 		{
 			byte[] headerData = br.ReadBytes((int)CompressedBufferHeader.HeaderLength);
 
-			using MemoryStream ms = new MemoryStream(headerData);
-			using BinaryReader reader = new BinaryReader(ms);
+			using MemoryStream ms = new(headerData);
+			using BinaryReader reader = new(ms);
 
 			// the header is always stored big endian
 			bool needsByteSwap = BitConverter.IsLittleEndian;
 
-			CompressedBufferHeader header = new CompressedBufferHeader
+			CompressedBufferHeader header = new()
 			{
 				Magic = reader.ReadUInt32(),
 				Crc32 = reader.ReadUInt32(),
@@ -169,7 +169,7 @@ namespace EpicGames.Serialization
 
 				for (int i = 0; i < header.BlockCount; i++)
 				{
-					ReadOnlySpan<byte> memory = new ReadOnlySpan<byte>(blocksData, i * sizeof(uint), sizeof(uint));
+					ReadOnlySpan<byte> memory = new(blocksData, i * sizeof(uint), sizeof(uint));
 					uint compressedBlockSize = BinaryPrimitives.ReadUInt32BigEndian(memory);
 					blocks[i] = compressedBlockSize;
 				}
@@ -227,7 +227,7 @@ namespace EpicGames.Serialization
 		/// <exception cref="Exception"></exception>
 		public static async Task<CompressedBufferHeader> DecompressContentAsync(Stream sourceStream, ulong streamSize, Stream targetStream, CancellationToken cancellationToken = default)
 		{
-			using BinaryReader br = new BinaryReader(sourceStream);
+			using BinaryReader br = new(sourceStream);
 			(CompressedBufferHeader header, uint[] compressedBlockSizes) = ExtractHeader(br);
 
 			if (streamSize != header.TotalCompressedSize)
@@ -306,7 +306,7 @@ namespace EpicGames.Serialization
 			const long DefaultBlockSize = 256 * 1024;
 			long blockSize = DefaultBlockSize;
 			long blockCount = (rawContents.LongLength + blockSize - 1) / blockSize;
-			Span<byte> contentsSpan = new Span<byte>(rawContents);
+			Span<byte> contentsSpan = new(rawContents);
 			List<byte[]> blocks = [];
 
 			for (int i = 0; i < blockCount; i++)
@@ -392,9 +392,9 @@ namespace EpicGames.Serialization
 
 			Hash blake3Hash = hasher.Finalize();
 			byte[] hashData = blake3Hash.AsSpan().Slice(0, 20).ToArray();
-			IoHash hash = new IoHash(hashData);
+			IoHash hash = new(hashData);
 
-			CompressedBufferHeader header = new CompressedBufferHeader
+			CompressedBufferHeader header = new()
 			{
 				Magic = CompressedBufferHeader.ExpectedMagic,
 				Crc32 = 0,
@@ -410,7 +410,7 @@ namespace EpicGames.Serialization
 
 			byte[] headerAndBlocks = WriteHeaderToBuffer(header, [.. compressedBlocks.Select(b => (uint)b.Length)]);
 
-			using BinaryWriter writer = new BinaryWriter(s, Encoding.Default, leaveOpen: true);
+			using BinaryWriter writer = new(s, Encoding.Default, leaveOpen: true);
 
 			writer.Write(headerAndBlocks);
 
@@ -431,8 +431,8 @@ namespace EpicGames.Serialization
 
 			// write the compressed buffer, but with the wrong crc which we update and rewrite later
 			{
-				using MemoryStream ms = new MemoryStream(headerBuffer);
-				using BinaryWriter writer = new BinaryWriter(ms);
+				using MemoryStream ms = new(headerBuffer);
+				using BinaryWriter writer = new(ms);
 
 				WriteHeader(header, writer);
 
@@ -454,8 +454,8 @@ namespace EpicGames.Serialization
 
 			// write the header again now that we have the crc
 			{
-				using MemoryStream ms = new MemoryStream(headerBuffer);
-				using BinaryWriter writer = new BinaryWriter(ms);
+				using MemoryStream ms = new(headerBuffer);
+				using BinaryWriter writer = new(ms);
 				WriteHeader(header, writer);
 			}
 
diff --git a/Engine/Source/Programs/Shared/EpicGames.Serialization/Converters/CbClassConverter.cs b/Engine/Source/Programs/Shared/EpicGames.Serialization/Converters/CbClassConverter.cs
--- a/Engine/Source/Programs/Shared/EpicGames.Serialization/Converters/CbClassConverter.cs
+++ b/Engine/Source/Programs/Shared/EpicGames.Serialization/Converters/CbClassConverter.cs
@@ -24,7 +24,7 @@ namespace EpicGames.Serialization.Converters
 			}
 		}
 
-		static readonly Utf8String s_discriminatorKey = new Utf8String("_t");
+		static readonly Utf8String s_discriminatorKey = new("_t");
 
 		public Type ClassType { get; }
 		public bool IsPolymorphic { get; }
@@ -141,7 +141,7 @@ namespace EpicGames.Serialization.Converters
 				CbClassConverterMethods methods = Create(knownType);
 
 				{
-					DynamicMethod dynamicMethod = new DynamicMethod("_", typeof(object), [typeof(CbField)]);
+					DynamicMethod dynamicMethod = new("_", typeof(object), [typeof(CbField)]);
 					ILGenerator generator = dynamicMethod.GetILGenerator();
 					generator.Emit(OpCodes.Ldarg_0);
 					generator.Emit(OpCodes.Call, methods.ReadConcreteMethod);
@@ -151,7 +151,7 @@ namespace EpicGames.Serialization.Converters
 				}
 
 				{
-					DynamicMethod dynamicMethod = new DynamicMethod("_", null, [typeof(CbWriter), typeof(object)]);
+					DynamicMethod dynamicMethod = new("_", null, [typeof(CbWriter), typeof(object)]);
 					ILGenerator generator = dynamicMethod.GetILGenerator();
 					generator.Emit(OpCodes.Ldarg_0);
 					generator.Emit(OpCodes.Ldarg_1);
@@ -469,7 +469,7 @@ namespace EpicGames.Serialization.Converters
 					CbFieldAttribute? attribute = property.GetCustomAttribute<CbFieldAttribute>();
 					if (attribute != null || (property.GetGetMethod()?.IsPublic ?? false))
 					{
-						Utf8String name = new Utf8String(attribute?.Name ?? property.Name);
+						Utf8String name = new(attribute?.Name ?? property.Name);
 						propertyList.Add((name, property));
 					}
 				}
diff --git a/Engine/Source/Programs/Shared/EpicGames.Slack/Blocks/ActionsBlock.cs b/Engine/Source/Programs/Shared/EpicGames.Slack/Blocks/ActionsBlock.cs
--- a/Engine/Source/Programs/Shared/EpicGames.Slack/Blocks/ActionsBlock.cs
+++ b/Engine/Source/Programs/Shared/EpicGames.Slack/Blocks/ActionsBlock.cs
@@ -36,7 +36,7 @@ namespace EpicGames.Slack.Blocks
 		/// <param name="container">Block container</param>
 		public static ActionsBlock AddActions(this ISlackBlockContainer container)
 		{
-			ActionsBlock block = new ActionsBlock();
+			ActionsBlock block = new();
 			container.Blocks.Add(block);
 			return block;
 		}
@@ -48,7 +48,7 @@ namespace EpicGames.Slack.Blocks
 		/// <param name="configure">Configuration function for the block</param>
 		public static void AddActions(this ISlackBlockContainer container, Action<ActionsBlock> configure)
 		{
-			ActionsBlock block = new ActionsBlock();
+			ActionsBlock block = new();
 			configure(block);
 			container.Blocks.Add(block);
 		}
diff --git a/Engine/Source/Programs/Shared/EpicGames.Slack/Blocks/ContextBlock.cs b/Engine/Source/Programs/Shared/EpicGames.Slack/Blocks/ContextBlock.cs
--- a/Engine/Source/Programs/Shared/EpicGames.Slack/Blocks/ContextBlock.cs
+++ b/Engine/Source/Programs/Shared/EpicGames.Slack/Blocks/ContextBlock.cs
@@ -38,7 +38,7 @@ namespace EpicGames.Slack.Blocks
 		/// <param name="configure">Configuration callback for the block</param>
 		public static void AddContext(this ISlackBlockContainer container, Action<ContextBlock> configure)
 		{
-			ContextBlock block = new ContextBlock();
+			ContextBlock block = new();
 			configure(block);
 			container.Blocks.Add(block);
 		}
@@ -50,7 +50,7 @@ namespace EpicGames.Slack.Blocks
 		/// <param name="text">Configuration callback for the block</param>
 		public static ContextBlock AddContext(this ISlackBlockContainer container, TextObject text)
 		{
-			ContextBlock block = new ContextBlock();
+			ContextBlock block = new();
 			block.Elements.Add(text);
 			container.Blocks.Add(block);
 			return block;
diff --git a/Engine/Source/Programs/Shared/EpicGames.Slack/Blocks/DividerBlock.cs b/Engine/Source/Programs/Shared/EpicGames.Slack/Blocks/DividerBlock.cs
--- a/Engine/Source/Programs/Shared/EpicGames.Slack/Blocks/DividerBlock.cs
+++ b/Engine/Source/Programs/Shared/EpicGames.Slack/Blocks/DividerBlock.cs
@@ -26,7 +26,7 @@ namespace EpicGames.Slack.Blocks
 		/// <param name="container">Block container</param>
 		public static void AddDivider(this ISlackBlockContainer container)
 		{
-			DividerBlock block = new DividerBlock();
+			DividerBlock block = new();
 			container.Blocks.Add(block);
 		}
 	}
diff --git a/Engine/Source/Programs/Shared/EpicGames.Slack/Blocks/HeaderBlock.cs b/Engine/Source/Programs/Shared/EpicGames.Slack/Blocks/HeaderBlock.cs
--- a/Engine/Source/Programs/Shared/EpicGames.Slack/Blocks/HeaderBlock.cs
+++ b/Engine/Source/Programs/Shared/EpicGames.Slack/Blocks/HeaderBlock.cs
@@ -46,7 +46,7 @@ namespace EpicGames.Slack.Blocks
 		/// <param name="text">Text for the header</param>
 		public static void AddHeader(this ISlackBlockContainer container, PlainTextObject text)
 		{
-			HeaderBlock block = new HeaderBlock(text);
+			HeaderBlock block = new(text);
 			container.Blocks.Add(block);
 		}
 
diff --git a/Engine/Source/Programs/Shared/EpicGames.Slack/Blocks/ImageBlock.cs b/Engine/Source/Programs/Shared/EpicGames.Slack/Blocks/ImageBlock.cs
--- a/Engine/Source/Programs/Shared/EpicGames.Slack/Blocks/ImageBlock.cs
+++ b/Engine/Source/Programs/Shared/EpicGames.Slack/Blocks/ImageBlock.cs
@@ -50,7 +50,7 @@ namespace EpicGames.Slack.Blocks
 		/// </summary>
 		public static void AddImage(this ISlackBlockContainer container, Uri imageUrl, string altText, PlainTextObject? title = null)
 		{
-			ImageBlock block = new ImageBlock(imageUrl, altText, title);
+			ImageBlock block = new(imageUrl, altText, title);
 			container.Blocks.Add(block);
 		}
 	}
diff --git a/Engine/Source/Programs/Shared/EpicGames.Slack/Blocks/InputBlock.cs b/Engine/Source/Programs/Shared/EpicGames.Slack/Blocks/InputBlock.cs
--- a/Engine/Source/Programs/Shared/EpicGames.Slack/Blocks/InputBlock.cs
+++ b/Engine/Source/Programs/Shared/EpicGames.Slack/Blocks/InputBlock.cs
@@ -60,7 +60,7 @@ namespace EpicGames.Slack.Blocks
 		/// </summary>
 		public static InputBlock AddInput(this ISlackBlockContainer container, PlainTextObject label, Element element)
 		{
-			InputBlock block = new InputBlock(label, element);
+			InputBlock block = new(label, element);
 			container.Blocks.Add(block);
 			return block;
 		}
diff --git a/Engine/Source/Programs/Shared/EpicGames.Slack/Elements/ButtonElement.cs b/Engine/Source/Programs/Shared/EpicGames.Slack/Elements/ButtonElement.cs
--- a/Engine/Source/Programs/Shared/EpicGames.Slack/Elements/ButtonElement.cs
+++ b/Engine/Source/Programs/Shared/EpicGames.Slack/Elements/ButtonElement.cs
@@ -116,7 +116,7 @@ namespace EpicGames.Slack.Elements
 		/// </summary>
 		public static ISlackElementContainer AddButton(this ISlackElementContainer list, string text, Uri? url = null, string? actionId = null, string? value = null, ButtonStyle? style = null)
 		{
-			ButtonElement button = new ButtonElement(text, url, actionId, value, style);
+			ButtonElement button = new(text, url, actionId, value, style);
 			list.Elements.Add(button);
 			return list;
 		}
diff --git a/Engine/Source/Programs/Shared/EpicGames.Slack/Elements/TextObject.cs b/Engine/Source/Programs/Shared/EpicGames.Slack/Elements/TextObject.cs
--- a/Engine/Source/Programs/Shared/EpicGames.Slack/Elements/TextObject.cs
+++ b/Engine/Source/Programs/Shared/EpicGames.Slack/Elements/TextObject.cs
@@ -41,7 +41,7 @@ namespace EpicGames.Slack.Elements
 		/// <summary>
 		/// Implicit conversion from a regular string
 		/// </summary>
-		public static implicit operator TextObject(string text) => new TextObject(text);
+		public static implicit operator TextObject(string text) => new(text);
 	}
 
 	/// <summary>
@@ -59,6 +59,6 @@ namespace EpicGames.Slack.Elements
 		/// <summary>
 		/// Implicit conversion from a regular string
 		/// </summary>
-		public static implicit operator PlainTextObject(string text) => new PlainTextObject(text);
+		public static implicit operator PlainTextObject(string text) => new(text);
 	}
 }
diff --git a/Engine/Source/Programs/Shared/EpicGames.Slack/SlackAttachment.cs b/Engine/Source/Programs/Shared/EpicGames.Slack/SlackAttachment.cs
--- a/Engine/Source/Programs/Shared/EpicGames.Slack/SlackAttachment.cs
+++ b/Engine/Source/Programs/Shared/EpicGames.Slack/SlackAttachment.cs
@@ -35,7 +35,7 @@ namespace EpicGames.Slack
 		/// <param name="attachment"></param>
 		public static implicit operator SlackMessage(SlackAttachment attachment)
 		{
-			SlackMessage message = new SlackMessage();
+			SlackMessage message = new();
 			message.Attachments.Add(attachment);
 			return message;
 		}
diff --git a/Engine/Source/Programs/Shared/EpicGames.Slack/SlackClient.cs b/Engine/Source/Programs/Shared/EpicGames.Slack/SlackClient.cs
--- a/Engine/Source/Programs/Shared/EpicGames.Slack/SlackClient.cs
+++ b/Engine/Source/Programs/Shared/EpicGames.Slack/SlackClient.cs
@@ -333,7 +333,7 @@ namespace EpicGames.Slack
 
 		async Task<SlackMessageId> PostOrUpdateMessageAsync(string recipient, string? ts, string? threadTs, SlackMessage message, bool replyBroadcast, CancellationToken cancellationToken)
 		{
-			PostMessageRequest request = new PostMessageRequest();
+			PostMessageRequest request = new();
 			request.Channel = recipient;
 			request.Ts = String.IsNullOrEmpty(ts) ? null : ts;
 			request.ThreadTs = String.IsNullOrEmpty(threadTs) ? null : threadTs;
@@ -394,7 +394,7 @@ namespace EpicGames.Slack
 		public async Task<string> GetPermalinkAsync(SlackMessageId id, CancellationToken cancellationToken = default)
 		{
 			string requestUrl = $"{GetPermalinkUrl}?channel={id.Channel}&message_ts={id.Ts}";
-			using (HttpRequestMessage sendMessageRequest = new HttpRequestMessage(HttpMethod.Get, requestUrl))
+			using (HttpRequestMessage sendMessageRequest = new(HttpMethod.Get, requestUrl))
 			{
 				GetPermalinkResponse response = await SendRequestAsync<GetPermalinkResponse>(sendMessageRequest, "", cancellationToken);
 				if (!response.Ok || response.Permalink == null)
@@ -442,7 +442,7 @@ namespace EpicGames.Slack
 		/// <param name="cancellationToken">Cancellation token for the operation</param>
 		public async Task AddReactionAsync(SlackMessageId messageId, string name, CancellationToken cancellationToken = default)
 		{
-			ReactionMessage message = new ReactionMessage();
+			ReactionMessage message = new();
 			message.Channel = messageId.Channel;
 			message.Ts = messageId.Ts;
 			message.Name = name;
@@ -458,7 +458,7 @@ namespace EpicGames.Slack
 		/// <param name="cancellationToken">Cancellation token for the operation</param>
 		public async Task RemoveReactionAsync(SlackMessageId messageId, string name, CancellationToken cancellationToken = default)
 		{
-			ReactionMessage message = new ReactionMessage();
+			ReactionMessage message = new();
 			message.Channel = messageId.Channel;
 			message.Ts = messageId.Ts;
 			message.Name = name;
@@ -516,7 +516,7 @@ namespace EpicGames.Slack
 		/// <param name="cancellationToken">Cancellation token for the operation</param>
 		public async Task InviteUsersAsync(string channel, IEnumerable<string> userIds, CancellationToken cancellationToken = default)
 		{
-			InviteMessage message = new InviteMessage();
+			InviteMessage message = new();
 			message.Channel = channel;
 			message.Users = String.Join(",", userIds);
 
@@ -531,7 +531,7 @@ namespace EpicGames.Slack
 		/// <param name="cancellationToken">Cancellation token for the operation</param>
 		public async Task<string?> TryInviteUsersAsync(string channel, IEnumerable<string> userIds, CancellationToken cancellationToken = default)
 		{
-			InviteMessage message = new InviteMessage();
+			InviteMessage message = new();
 			message.Channel = channel;
 			message.Users = String.Join(",", userIds);
 
@@ -563,7 +563,7 @@ namespace EpicGames.Slack
 		/// <returns></returns>
 		public async Task <string?> OpenDirectMessageAsync(IEnumerable<string> userIds, CancellationToken cancellationToken = default)
 		{
-			OpenDirectMessage message = new OpenDirectMessage();
+			OpenDirectMessage message = new();
 			message.Users = String.Join(",", userIds);
 
 			GetChannelResponse response = await SendRequestAsync<GetChannelResponse>(ConversationsOpenUrl, message, cancellationToken);
@@ -649,7 +649,7 @@ namespace EpicGames.Slack
 		/// <param name="cancellationToken">Cancellation token for the operation</param>
 		public async Task OpenViewAsync(string triggerId, SlackView view, CancellationToken cancellationToken = default)
 		{
-			ViewsOpenRequest request = new ViewsOpenRequest();
+			ViewsOpenRequest request = new();
 			request.TriggerId = triggerId;
 			request.View = view;
 			await SendRequestAsync<PostMessageResponse>(ViewsOpenUrl, request, cancellationToken);
diff --git a/Engine/Source/Programs/Shared/EpicGames.Slack/SlackMessage.cs b/Engine/Source/Programs/Shared/EpicGames.Slack/SlackMessage.cs
--- a/Engine/Source/Programs/Shared/EpicGames.Slack/SlackMessage.cs
+++ b/Engine/Source/Programs/Shared/EpicGames.Slack/SlackMessage.cs
@@ -45,7 +45,7 @@ namespace EpicGames.Slack
 		/// <param name="text"></param>
 		public static implicit operator SlackMessage(string text)
 		{
-			SlackMessage message = new SlackMessage();
+			SlackMessage message = new();
 			message.Text = text;
 			return message;
 		}
diff --git a/Engine/Source/Programs/Shared/EpicGames.Tracing.Tests/UnrealInsights/EventTest.cs b/Engine/Source/Programs/Shared/EpicGames.Tracing.Tests/UnrealInsights/EventTest.cs
--- a/Engine/Source/Programs/Shared/EpicGames.Tracing.Tests/UnrealInsights/EventTest.cs
+++ b/Engine/Source/Programs/Shared/EpicGames.Tracing.Tests/UnrealInsights/EventTest.cs
@@ -13,15 +13,15 @@ namespace EpicGames.Tracing.Tests.UnrealInsights
 		public void EnterScopeEventTimestampDeserialize()
 		{
 			{
-				using MemoryStream ms = new MemoryStream([0x10, 0x05, 0x00, 0x00, 0x00, 0x00, 0x00, 0x00]);
-				using BinaryReader reader = new BinaryReader(ms);
+				using MemoryStream ms = new([0x10, 0x05, 0x00, 0x00, 0x00, 0x00, 0x00, 0x00]);
+				using BinaryReader reader = new(ms);
 				EnterScopeEventTimestamp @event = EnterScopeEventTimestamp.Deserialize(reader);
 				Assert.AreEqual((ulong)0x05, @event.Timestamp);
 			}
 
 			{
-				using MemoryStream ms = new MemoryStream([0x10, 0x00, 0x00, 0x00, 0x00, 0x00, 0x00, 0xFF]);
-				using BinaryReader reader = new BinaryReader(ms);
+				using MemoryStream ms = new([0x10, 0x00, 0x00, 0x00, 0x00, 0x00, 0x00, 0xFF]);
+				using BinaryReader reader = new(ms);
 				EnterScopeEventTimestamp @event = EnterScopeEventTimestamp.Deserialize(reader);
 				Assert.AreEqual((ulong)0x00_FF_00_00_00_00_00_00, @event.Timestamp);
 			}
diff --git a/Engine/Source/Programs/Shared/EpicGames.Tracing.Tests/UnrealInsights/EventTypeTest.cs b/Engine/Source/Programs/Shared/EpicGames.Tracing.Tests/UnrealInsights/EventTypeTest.cs
--- a/Engine/Source/Programs/Shared/EpicGames.Tracing.Tests/UnrealInsights/EventTypeTest.cs
+++ b/Engine/Source/Programs/Shared/EpicGames.Tracing.Tests/UnrealInsights/EventTypeTest.cs
@@ -67,8 +67,8 @@ namespace EpicGames.Tracing.Tests.UnrealInsights
 		private static EventType DeserializeAndAssertEventType(string hexData, EventType expectedType, ushort expectedUid)
 		{
 			byte[] block = GenericEventTest.HexStringToBytes(hexData);
-			using MemoryStream ms = new MemoryStream(block);
-			using BinaryReader reader = new BinaryReader(ms);
+			using MemoryStream ms = new(block);
+			using BinaryReader reader = new(ms);
 
 			TraceImportantEventHeader eventHeader = TraceImportantEventHeader.Deserialize(reader);
 			Assert.AreEqual(PredefinedEventUid.NewEvent, eventHeader.Uid);
@@ -97,8 +97,8 @@ namespace EpicGames.Tracing.Tests.UnrealInsights
 
 		private static void SerializeAndAssertEventType(string expectedHexData, EventType eventType, ushort uid)
 		{
-			using MemoryStream ms = new MemoryStream();
-			using BinaryWriter writer = new BinaryWriter(ms);
+			using MemoryStream ms = new();
+			using BinaryWriter writer = new(ms);
 			byte[] expectedData = GenericEventTest.HexStringToBytes(expectedHexData);
 			int expectedEventSize = expectedData.Length;
 			Assert.AreEqual(expectedEventSize, eventType.Size);
diff --git a/Engine/Source/Programs/Shared/EpicGames.Tracing.Tests/UnrealInsights/GenericEventTest.cs b/Engine/Source/Programs/Shared/EpicGames.Tracing.Tests/UnrealInsights/GenericEventTest.cs
--- a/Engine/Source/Programs/Shared/EpicGames.Tracing.Tests/UnrealInsights/GenericEventTest.cs
+++ b/Engine/Source/Programs/Shared/EpicGames.Tracing.Tests/UnrealInsights/GenericEventTest.cs
@@ -83,7 +83,7 @@ namespace EpicGames.Tracing.Tests.UnrealInsights
 		[TestMethod]
 		public void Serialize_Trace_NewTrace()
 		{
-			TraceNewTraceEvent @event = new TraceNewTraceEvent(25582215261913, 10000000, 21069, 8);
+			TraceNewTraceEvent @event = new(25582215261913, 10000000, 21069, 8);
 			SerializeAndAssertEvent(_traceNewTraceHex, @event, true, 16);
 		}
 
@@ -122,17 +122,17 @@ namespace EpicGames.Tracing.Tests.UnrealInsights
 		public void Serialize_Trace_ThreadInfo()
 		{
 			{
-				TraceThreadInfoEvent @event = new TraceThreadInfoEvent(2, 171380, -1, "GameThread");
+				TraceThreadInfoEvent @event = new(2, 171380, -1, "GameThread");
 				SerializeAndAssertEvent(_traceThreadInfoHex1, @event, true, 36);
 			}
 
 			{
-				TraceThreadInfoEvent @event = new TraceThreadInfoEvent(3, 0, 2147483647, "Trace");
+				TraceThreadInfoEvent @event = new(3, 0, 2147483647, "Trace");
 				SerializeAndAssertEvent(_traceThreadInfoHex2, @event, true, 36);
 			}
 
 			{
-				TraceThreadInfoEvent @event = new TraceThreadInfoEvent(4, 441304, 80, "Foreground Worker #0");
+				TraceThreadInfoEvent @event = new(4, 441304, 80, "Foreground Worker #0");
 				SerializeAndAssertEvent(_traceThreadInfoHex3, @event, true, 36);
 			}
 		}
@@ -171,7 +171,7 @@ namespace EpicGames.Tracing.Tests.UnrealInsights
 		[TestMethod]
 		public void Serialize_Diagnostics_Session2()
 		{
-			DiagnosticsSession2Event @event = new DiagnosticsSession2Event("Win64", "UnrealEditor", @" D:\depot\starship-main\QAGame\QAGame.uproject -trace=cpu -game", "++UE5+Main", "++UE5+Main-CL-17442524", 17442524, 3, 4);
+			DiagnosticsSession2Event @event = new("Win64", "UnrealEditor", @" D:\depot\starship-main\QAGame\QAGame.uproject -trace=cpu -game", "++UE5+Main", "++UE5+Main-CL-17442524", 17442524, 3, 4);
 			SerializeAndAssertEvent(_diagnosticsSession2Hex, @event, true, 34);
 		}
 
@@ -196,7 +196,7 @@ namespace EpicGames.Tracing.Tests.UnrealInsights
 		[TestMethod]
 		public void CpuProfiler_EventSpec_Serialize()
 		{
-			CpuProfilerEventSpecEvent @event = new CpuProfilerEventSpecEvent(9524, "K2Node_VariableGet /Engine/EngineSky/BP_Sky_Sphere.BP_Sky_Sphere:UserConstructionScript.K2Node_VariableGet_983");
+			CpuProfilerEventSpecEvent @event = new(9524, "K2Node_VariableGet /Engine/EngineSky/BP_Sky_Sphere.BP_Sky_Sphere:UserConstructionScript.K2Node_VariableGet_983");
 			SerializeAndAssertEvent(_cpuProfilerEventSpecHex, @event, true, 33);
 		}
 
@@ -207,7 +207,7 @@ namespace EpicGames.Tracing.Tests.UnrealInsights
 			byte[] eventData = @event.GetFields()[0].GetArray()!;
 			Assert.AreEqual(241, eventData.Length);
 
-			CpuProfilerSerializer serializer = new CpuProfilerSerializer(10000000);
+			CpuProfilerSerializer serializer = new(10000000);
 			serializer.Read(eventData);
 
 			byte[] eventDataSerialized = serializer.Write()[0];
@@ -219,7 +219,7 @@ namespace EpicGames.Tracing.Tests.UnrealInsights
 		public void CpuProfiler_EventBatch_Serialize()
 		{
 			byte[] data = HexStringToBytes("87 A7 EB AE 8A D1 0B 01 A7 02 02 04 FB 23 03 A4 0A 61 04 0B 05 F2 80 42 AA 1C EB 5C 06 C0 BC 03 17 07 02 47 08 C2 02 AB 15 09 AA 03 07 0A AE 01 05 0B D8 0D 05 0C B7 01 0D D8 03 00 03 0E A4 03 6F 0F 89 FF 03 10 F3 C4 04 11 60 06 A7 02 10 AB 19 12 B3 04 10 9F 11 13 0A 04 02 04 ED 03 10 91 11 14 80 01 02 97 02 10 A9 D6 0B 15 10 06 DD 03 10 CB 85 04 16 12 06 D3 01 10 89 16 17 0C 02 D6 02 07 18 DA 84 28 99 01 19 18 05 1A 0A A5 10 1B CD 01 1C 06 88 6A C5 49 1D B8 EB 24 A9 17 20 AF 01 10 CD F9 01 21 71 22 F2 AE 12 85 04 23 D1 0D 24 85 74 25 EA 58 91 0D 25 80 E6 02 A9 1C 25 88 CD 06 A7 27 25 D0 DA 09 8B 20 25 DE AD 0B AF 33 25 C0 A9 09 ED 23 25 C4 81 0A F5 2D 25 9A EC 0A F3 35 25 D4 CF 03 FF 19 25 BE 92 01 B1 0D 25 98 22");
-			CpuProfilerEventBatchEvent @event = new CpuProfilerEventBatchEvent(data);
+			CpuProfilerEventBatchEvent @event = new(data);
 			SerializeAndAssertEvent(_cpuProfilerEventBatchHex1, @event, false, 39);
 		}
 
@@ -267,18 +267,18 @@ namespace EpicGames.Tracing.Tests.UnrealInsights
 		public void SerializeEventWithAuxAndImportant()
 		{
 			ushort uid = 33;
-			using MemoryStream ms = new MemoryStream();
-			using BinaryWriter writer = new BinaryWriter(ms);
+			using MemoryStream ms = new();
+			using BinaryWriter writer = new(ms);
 
 			Field[] fields =
 			[
 				Field.FromInt(9524),
 				Field.FromString("K2Node_VariableGet /Engine/EngineSky/BP_Sky_Sphere.BP_Sky_Sphere:UserConstructionScript.K2Node_VariableGet_983"),
 			];
-			GenericEvent @event = new GenericEvent(98989, fields, CpuProfilerEventSpecEvent.EventType);
+			GenericEvent @event = new(98989, fields, CpuProfilerEventSpecEvent.EventType);
 			Assert.AreEqual(0x77, @event.Size);
 
-			TraceImportantEventHeader eventHeader = new TraceImportantEventHeader(uid, @event.Size);
+			TraceImportantEventHeader eventHeader = new(uid, @event.Size);
 			eventHeader.Serialize(writer);
 			@event.Serialize(20, writer);
 
@@ -289,8 +289,8 @@ namespace EpicGames.Tracing.Tests.UnrealInsights
 		public void AuxHeaderDeserialize()
 		{
 			// Aux header from field "Name" in CpuProfiler.EventSpec
-			using MemoryStream ms = new MemoryStream([0x01, 0xC1, 0x0D, 0x00]);
-			using BinaryReader reader = new BinaryReader(ms);
+			using MemoryStream ms = new([0x01, 0xC1, 0x0D, 0x00]);
+			using BinaryReader reader = new(ms);
 
 			(ushort uid, int fieldIndex, int size) = GenericEvent.DeserializeAuxHeader(reader.ReadUInt32());
 			Assert.AreEqual(PredefinedEventUid.AuxData, uid);
@@ -302,8 +302,8 @@ namespace EpicGames.Tracing.Tests.UnrealInsights
 		public void AuxHeaderSerialize()
 		{
 			uint header = GenericEvent.SerializeAuxHeader(1, 1, 110);
-			using MemoryStream ms = new MemoryStream();
-			using BinaryWriter writer = new BinaryWriter(ms);
+			using MemoryStream ms = new();
+			using BinaryWriter writer = new(ms);
 			writer.Write(header);
 			byte[] data = ms.ToArray();
 			Assert.AreEqual(0x01, data[0]);
@@ -363,8 +363,8 @@ namespace EpicGames.Tracing.Tests.UnrealInsights
 		private static GenericEvent DeserializeAndAssertGenericEvent(string hexData, EventType eventType, bool isImportant, ushort expectedUid, uint? expectedSerial = 0)
 		{
 			byte[] block = HexStringToBytes(hexData);
-			using MemoryStream ms = new MemoryStream(block);
-			using BinaryReader reader = new BinaryReader(ms);
+			using MemoryStream ms = new(block);
+			using BinaryReader reader = new(ms);
 
 			ushort uid;
 			if (isImportant)
@@ -390,17 +390,17 @@ namespace EpicGames.Tracing.Tests.UnrealInsights
 
 		private static void SerializeAndAssertGenericEvent(string expectedHexData, EventType eventType, bool isImportant, ushort uid, Field[] fields, uint serial = 0)
 		{
-			using MemoryStream ms = new MemoryStream();
-			using BinaryWriter writer = new BinaryWriter(ms);
+			using MemoryStream ms = new();
+			using BinaryWriter writer = new(ms);
 
 			byte[] expectedData = HexStringToBytes(expectedHexData);
-			GenericEvent @event = new GenericEvent(serial, fields, eventType);
+			GenericEvent @event = new(serial, fields, eventType);
 
 			if (isImportant)
 			{
 				int expectedEventSize = expectedData.Length - TraceImportantEventHeader.HeaderSize;
 				Assert.AreEqual(expectedEventSize, @event.Size);
-				TraceImportantEventHeader eventHeader = new TraceImportantEventHeader(uid, @event.Size);
+				TraceImportantEventHeader eventHeader = new(uid, @event.Size);
 				eventHeader.Serialize(writer);
 			}
 			else
@@ -415,8 +415,8 @@ namespace EpicGames.Tracing.Tests.UnrealInsights
 
 		private static void SerializeAndAssertEvent(string expectedHexData, ITraceEvent @event, bool isImportant, ushort uid)
 		{
-			using MemoryStream ms = new MemoryStream();
-			using BinaryWriter writer = new BinaryWriter(ms);
+			using MemoryStream ms = new();
+			using BinaryWriter writer = new(ms);
 			@event.Serialize(uid, writer);
 
 			if (isImportant)
diff --git a/Engine/Source/Programs/Shared/EpicGames.Tracing.Tests/UnrealInsights/TransportPacketTest.cs b/Engine/Source/Programs/Shared/EpicGames.Tracing.Tests/UnrealInsights/TransportPacketTest.cs
--- a/Engine/Source/Programs/Shared/EpicGames.Tracing.Tests/UnrealInsights/TransportPacketTest.cs
+++ b/Engine/Source/Programs/Shared/EpicGames.Tracing.Tests/UnrealInsights/TransportPacketTest.cs
@@ -59,8 +59,8 @@ namespace EpicGames.Tracing.Tests.UnrealInsights
 		{
 			static void AssertWriteUid(ushort uid, bool expectTwoByteUid)
 			{
-				using MemoryStream ms = new MemoryStream();
-				using BinaryWriter writer = new BinaryWriter(ms);
+				using MemoryStream ms = new();
+				using BinaryWriter writer = new(ms);
 				writer.WritePackedUid(uid);
 				byte[] data = ms.ToArray();
 				Assert.AreEqual(expectTwoByteUid ? 2 : 1, data.Length);
@@ -89,8 +89,8 @@ namespace EpicGames.Tracing.Tests.UnrealInsights
 			static void Assert7Bit(ulong expectedValue, string actualBytesHex)
 			{
 				byte[] actualBytes = GenericEventTest.StringToByteArray(actualBytesHex.Replace(" ", "", StringComparison.Ordinal));
-				using MemoryStream ms = new MemoryStream(actualBytes);
-				using BinaryReader reader = new BinaryReader(ms);
+				using MemoryStream ms = new(actualBytes);
+				using BinaryReader reader = new(ms);
 				Assert.AreEqual(expectedValue, TraceUtils.Read7BitUint(reader));
 			}
 
@@ -127,8 +127,8 @@ namespace EpicGames.Tracing.Tests.UnrealInsights
 			static void Assert7Bit(ulong value, string expectedBytesHex)
 			{
 				byte[] expectedBytes = GenericEventTest.StringToByteArray(expectedBytesHex.Replace(" ", "", StringComparison.Ordinal));
-				using MemoryStream ms = new MemoryStream();
-				using BinaryWriter writer = new BinaryWriter(ms);
+				using MemoryStream ms = new();
+				using BinaryWriter writer = new(ms);
 
 				int actualNumBytesWritten = TraceUtils.Write7BitUint(writer, value);
 				byte[] actualBytes = ms.ToArray();
@@ -168,19 +168,19 @@ namespace EpicGames.Tracing.Tests.UnrealInsights
 		[TestMethod]
 		public void SerializeDeserializeWithStubs()
 		{
-			StubTraceEvent event1 = new StubTraceEvent(1000, 20);
-			StubTraceEvent event2 = new StubTraceEvent(2000, 30);
+			StubTraceEvent event1 = new(1000, 20);
+			StubTraceEvent event2 = new(2000, 30);
 			ushort uid1 = 5555;
 			ushort uid2 = 6666;
 
 			{
-				using MemoryStream ms = new MemoryStream();
-				using BinaryWriter writer = new BinaryWriter(ms);
+				using MemoryStream ms = new();
+				using BinaryWriter writer = new(ms);
 				TransportPacket packet = TransportPacket.Create(0, 0);
 				packet.Serialize(writer, [(uid1, event1 as ITraceEvent), (uid2, event2 as ITraceEvent)]);
 
 				ms.Position = 0;
-				using BinaryReader reader = new BinaryReader(ms);
+				using BinaryReader reader = new(ms);
 
 				TransportPacket deserializedPacket = TransportPacket.Deserialize(reader);
 				Assert.AreEqual(0, deserializedPacket.ThreadIdAndMarkers);
@@ -188,13 +188,13 @@ namespace EpicGames.Tracing.Tests.UnrealInsights
 			}
 
 			{
-				using MemoryStream ms = new MemoryStream();
-				using BinaryWriter writer = new BinaryWriter(ms);
+				using MemoryStream ms = new();
+				using BinaryWriter writer = new(ms);
 				TransportPacket packet = TransportPacket.Create(0, 1);
 				packet.Serialize(writer, []);
 
 				ms.Position = 0;
-				using BinaryReader reader = new BinaryReader(ms);
+				using BinaryReader reader = new(ms);
 
 				TransportPacket deserializedPacket = TransportPacket.Deserialize(reader);
 				Assert.AreEqual(1, deserializedPacket.ThreadIdAndMarkers);
@@ -202,13 +202,13 @@ namespace EpicGames.Tracing.Tests.UnrealInsights
 			}
 
 			{
-				using MemoryStream ms = new MemoryStream();
-				using BinaryWriter writer = new BinaryWriter(ms);
+				using MemoryStream ms = new();
+				using BinaryWriter writer = new(ms);
 				TransportPacket packet = TransportPacket.Create(0, 12345);
 				packet.Serialize(writer, []);
 
 				ms.Position = 0;
-				using BinaryReader reader = new BinaryReader(ms);
+				using BinaryReader reader = new(ms);
 
 				TransportPacket deserializedPacket = TransportPacket.Deserialize(reader);
 				Assert.AreEqual(12345, deserializedPacket.ThreadIdAndMarkers);
@@ -219,22 +219,22 @@ namespace EpicGames.Tracing.Tests.UnrealInsights
 		[TestMethod]
 		public void SerializeDeserialize()
 		{
-			TraceNewTraceEvent event1 = new TraceNewTraceEvent(25582215261913, 10000000, 21069, 8);
-			DiagnosticsSession2Event event2 = new DiagnosticsSession2Event("Win64", "UnrealEditor", @" D:\depot\starship-main\QAGame\QAGame.uproject -trace=cpu -game", "++UE5+Main", "++UE5+Main-CL-17442524", 17442524, 3, 4);
-			CpuProfilerEventSpecEvent event3 = new CpuProfilerEventSpecEvent(9524, "K2Node_VariableGet /Engine/EngineSky/BP_Sky_Sphere.BP_Sky_Sphere:UserConstructionScript.K2Node_VariableGet_983");
+			TraceNewTraceEvent event1 = new(25582215261913, 10000000, 21069, 8);
+			DiagnosticsSession2Event event2 = new("Win64", "UnrealEditor", @" D:\depot\starship-main\QAGame\QAGame.uproject -trace=cpu -game", "++UE5+Main", "++UE5+Main-CL-17442524", 17442524, 3, 4);
+			CpuProfilerEventSpecEvent event3 = new(9524, "K2Node_VariableGet /Engine/EngineSky/BP_Sky_Sphere.BP_Sky_Sphere:UserConstructionScript.K2Node_VariableGet_983");
 
 			ushort uid1 = 10;
 			ushort uid2 = 20;
 			ushort uid3 = 30;
 
 			{
-				using MemoryStream ms = new MemoryStream();
-				using BinaryWriter writer = new BinaryWriter(ms);
+				using MemoryStream ms = new();
+				using BinaryWriter writer = new(ms);
 				TransportPacket packet = TransportPacket.Create(0, 0);
 				packet.Serialize(writer, [(uid1, event1 as ITraceEvent), (uid2, event2 as ITraceEvent)]);
 
 				ms.Position = 0;
-				using BinaryReader reader = new BinaryReader(ms);
+				using BinaryReader reader = new(ms);
 
 				TransportPacket deserializedPacket = TransportPacket.Deserialize(reader);
 				Assert.AreEqual(0, deserializedPacket.ThreadIdAndMarkers);
@@ -243,8 +243,8 @@ namespace EpicGames.Tracing.Tests.UnrealInsights
 			}
 
 			{
-				using MemoryStream ms = new MemoryStream();
-				using BinaryWriter writer = new BinaryWriter(ms);
+				using MemoryStream ms = new();
+				using BinaryWriter writer = new(ms);
 				TransportPacket packet = TransportPacket.Create(0, 0);
 				packet.Serialize(writer,
 				[
@@ -254,7 +254,7 @@ namespace EpicGames.Tracing.Tests.UnrealInsights
 				]);
 
 				ms.Position = 0;
-				using BinaryReader reader = new BinaryReader(ms);
+				using BinaryReader reader = new(ms);
 
 				TransportPacket deserializedPacket = TransportPacket.Deserialize(reader);
 				Assert.AreEqual(0, deserializedPacket.ThreadIdAndMarkers);
diff --git a/Engine/Source/Programs/Shared/EpicGames.Tracing.Tests/UnrealInsights/UnrealInsightsTest.cs b/Engine/Source/Programs/Shared/EpicGames.Tracing.Tests/UnrealInsights/UnrealInsightsTest.cs
--- a/Engine/Source/Programs/Shared/EpicGames.Tracing.Tests/UnrealInsights/UnrealInsightsTest.cs
+++ b/Engine/Source/Programs/Shared/EpicGames.Tracing.Tests/UnrealInsights/UnrealInsightsTest.cs
@@ -48,15 +48,15 @@ namespace EpicGames.Tracing.Tests.UnrealInsights
 			ulong test3 = 444;
 			ushort eventSize = 4 + 4 + 8;
 
-			using MemoryStream ms = new MemoryStream();
-			using BinaryWriter writer = new BinaryWriter(ms);
+			using MemoryStream ms = new();
+			using BinaryWriter writer = new(ms);
 			new TraceImportantEventHeader(1000, eventSize).Serialize(writer);
 			writer.Write(test1);
 			writer.Write(test2);
 			writer.Write(test3);
 
 			ms.Position = 0;
-			using BinaryReader reader = new BinaryReader(ms);
+			using BinaryReader reader = new(ms);
 			TraceImportantEventHeader eventHeader = TraceImportantEventHeader.Deserialize(reader);
 			Assert.AreEqual(1000, eventHeader.Uid);
 			Assert.AreEqual(eventSize, eventHeader.EventSize);
@@ -72,15 +72,15 @@ namespace EpicGames.Tracing.Tests.UnrealInsights
 		[TestMethod]
 		public void WriteUtraceFile()
 		{
-			using MemoryStream ms = new MemoryStream();
-			using BinaryWriter binaryWriter = new BinaryWriter(ms);
+			using MemoryStream ms = new();
+			using BinaryWriter binaryWriter = new(ms);
 
-			UnrealInsightsWriter writer = new UnrealInsightsWriter();
+			UnrealInsightsWriter writer = new();
 
-			TraceNewTraceEvent newTraceEvent1 = new TraceNewTraceEvent(1001, 2001, 31, 8);
-			TraceNewTraceEvent newTraceEvent2 = new TraceNewTraceEvent(1002, 2002, 32, 8);
-			TraceNewTraceEvent newTraceEvent3 = new TraceNewTraceEvent(1003, 2003, 33, 8);
-			CpuProfilerEventSpecEvent cpu1 = new CpuProfilerEventSpecEvent(400, "MyCpuEventSpecEvent1");
+			TraceNewTraceEvent newTraceEvent1 = new(1001, 2001, 31, 8);
+			TraceNewTraceEvent newTraceEvent2 = new(1002, 2002, 32, 8);
+			TraceNewTraceEvent newTraceEvent3 = new(1003, 2003, 33, 8);
+			CpuProfilerEventSpecEvent cpu1 = new(400, "MyCpuEventSpecEvent1");
 
 			writer.AddEvent(TransportPacket.ThreadIdImportants, newTraceEvent1);
 			writer.AddEvent(TransportPacket.ThreadIdImportants, newTraceEvent2);
@@ -90,7 +90,7 @@ namespace EpicGames.Tracing.Tests.UnrealInsights
 
 			ms.Position = 0;
 
-			UnrealInsightsReader reader = new UnrealInsightsReader();
+			UnrealInsightsReader reader = new();
 			reader.Read(ms);
 			reader.PrintEventSummary();
 			Assert.AreEqual(4, reader.EventTypes.Count);
@@ -133,27 +133,27 @@ namespace EpicGames.Tracing.Tests.UnrealInsights
 		[TestMethod]
 		public void WriteUtraceExample()
 		{
-			using MemoryStream ms = new MemoryStream();
-			using BinaryWriter binaryWriter = new BinaryWriter(ms);
-			UnrealInsightsWriter writer = new UnrealInsightsWriter();
+			using MemoryStream ms = new();
+			using BinaryWriter binaryWriter = new(ms);
+			UnrealInsightsWriter writer = new();
 
 			byte[] cpuBatchData1 = GenericEventTest.HexStringToBytes("87 A7 EB AE 8A D1 0B 01 A7 02 02 04 FB 23 03 A4 0A 61 04 0B 05 F2 80 42 AA 1C EB 5C 06 C0 BC 03 17 07 02 47 08 C2 02 AB 15 09 AA 03 07 0A AE 01 05 0B D8 0D 05 0C B7 01 0D D8 03 00 03 0E A4 03 6F 0F 89 FF 03 10 F3 C4 04 11 60 06 A7 02 10 AB 19 12 B3 04 10 9F 11 13 0A 04 02 04 ED 03 10 91 11 14 80 01 02 97 02 10 A9 D6 0B 15 10 06 DD 03 10 CB 85 04 16 12 06 D3 01 10 89 16 17 0C 02 D6 02 07 18 DA 84 28 99 01 19 18 05 1A 0A A5 10 1B CD 01 1C 06 88 6A C5 49 1D B8 EB 24 A9 17 20 AF 01 10 CD F9 01 21 71 22 F2 AE 12 85 04 23 D1 0D 24 85 74 25 EA 58 91 0D 25 80 E6 02 A9 1C 25 88 CD 06 A7 27 25 D0 DA 09 8B 20 25 DE AD 0B AF 33 25 C0 A9 09 ED 23 25 C4 81 0A F5 2D 25 9A EC 0A F3 35 25 D4 CF 03 FF 19 25 BE 92 01 B1 0D 25 98 22");
 
 			ulong cycleFrequency = 10000000;
-			TraceNewTraceEvent newTrace1 = new TraceNewTraceEvent(25582215261913, cycleFrequency, 21069, 8);
-			TraceThreadInfoEvent threadInfo1 = new TraceThreadInfoEvent(2, 171380, -1, "GameThread");
-			TraceThreadInfoEvent threadInfo2 = new TraceThreadInfoEvent(3, 0, 2147483647, "Trace");
-			DiagnosticsSession2Event diagnosticsSession = new DiagnosticsSession2Event("Finally it works", "HELLO WORLD", "QAGame.uproject -trace=cpu -game", "++UE5+Main", "++UE5+Main-CL-17442524", 1000, 3, 4);
-			CpuProfilerEventSpecEvent cpuEventSpec1 = new CpuProfilerEventSpecEvent(1, "MyCpuEventSpecEvent1");
-			CpuProfilerEventSpecEvent cpuEventSpec2 = new CpuProfilerEventSpecEvent(2, "MyCpuEventSpecEvent2");
-			CpuProfilerEventSpecEvent cpuEventSpec3 = new CpuProfilerEventSpecEvent(3, "MyCpuEventSpecEvent3");
-			CpuProfilerEventSpecEvent cpuEventSpec4 = new CpuProfilerEventSpecEvent(4, "MyCpuEventSpecEvent4");
-			CpuProfilerEventSpecEvent cpuEventSpec5 = new CpuProfilerEventSpecEvent(5, "MyCpuEventSpecEvent5");
-			CpuProfilerEventSpecEvent cpuEventSpec6 = new CpuProfilerEventSpecEvent(6, "MyCpuEventSpecEvent6");
-			CpuProfilerEventSpecEvent cpuEventSpec7 = new CpuProfilerEventSpecEvent(7, "MyCpuEventSpecEvent7");
-			CpuProfilerEventSpecEvent cpuEventSpec8 = new CpuProfilerEventSpecEvent(8, "MyCpuEventSpecEvent8");
-			CpuProfilerEventSpecEvent cpuEventSpecTesting = new CpuProfilerEventSpecEvent(9, "TestingSome");
-			CpuProfilerEventBatchEvent cpuBatch1 = new CpuProfilerEventBatchEvent(cpuBatchData1);
+			TraceNewTraceEvent newTrace1 = new(25582215261913, cycleFrequency, 21069, 8);
+			TraceThreadInfoEvent threadInfo1 = new(2, 171380, -1, "GameThread");
+			TraceThreadInfoEvent threadInfo2 = new(3, 0, 2147483647, "Trace");
+			DiagnosticsSession2Event diagnosticsSession = new("Finally it works", "HELLO WORLD", "QAGame.uproject -trace=cpu -game", "++UE5+Main", "++UE5+Main-CL-17442524", 1000, 3, 4);
+			CpuProfilerEventSpecEvent cpuEventSpec1 = new(1, "MyCpuEventSpecEvent1");
+			CpuProfilerEventSpecEvent cpuEventSpec2 = new(2, "MyCpuEventSpecEvent2");
+			CpuProfilerEventSpecEvent cpuEventSpec3 = new(3, "MyCpuEventSpecEvent3");
+			CpuProfilerEventSpecEvent cpuEventSpec4 = new(4, "MyCpuEventSpecEvent4");
+			CpuProfilerEventSpecEvent cpuEventSpec5 = new(5, "MyCpuEventSpecEvent5");
+			CpuProfilerEventSpecEvent cpuEventSpec6 = new(6, "MyCpuEventSpecEvent6");
+			CpuProfilerEventSpecEvent cpuEventSpec7 = new(7, "MyCpuEventSpecEvent7");
+			CpuProfilerEventSpecEvent cpuEventSpec8 = new(8, "MyCpuEventSpecEvent8");
+			CpuProfilerEventSpecEvent cpuEventSpecTesting = new(9, "TestingSome");
+			CpuProfilerEventBatchEvent cpuBatch1 = new(cpuBatchData1);
 
 			writer.AddEvent(TransportPacket.ThreadIdImportants, newTrace1);
 			writer.AddEvent(TransportPacket.ThreadIdImportants, threadInfo1);
@@ -172,7 +172,7 @@ namespace EpicGames.Tracing.Tests.UnrealInsights
 
 			writer.AddEvent(2, cpuBatch1);
 
-			CpuProfilerSerializer cpuSerializerThread3 = new CpuProfilerSerializer(cycleFrequency);
+			CpuProfilerSerializer cpuSerializerThread3 = new(cycleFrequency);
 			cpuSerializerThread3.ScopeEvents.Add(new CpuProfilerScopeEvent(25582215261913 + 400 * 1000, true, cpuEventSpecTesting.Id));
 			cpuSerializerThread3.ScopeEvents.Add(new CpuProfilerScopeEvent(25582215261913 + (cycleFrequency * cycleFrequency), false, null));
 
@@ -185,7 +185,7 @@ namespace EpicGames.Tracing.Tests.UnrealInsights
 
 			ms.Seek(0, SeekOrigin.Begin);
 
-			UnrealInsightsReader reader = new UnrealInsightsReader();
+			UnrealInsightsReader reader = new();
 			reader.Read(ms);
 			reader.PrintEventSummary();
 
@@ -197,7 +197,7 @@ namespace EpicGames.Tracing.Tests.UnrealInsights
 		{
 			string exampleDecompTrace = "UnrealInsights/example_trace.decomp.utrace";
 
-			UnrealInsightsReader reader = new UnrealInsightsReader();
+			UnrealInsightsReader reader = new();
 			using FileStream stream = File.Open(exampleDecompTrace, FileMode.Open);
 			reader.Read(stream);
 			reader.PrintEventSummary();
diff --git a/Engine/Source/Programs/Shared/EpicGames.Tracing/UnrealInsights/Events/CpuProfiler.cs b/Engine/Source/Programs/Shared/EpicGames.Tracing/UnrealInsights/Events/CpuProfiler.cs
--- a/Engine/Source/Programs/Shared/EpicGames.Tracing/UnrealInsights/Events/CpuProfiler.cs
+++ b/Engine/Source/Programs/Shared/EpicGames.Tracing/UnrealInsights/Events/CpuProfiler.cs
@@ -8,7 +8,7 @@ namespace EpicGames.Tracing.UnrealInsights.Events
 {
 	public class CpuProfilerEventSpecEvent : ITraceEvent
 	{
-		public static readonly EventType EventType = new EventType(0, "CpuProfiler", "EventSpec", EventType.FlagImportant | EventType.FlagMaybeHasAux | EventType.FlagNoSync,
+		public static readonly EventType EventType = new(0, "CpuProfiler", "EventSpec", EventType.FlagImportant | EventType.FlagMaybeHasAux | EventType.FlagNoSync,
 			[
 				new EventTypeField(0, 4, EventTypeField.TypeInt32, "Id"),
 				new EventTypeField(4, 0, EventTypeField.TypeAnsiString, "Name")
@@ -45,7 +45,7 @@ namespace EpicGames.Tracing.UnrealInsights.Events
 
 	public class CpuProfilerEventBatchEvent : ITraceEvent
 	{
-		public static readonly EventType EventType = new EventType(0, "CpuProfiler", "EventBatch", EventType.FlagMaybeHasAux | EventType.FlagNoSync,
+		public static readonly EventType EventType = new(0, "CpuProfiler", "EventBatch", EventType.FlagMaybeHasAux | EventType.FlagNoSync,
 			[new EventTypeField(0, 0, EventTypeField.TypeArray, "Data")]);
 
 		public ushort Size => (ushort)(_genericEvent.Size + sizeof(ushort));
@@ -101,8 +101,8 @@ namespace EpicGames.Tracing.UnrealInsights.Events
 
 		public void Read(byte[] eventBatchData)
 		{
-			using MemoryStream ms = new MemoryStream(eventBatchData);
-			using BinaryReader reader = new BinaryReader(ms);
+			using MemoryStream ms = new(eventBatchData);
+			using BinaryReader reader = new(ms);
 
 			ulong absTimestamp = 0;
 			ulong lastTimestamp = 0;
@@ -125,12 +125,12 @@ namespace EpicGames.Tracing.UnrealInsights.Events
 				if (isEnterScope)
 				{
 					ulong specId = TraceUtils.Read7BitUint(reader);
-					CpuProfilerScopeEvent scopeEvent = new CpuProfilerScopeEvent(timestamp, true, (uint?)specId);
+					CpuProfilerScopeEvent scopeEvent = new(timestamp, true, (uint?)specId);
 					ScopeEvents.Add(scopeEvent);
 				}
 				else
 				{
-					CpuProfilerScopeEvent scopeEvent = new CpuProfilerScopeEvent(timestamp, false, null);
+					CpuProfilerScopeEvent scopeEvent = new(timestamp, false, null);
 					ScopeEvents.Add(scopeEvent);
 				}
 			}
@@ -139,8 +139,8 @@ namespace EpicGames.Tracing.UnrealInsights.Events
 		public List<byte[]> Write()
 		{
 			List<MemoryStream> streams = [];
-			MemoryStream stream = new MemoryStream();
-			BinaryWriter writer = new BinaryWriter(stream);
+			MemoryStream stream = new();
+			BinaryWriter writer = new(stream);
 			streams.Add(stream);
 
 			ulong lastAbsTimestamp = 0;
diff --git a/Engine/Source/Programs/Shared/EpicGames.Tracing/UnrealInsights/Events/Diagnostics.cs b/Engine/Source/Programs/Shared/EpicGames.Tracing/UnrealInsights/Events/Diagnostics.cs
--- a/Engine/Source/Programs/Shared/EpicGames.Tracing/UnrealInsights/Events/Diagnostics.cs
+++ b/Engine/Source/Programs/Shared/EpicGames.Tracing/UnrealInsights/Events/Diagnostics.cs
@@ -6,7 +6,7 @@ namespace EpicGames.Tracing.UnrealInsights.Events
 {
 	public class DiagnosticsSession2Event : ITraceEvent
 	{
-		public static readonly EventType EventType = new EventType(0, "Diagnostics", "Session2", EventType.FlagImportant | EventType.FlagMaybeHasAux | EventType.FlagNoSync,
+		public static readonly EventType EventType = new(0, "Diagnostics", "Session2", EventType.FlagImportant | EventType.FlagMaybeHasAux | EventType.FlagNoSync,
 			[
 				new EventTypeField(0, 0, EventTypeField.TypeAnsiString, "Platform"),
 				new EventTypeField(0, 0, EventTypeField.TypeAnsiString, "AppName"),
diff --git a/Engine/Source/Programs/Shared/EpicGames.Tracing/UnrealInsights/Events/GenericEvent.cs b/Engine/Source/Programs/Shared/EpicGames.Tracing/UnrealInsights/Events/GenericEvent.cs
--- a/Engine/Source/Programs/Shared/EpicGames.Tracing/UnrealInsights/Events/GenericEvent.cs
+++ b/Engine/Source/Programs/Shared/EpicGames.Tracing/UnrealInsights/Events/GenericEvent.cs
@@ -26,13 +26,13 @@ namespace EpicGames.Tracing.UnrealInsights.Events
 		private Field(double value) { Float = value; }
 		private Field(string value) { String = value; }
 		private Field(byte[] value) { _array = value; }
-		public static Field FromBool(bool value) => new Field(value);
-		public static Field FromInt(int value) => new Field(value);
-		public static Field FromLong(long value) => new Field(value);
-		public static Field FromFloat(float value) => new Field(value);
-		public static Field FromDouble(double value) => new Field(value);
-		public static Field FromString(string value) => new Field(value);
-		public static Field FromArray(byte[] value) => new Field(value);
+		public static Field FromBool(bool value) => new(value);
+		public static Field FromInt(int value) => new(value);
+		public static Field FromLong(long value) => new(value);
+		public static Field FromFloat(float value) => new(value);
+		public static Field FromDouble(double value) => new(value);
+		public static Field FromString(string value) => new(value);
+		public static Field FromArray(byte[] value) => new(value);
 	}
 
 	/// <summary>
diff --git a/Engine/Source/Programs/Shared/EpicGames.Tracing/UnrealInsights/Events/Logging.cs b/Engine/Source/Programs/Shared/EpicGames.Tracing/UnrealInsights/Events/Logging.cs
--- a/Engine/Source/Programs/Shared/EpicGames.Tracing/UnrealInsights/Events/Logging.cs
+++ b/Engine/Source/Programs/Shared/EpicGames.Tracing/UnrealInsights/Events/Logging.cs
@@ -7,7 +7,7 @@ namespace EpicGames.Tracing.UnrealInsights.Events
 {
 	public class LoggingLogMessage : ITraceEvent
 	{
-		public static readonly EventType EventType = new EventType(0, "Logging", "LogMessage", EventType.FlagMaybeHasAux | EventType.FlagNoSync,
+		public static readonly EventType EventType = new(0, "Logging", "LogMessage", EventType.FlagMaybeHasAux | EventType.FlagNoSync,
 			[
 				new EventTypeField(0, 8, EventTypeField.TypeInt64, "LogPoint"),
 				new EventTypeField(8, 8, EventTypeField.TypeInt64, "Cycle"),
