Shared: prefer implicit object creation when type is apparent

diff --git a/Engine/Source/Programs/Shared/EpicGames.Tracing/UnrealInsights/Events/Memory.cs b/Engine/Source/Programs/Shared/EpicGames.Tracing/UnrealInsights/Events/Memory.cs
--- a/Engine/Source/Programs/Shared/EpicGames.Tracing/UnrealInsights/Events/Memory.cs
+++ b/Engine/Source/Programs/Shared/EpicGames.Tracing/UnrealInsights/Events/Memory.cs
@@ -7,7 +7,7 @@ namespace EpicGames.Tracing.UnrealInsights.Events
 {
 	public class MemoryMemoryScopeEvent : ITraceEvent
 	{
-		public static readonly EventType EventType = new EventType(0, "Memory", "MemoryScope", EventType.FlagNone,
+		public static readonly EventType EventType = new(0, "Memory", "MemoryScope", EventType.FlagNone,
 			[new EventTypeField(0, 4, EventTypeField.TypeInt32, "Id")]);
 
 		public ushort Size => throw new NotImplementedException();
diff --git a/Engine/Source/Programs/Shared/EpicGames.Tracing/UnrealInsights/Events/NewEvent.cs b/Engine/Source/Programs/Shared/EpicGames.Tracing/UnrealInsights/Events/NewEvent.cs
--- a/Engine/Source/Programs/Shared/EpicGames.Tracing/UnrealInsights/Events/NewEvent.cs
+++ b/Engine/Source/Programs/Shared/EpicGames.Tracing/UnrealInsights/Events/NewEvent.cs
@@ -212,7 +212,7 @@ namespace EpicGames.Tracing.UnrealInsights.Events
 
 		public void AddEventType(ushort offset, ushort size, byte typeInfo, string name)
 		{
-			EventTypeField field = new EventTypeField(offset, size, typeInfo, name);
+			EventTypeField field = new(offset, size, typeInfo, name);
 			Fields.Add(field);
 			_eventSize += field.Size;
 		}
diff --git a/Engine/Source/Programs/Shared/EpicGames.Tracing/UnrealInsights/Events/Trace.cs b/Engine/Source/Programs/Shared/EpicGames.Tracing/UnrealInsights/Events/Trace.cs
--- a/Engine/Source/Programs/Shared/EpicGames.Tracing/UnrealInsights/Events/Trace.cs
+++ b/Engine/Source/Programs/Shared/EpicGames.Tracing/UnrealInsights/Events/Trace.cs
@@ -6,7 +6,7 @@ namespace EpicGames.Tracing.UnrealInsights.Events
 {
 	public class TraceNewTraceEvent : ITraceEvent
 	{
-		public static readonly EventType EventType = new EventType(0, "$Trace", "NewTrace", EventType.FlagImportant | EventType.FlagNoSync,
+		public static readonly EventType EventType = new(0, "$Trace", "NewTrace", EventType.FlagImportant | EventType.FlagNoSync,
 			[
 				new EventTypeField(0, 8, EventTypeField.TypeInt64, "StartCycle"),
 				new EventTypeField(8, 8, EventTypeField.TypeInt64, "CycleFrequency"),
@@ -51,7 +51,7 @@ namespace EpicGames.Tracing.UnrealInsights.Events
 
 	public class TraceThreadInfoEvent : ITraceEvent
 	{
-		public static readonly EventType EventType = new EventType(0, "$Trace", "ThreadInfo", EventType.FlagImportant | EventType.FlagMaybeHasAux | EventType.FlagNoSync,
+		public static readonly EventType EventType = new(0, "$Trace", "ThreadInfo", EventType.FlagImportant | EventType.FlagMaybeHasAux | EventType.FlagNoSync,
 			[
 				new EventTypeField(0, 4, EventTypeField.TypeInt32, "ThreadId"),
 				new EventTypeField(4, 4, EventTypeField.TypeInt32, "SystemId"),
diff --git a/Engine/Source/Programs/Shared/EpicGames.Tracing/UnrealInsights/Reader.cs b/Engine/Source/Programs/Shared/EpicGames.Tracing/UnrealInsights/Reader.cs
--- a/Engine/Source/Programs/Shared/EpicGames.Tracing/UnrealInsights/Reader.cs
+++ b/Engine/Source/Programs/Shared/EpicGames.Tracing/UnrealInsights/Reader.cs
@@ -45,7 +45,7 @@ namespace EpicGames.Tracing.UnrealInsights
 
 		public void Read(Stream stream)
 		{
-			using BinaryReader reader = new BinaryReader(stream);
+			using BinaryReader reader = new(stream);
 			StreamHeader.Deserialize(reader);
 
 			Dictionary<ushort, MemoryStream> threadStreams = Demux(reader);
@@ -54,7 +54,7 @@ namespace EpicGames.Tracing.UnrealInsights
 			void ReadThreadStream(ushort threadId)
 			{
 				MemoryStream threadStream = threadStreams[threadId];
-				using BinaryReader threadReader = new BinaryReader(threadStream);
+				using BinaryReader threadReader = new(threadStream);
 				EventsPerThread[threadId] = ReadEvents(threadId, threadReader);
 				threadReader.EnsureEntireStreamIsConsumed();
 				threadStream.Dispose();
diff --git a/Engine/Source/Programs/Shared/EpicGames.Tracing/UnrealInsights/StreamHeader.cs b/Engine/Source/Programs/Shared/EpicGames.Tracing/UnrealInsights/StreamHeader.cs
--- a/Engine/Source/Programs/Shared/EpicGames.Tracing/UnrealInsights/StreamHeader.cs
+++ b/Engine/Source/Programs/Shared/EpicGames.Tracing/UnrealInsights/StreamHeader.cs
@@ -24,7 +24,7 @@ namespace EpicGames.Tracing.UnrealInsights
 
 		public static StreamHeader Default()
 		{
-			StreamHeader handshake = new StreamHeader();
+			StreamHeader handshake = new();
 			handshake._magicFourCc = MagicTrc;
 			handshake._transportVersion = 3;
 			handshake._protocolVersion = 5;
@@ -54,7 +54,7 @@ namespace EpicGames.Tracing.UnrealInsights
 
 		public static StreamHeader Deserialize(BinaryReader reader)
 		{
-			StreamHeader header = new StreamHeader();
+			StreamHeader header = new();
 			header._magicFourCc = reader.ReadUInt32();
 			if (header._magicFourCc != MagicTrc && header._magicFourCc != MagicTrc2)
 			{
diff --git a/Engine/Source/Programs/Shared/EpicGames.Tracing/UnrealInsights/TransportPacket.cs b/Engine/Source/Programs/Shared/EpicGames.Tracing/UnrealInsights/TransportPacket.cs
--- a/Engine/Source/Programs/Shared/EpicGames.Tracing/UnrealInsights/TransportPacket.cs
+++ b/Engine/Source/Programs/Shared/EpicGames.Tracing/UnrealInsights/TransportPacket.cs
@@ -90,7 +90,7 @@ namespace EpicGames.Tracing.UnrealInsights
 
 		public static TransportPacket Create(ushort packetSize, ushort threadId)
 		{
-			TransportPacket packet = new TransportPacket();
+			TransportPacket packet = new();
 			packet.PacketSize = packetSize;
 			packet.ThreadIdAndMarkers = threadId;
 			return packet;
@@ -98,7 +98,7 @@ namespace EpicGames.Tracing.UnrealInsights
 
 		public static TransportPacket Deserialize(BinaryReader reader)
 		{
-			TransportPacket packet = new TransportPacket();
+			TransportPacket packet = new();
 			packet.PacketSize = reader.ReadUInt16();
 			packet.ThreadIdAndMarkers = reader.ReadUInt16();
 
diff --git a/Engine/Source/Programs/Shared/EpicGames.Tracing/UnrealInsights/Writer.cs b/Engine/Source/Programs/Shared/EpicGames.Tracing/UnrealInsights/Writer.cs
--- a/Engine/Source/Programs/Shared/EpicGames.Tracing/UnrealInsights/Writer.cs
+++ b/Engine/Source/Programs/Shared/EpicGames.Tracing/UnrealInsights/Writer.cs
@@ -15,22 +15,22 @@ namespace EpicGames.Tracing.UnrealInsights
 		readonly Dictionary<ushort, List<ITraceEvent>> _threadToEvents = [];
 
 		private ushort _uidCounter = PredefinedEventUid.WellKnownNum; // IDs below 16 are reserved for well-known events
-		private readonly object _lock = new object();
+		private readonly object _lock = new();
 
 		public UnrealInsightsWriter()
 		{
 			_threadToEvents[TransportPacket.ThreadIdEvents] = [];
 
-			EnterScopeEvent enterScopeEvent = new EnterScopeEvent();
+			EnterScopeEvent enterScopeEvent = new();
 			RegisterEventType(PredefinedEventUid.EnterScope, enterScopeEvent.Type);
 
-			LeaveScopeEvent leaveScopeEvent = new LeaveScopeEvent();
+			LeaveScopeEvent leaveScopeEvent = new();
 			RegisterEventType(PredefinedEventUid.LeaveScope, leaveScopeEvent.Type);
 
-			EnterScopeEvent enterScopeTimestampEvent = new EnterScopeEvent();
+			EnterScopeEvent enterScopeTimestampEvent = new();
 			RegisterEventType(PredefinedEventUid.EnterScope_T, enterScopeTimestampEvent.Type);
 
-			LeaveScopeEventTimestamp leaveScopeTimestampEvent = new LeaveScopeEventTimestamp(0);
+			LeaveScopeEventTimestamp leaveScopeTimestampEvent = new(0);
 			RegisterEventType(PredefinedEventUid.LeaveScope_T, leaveScopeTimestampEvent.Type);
 		}
 
diff --git a/Engine/Source/Programs/Shared/EpicGames.UBA/Impl/LoggerImpl.cs b/Engine/Source/Programs/Shared/EpicGames.UBA/Impl/LoggerImpl.cs
--- a/Engine/Source/Programs/Shared/EpicGames.UBA/Impl/LoggerImpl.cs
+++ b/Engine/Source/Programs/Shared/EpicGames.UBA/Impl/LoggerImpl.cs
@@ -22,7 +22,7 @@ namespace EpicGames.UBA.Impl
 		readonly BeginScopeCallback _beginScopeCallbackDelegate;
 		readonly EndScopeCallback _endScopeCallbackDelegate;
 		readonly LogCallback _logCallbackDelegate;
-		readonly object _lock = new object();
+		readonly object _lock = new();
 
 		#region DllImport
 		[DllImport("UbaHost", CharSet = CharSet.Auto)]
diff --git a/Engine/Source/Programs/Shared/EpicGames.UHT/Exporters/CodeGen/UhtCodeGenExpressions.cs b/Engine/Source/Programs/Shared/EpicGames.UHT/Exporters/CodeGen/UhtCodeGenExpressions.cs
--- a/Engine/Source/Programs/Shared/EpicGames.UHT/Exporters/CodeGen/UhtCodeGenExpressions.cs
+++ b/Engine/Source/Programs/Shared/EpicGames.UHT/Exporters/CodeGen/UhtCodeGenExpressions.cs
@@ -15,20 +15,20 @@ namespace EpicGames.UHT.Exporters.CodeGen
 	{
 		static readonly ReadOnlyMemory<char> s_tabs = String.Empty.PadLeft(100, '\t').AsMemory();
 
-		internal static TabsExpression Tabs(int count) => new TabsExpression(count, s_tabs);
+		internal static TabsExpression Tabs(int count) => new(count, s_tabs);
 
 		internal static SingletonRefExpression SingletonRef(IUhtObjectLinker provider, UhtObject? obj, UhtSingletonType type)
-			=> new SingletonRefExpression(provider, obj, type, false);
+			=> new(provider, obj, type, false);
 
 		internal static SingletonRefExpression ConstInitSingletonRef(IUhtObjectLinker provider, UhtObject? obj, bool explicitType = false)
-			=> new SingletonRefExpression(provider, obj, UhtSingletonType.ConstInit, explicitType);
+			=> new(provider, obj, UhtSingletonType.ConstInit, explicitType);
 
 		internal static UhtCppIdentifier VniModuleStatics(UhtModule module)
-			=> new UhtCppIdentifier("Z_VerseRegistration_", module.Module.Name, "_Statics");
+			=> new("Z_VerseRegistration_", module.Module.Name, "_Statics");
 
 		// NOTE: This name is generated by V_CPP_MODULE_PACKAGEDESC_NAME macro
 		internal static UhtCppIdentifier VniPackageDesc(UhtModule module)
-			=> new UhtCppIdentifier("UHTAutoModule_", module.Module.Name, "_PackageDesc");
+			=> new("UHTAutoModule_", module.Module.Name, "_PackageDesc");
 	}
 
 	#region Expression structs
diff --git a/Engine/Source/Programs/Shared/EpicGames.UHT/Exporters/CodeGen/UhtPackageCodeGeneratorCppFile.cs b/Engine/Source/Programs/Shared/EpicGames.UHT/Exporters/CodeGen/UhtPackageCodeGeneratorCppFile.cs
--- a/Engine/Source/Programs/Shared/EpicGames.UHT/Exporters/CodeGen/UhtPackageCodeGeneratorCppFile.cs
+++ b/Engine/Source/Programs/Shared/EpicGames.UHT/Exporters/CodeGen/UhtPackageCodeGeneratorCppFile.cs
@@ -370,7 +370,7 @@ namespace EpicGames.UHT.Exporters.CodeGen
 			return (start, start);
 		}
 
-		private static SimpleExpression ModuleLinkableObjects(UhtModule module) => new SimpleExpression("Z_LinkableObjects_", module.ShortName);
+		private static SimpleExpression ModuleLinkableObjects(UhtModule module) => new("Z_LinkableObjects_", module.ShortName);
 		private static SimpleExpression ModuleLinkedModules(UhtModule module) => module.Linking!.ReferencedModules.Count == 0
 			? new SimpleExpression("", "{}")
 			: new SimpleExpression("Z_LinkedModules_", module.ShortName);
diff --git a/Engine/Source/Programs/Shared/EpicGames.UHT/Parsers/UhtHeaderFileParser.cs b/Engine/Source/Programs/Shared/EpicGames.UHT/Parsers/UhtHeaderFileParser.cs
--- a/Engine/Source/Programs/Shared/EpicGames.UHT/Parsers/UhtHeaderFileParser.cs
+++ b/Engine/Source/Programs/Shared/EpicGames.UHT/Parsers/UhtHeaderFileParser.cs
@@ -885,7 +885,7 @@ namespace EpicGames.UHT.Parsers
 						HeaderFile.LogError("#include found after .generated.h file - the .generated.h file should always be the last #include in a header");
 					}
 
-					StringView includeNameString = new StringView();
+					StringView includeNameString = new();
 
 					UhtToken includeName = lineTokenReader.GetToken();
 					if (includeName.IsConstString())
@@ -1108,7 +1108,7 @@ namespace EpicGames.UHT.Parsers
 			majorToken.GetConstInt(out int major);
 			minorToken.GetConstInt(out int minor);
 			patchToken.GetConstInt(out int patch);
-			EngineVersion defineVersion = new EngineVersion(major, minor, patch);
+			EngineVersion defineVersion = new(major, minor, patch);
 
 			bool? conditionPasses = null;
 			if (define.IsValue("UE_VERSION_NEWER_THAN_OR_EQUAL"))
diff --git a/Engine/Source/Programs/Shared/EpicGames.UHT/Types/Properties/UhtStructProperty.cs b/Engine/Source/Programs/Shared/EpicGames.UHT/Types/Properties/UhtStructProperty.cs
--- a/Engine/Source/Programs/Shared/EpicGames.UHT/Types/Properties/UhtStructProperty.cs
+++ b/Engine/Source/Programs/Shared/EpicGames.UHT/Types/Properties/UhtStructProperty.cs
@@ -721,7 +721,7 @@ namespace EpicGames.UHT.Types
 			}
 			identifier.RedirectTypeIdentifier(args.Config);
 
-			UhtStructProperty instancedStructProperty = new UhtStructProperty(propertySettings, args.Session.FStateTreePropertyRef!);
+			UhtStructProperty instancedStructProperty = new(propertySettings, args.Session.FStateTreePropertyRef!);
 
 			// TStateTreePropertyRef supports UStructs, UClasses, enums and primitive types.
 			UhtType? foundType = propertySettings.Outer.FindType(UhtFindOptions.SourceName | UhtFindOptions.TypesMask, identifier);
diff --git a/Engine/Source/Programs/Shared/EpicGames.UHT/Utils/UhtSession.cs b/Engine/Source/Programs/Shared/EpicGames.UHT/Utils/UhtSession.cs
--- a/Engine/Source/Programs/Shared/EpicGames.UHT/Utils/UhtSession.cs
+++ b/Engine/Source/Programs/Shared/EpicGames.UHT/Utils/UhtSession.cs
@@ -330,10 +330,10 @@ namespace EpicGames.UHT.Utils
 		public string GetShortestIncludePath(UHTManifest.Module moduleObj, string filePath)
 		{
 			string? relativePath = null;
-			FileReference fileRef = new FileReference(filePath);
+			FileReference fileRef = new(filePath);
 			foreach (string includePath in moduleObj.IncludePaths)
 			{
-				DirectoryReference includePathRef = new DirectoryReference(includePath);
+				DirectoryReference includePathRef = new(includePath);
 				if (fileRef.IsUnderDirectory(new(includePath)))
 				{
 					if (relativePath == null || (filePath.Length - includePath.Length - 1) < relativePath.Length)
@@ -2017,7 +2017,7 @@ namespace EpicGames.UHT.Utils
 				throw new UhtException(String.Format("Error reading engine version header '{0}': {1}", versionHeader, ex));
 			}
 
-			EngineVersion loadedVersion = new EngineVersion();
+			EngineVersion loadedVersion = new();
 			bool gotMajor = false, gotMinor = false, gotPatch = false;
 			Regex versionRegex = GetEngineVersionFragmentRegex();
 			foreach (string line in lines)
