UnrealBuildTool: Fix some additional analyze warnings

#rnx

#Synced-CL 50888098

#Synced-CL 50888098

diff --git a/Engine/Source/Programs/Shared/UnrealBuildTool.Tests/ConfigFileTests.cs b/Engine/Source/Programs/Shared/UnrealBuildTool.Tests/ConfigFileTests.cs
--- a/Engine/Source/Programs/Shared/UnrealBuildTool.Tests/ConfigFileTests.cs
+++ b/Engine/Source/Programs/Shared/UnrealBuildTool.Tests/ConfigFileTests.cs
@@ -23,7 +23,7 @@ namespace UnrealBuildTool.Tests
 				try
 				{
 					TemporaryFile = Path.Combine(System.IO.Path.GetTempPath(), Guid.NewGuid().ToString() + "_" + fileName);
-					using (StreamWriter writer = new StreamWriter(TemporaryFile))
+					using (StreamWriter writer = new(TemporaryFile))
 					{
 						writer.Write(contents);
 					}
@@ -54,14 +54,14 @@ namespace UnrealBuildTool.Tests
 
 			string testContents = String.Format("[{0}]\r\n{1}={2}", TestSection, TestValueKey, TestValue);
 
-			using (TestConfigFile temporaryConfigFile = new TestConfigFile(TEST_FILE_NAME, testContents))
+			using (TestConfigFile temporaryConfigFile = new(TEST_FILE_NAME, testContents))
 			{
 				string? testFilePath = temporaryConfigFile.TemporaryFile;
 
 				Assert.IsFalse(String.IsNullOrEmpty(testFilePath), "Unable to generate a test file.");
 
 				ConfigFile configFile;
-				FileReference configFileLocation = new FileReference(testFilePath);
+				FileReference configFileLocation = new(testFilePath);
 
 				Assert.IsTrue(FileReference.Exists(configFileLocation), "Unable to acquire a test file.");
 				
diff --git a/Engine/Source/Programs/Shared/UnrealBuildTool.Tests/CppCompileWarningsTests.cs b/Engine/Source/Programs/Shared/UnrealBuildTool.Tests/CppCompileWarningsTests.cs
--- a/Engine/Source/Programs/Shared/UnrealBuildTool.Tests/CppCompileWarningsTests.cs
+++ b/Engine/Source/Programs/Shared/UnrealBuildTool.Tests/CppCompileWarningsTests.cs
@@ -57,7 +57,7 @@ namespace UnrealBuildTool.Tests
 		[TestMethod]
 		public void TestDefaultBuildSettings()
 		{
-			CppCompileWarnings defaultBuildSettings = new CppCompileWarnings();
+			CppCompileWarnings defaultBuildSettings = new();
 
 			Assert.AreEqual(WarningLevel.Warning, defaultBuildSettings.ShadowVariableWarningLevel);
 			Assert.AreEqual(WarningLevel.Off, defaultBuildSettings.UndefinedIdentifierWarningLevel);
@@ -68,13 +68,13 @@ namespace UnrealBuildTool.Tests
 		[TestMethod]
 		public void TestClonedMutableReturn()
 		{
-			CppCompileWarnings defaultBuildSettings = new CppCompileWarnings();
+			CppCompileWarnings defaultBuildSettings = new();
 			defaultBuildSettings.ShadowVariableWarningLevel = WarningLevel.Error;
 			defaultBuildSettings.UndefinedIdentifierWarningLevel = WarningLevel.Warning;
 			defaultBuildSettings.UnsafeTypeCastWarningLevel = WarningLevel.Error;
 			defaultBuildSettings.SwitchUnhandledEnumeratorWarningLevel = WarningLevel.Warning;
 
-			ReadOnlyCppCompileWarnings wrapper = new ReadOnlyCppCompileWarnings(defaultBuildSettings);
+			ReadOnlyCppCompileWarnings wrapper = new(defaultBuildSettings);
 
 			CppCompileWarnings cloned = wrapper.CloneAsWriteable();
 
@@ -102,10 +102,10 @@ namespace UnrealBuildTool.Tests
 		[TestMethod]
 		public void TestContextBuildSettingsOverride()
 		{
-			TestBuildSettingsProvider buildContextProvider = new TestBuildSettingsProvider();
+			TestBuildSettingsProvider buildContextProvider = new();
 			buildContextProvider._testBuildSettings = BuildSettingsVersion.V1;
 
-			CppCompileWarnings newWarnings = new CppCompileWarnings(buildContextProvider, null);
+			CppCompileWarnings newWarnings = new(buildContextProvider, null);
 
 			// Shadow Variable build setting defaults
 			{
@@ -126,10 +126,10 @@ namespace UnrealBuildTool.Tests
 		[TestMethod]
 		public void TestParentSettings()
 		{
-			TestBuildSettingsProvider buildContextProvider = new TestBuildSettingsProvider();
+			TestBuildSettingsProvider buildContextProvider = new();
 			buildContextProvider._testBuildSettings = BuildSettingsVersion.V1;
-			CppCompileWarnings parentWarnings = new CppCompileWarnings(buildContextProvider, null);
-			CppCompileWarnings newWarnings = new CppCompileWarnings(buildContextProvider, null, new ReadOnlyCppCompileWarnings(parentWarnings));
+			CppCompileWarnings parentWarnings = new(buildContextProvider, null);
+			CppCompileWarnings newWarnings = new(buildContextProvider, null, new ReadOnlyCppCompileWarnings(parentWarnings));
 
 			// Assert invariants
 			{
@@ -153,10 +153,10 @@ namespace UnrealBuildTool.Tests
 		[TestMethod]
 		public void TestReparentOperations()
 		{
-			TestBuildSettingsProvider buildContextProvider = new TestBuildSettingsProvider();
+			TestBuildSettingsProvider buildContextProvider = new();
 			buildContextProvider._testBuildSettings = BuildSettingsVersion.V1;
-			CppCompileWarnings parentWarnings = new CppCompileWarnings(buildContextProvider, null);
-			CppCompileWarnings newWarnings = new CppCompileWarnings(buildContextProvider, null, new ReadOnlyCppCompileWarnings(parentWarnings));
+			CppCompileWarnings parentWarnings = new(buildContextProvider, null);
+			CppCompileWarnings newWarnings = new(buildContextProvider, null, new ReadOnlyCppCompileWarnings(parentWarnings));
 
 			// Assert invariant parent values
 			{
@@ -170,7 +170,7 @@ namespace UnrealBuildTool.Tests
 			}
 
 			// Assert reparenting
-			CppCompileWarnings parentWarnings2 = new CppCompileWarnings(buildContextProvider, null);
+			CppCompileWarnings parentWarnings2 = new(buildContextProvider, null);
 			{
 				newWarnings.AddParent(parentWarnings2);
 				parentWarnings2.ShadowVariableWarningLevel = WarningLevel.Warning;
@@ -208,10 +208,10 @@ namespace UnrealBuildTool.Tests
 		[TestMethod]
 		public void TestChildOverrides()
 		{
-			TestBuildSettingsProvider buildContextProvider = new TestBuildSettingsProvider();
+			TestBuildSettingsProvider buildContextProvider = new();
 			buildContextProvider._testBuildSettings = BuildSettingsVersion.V1;
-			CppCompileWarnings parentWarnings = new CppCompileWarnings(buildContextProvider, null);
-			CppCompileWarnings newWarnings = new CppCompileWarnings(buildContextProvider, null, new ReadOnlyCppCompileWarnings(parentWarnings));
+			CppCompileWarnings parentWarnings = new(buildContextProvider, null);
+			CppCompileWarnings newWarnings = new(buildContextProvider, null, new ReadOnlyCppCompileWarnings(parentWarnings));
 
 			parentWarnings.ShadowVariableWarningLevel = WarningLevel.Error;
 			parentWarnings.UndefinedIdentifierWarningLevel = WarningLevel.Error;
@@ -244,22 +244,22 @@ namespace UnrealBuildTool.Tests
 			// Generate ephemeral TargetRules & uproject files
 			using (SafeTestDirectory testRootDirectory = TestDirs.Create())
 			using (SafeTestDirectory testSourceDirectory = SafeTestDirectory.CreateTestDirectory(Path.Combine(testRootDirectory.TemporaryDirectory, "Source")))
-			using (SafeTestFile testEditorTargetTestFile = new SafeTestFile(File.ReadAllText(Path.Combine(testStubsDirectory.FullName, "TargetStubs", "TestStubEditor.Target.ubttest")), "TestStubEditor.Target.cs", testSourceDirectory.TemporaryDirectory))
-			using (SafeTestFile testTargetTestFile = new SafeTestFile(File.ReadAllText(Path.Combine(testStubsDirectory.FullName, "TargetStubs", "TestStub.Target.ubttest")), "TestStub.Target.cs", testSourceDirectory.TemporaryDirectory))
-			using (SafeTestFile testUprojectTestFile = new SafeTestFile(File.ReadAllText(Path.Combine(testStubsDirectory.FullName, "UprojectStubs", "TestStub.uproject.ubttest")), "TestStub.uproject", testRootDirectory.TemporaryDirectory))
+			using (SafeTestFile testEditorTargetTestFile = new(File.ReadAllText(Path.Combine(testStubsDirectory.FullName, "TargetStubs", "TestStubEditor.Target.ubttest")), "TestStubEditor.Target.cs", testSourceDirectory.TemporaryDirectory))
+			using (SafeTestFile testTargetTestFile = new(File.ReadAllText(Path.Combine(testStubsDirectory.FullName, "TargetStubs", "TestStub.Target.ubttest")), "TestStub.Target.cs", testSourceDirectory.TemporaryDirectory))
+			using (SafeTestFile testUprojectTestFile = new(File.ReadAllText(Path.Combine(testStubsDirectory.FullName, "UprojectStubs", "TestStub.uproject.ubttest")), "TestStub.uproject", testRootDirectory.TemporaryDirectory))
 			{
-				CppCompileEnvironment newCompileEnvironment = new CppCompileEnvironment(UnrealTargetPlatform.Win64, CppConfiguration.Development, new UnrealArchitectures(UnrealArch.X64), null!);
+				CppCompileEnvironment newCompileEnvironment = new(UnrealTargetPlatform.Win64, CppConfiguration.Development, new UnrealArchitectures(UnrealArch.X64), null!);
 
-				FileReference testUproject = new FileReference(testUprojectTestFile.TemporaryFile);
+				FileReference testUproject = new(testUprojectTestFile.TemporaryFile);
 				IEnumerable<string> args = null!;
 
 				RulesAssembly generatedTestGameRulesAssembly = RulesCompiler.CreateTargetRulesAssembly(testUproject, "TestStub", false, false, false, null, false, new TestLogger());
 				TargetRules testTargetRulesNonWin64 = generatedTestGameRulesAssembly.CreateTargetRules("TestStub", UnrealTargetPlatform.Android, UnrealTargetConfiguration.Development, null, testUproject, null, new TestLogger());
-				CppCompileWarnings specializedWarnings = new CppCompileWarnings(testTargetRulesNonWin64, null);
+				CppCompileWarnings specializedWarnings = new(testTargetRulesNonWin64, null);
 
 				// Verify the ordering
 				{
-					VersionNumber newVersionNumber = new VersionNumber(18, 9, 9);
+					VersionNumber newVersionNumber = new(18, 9, 9);
 					specializedWarnings.DeprecatedCopyWarningLevel = WarningLevel.Off;
 					specializedWarnings.DeprecatedCopyWithUserProvidedCopyWarningLevel = WarningLevel.Error;
 					args = specializedWarnings.GenerateWarningCommandLineArgs(newCompileEnvironment, typeof(ClangToolChain), newVersionNumber);
@@ -280,19 +280,19 @@ namespace UnrealBuildTool.Tests
 			// Generate ephemeral TargetRules & uproject files
 			using (SafeTestDirectory testRootDirectory = TestDirs.Create())
 			using (SafeTestDirectory testSourceDirectory = SafeTestDirectory.CreateTestDirectory(Path.Combine(testRootDirectory.TemporaryDirectory, "Source")))
-			using (SafeTestFile testEditorTargetTestFile = new SafeTestFile(File.ReadAllText(Path.Combine(testStubsDirectory.FullName, "TargetStubs", "TestStubEditor.Target.ubttest")), "TestStubEditor.Target.cs", testSourceDirectory.TemporaryDirectory))
-			using (SafeTestFile testTargetTestFile = new SafeTestFile(File.ReadAllText(Path.Combine(testStubsDirectory.FullName, "TargetStubs", "TestStub.Target.ubttest")), "TestStub.Target.cs", testSourceDirectory.TemporaryDirectory))
-			using (SafeTestFile testUprojectTestFile = new SafeTestFile(File.ReadAllText(Path.Combine(testStubsDirectory.FullName, "UprojectStubs", "TestStub.uproject.ubttest")), "TestStub.uproject", testRootDirectory.TemporaryDirectory))
+			using (SafeTestFile testEditorTargetTestFile = new(File.ReadAllText(Path.Combine(testStubsDirectory.FullName, "TargetStubs", "TestStubEditor.Target.ubttest")), "TestStubEditor.Target.cs", testSourceDirectory.TemporaryDirectory))
+			using (SafeTestFile testTargetTestFile = new(File.ReadAllText(Path.Combine(testStubsDirectory.FullName, "TargetStubs", "TestStub.Target.ubttest")), "TestStub.Target.cs", testSourceDirectory.TemporaryDirectory))
+			using (SafeTestFile testUprojectTestFile = new(File.ReadAllText(Path.Combine(testStubsDirectory.FullName, "UprojectStubs", "TestStub.uproject.ubttest")), "TestStub.uproject", testRootDirectory.TemporaryDirectory))
 			{
-				CppCompileEnvironment newCompileEnvironment = new CppCompileEnvironment(UnrealTargetPlatform.Win64, CppConfiguration.Development, new UnrealArchitectures(UnrealArch.X64), null!);
+				CppCompileEnvironment newCompileEnvironment = new(UnrealTargetPlatform.Win64, CppConfiguration.Development, new UnrealArchitectures(UnrealArch.X64), null!);
 
-				FileReference testUproject = new FileReference(testUprojectTestFile.TemporaryFile);
+				FileReference testUproject = new(testUprojectTestFile.TemporaryFile);
 
 				// Generate rules assembly & subsequent targets
 				RulesAssembly generatedTestRulesAssembly = RulesCompiler.CreateTargetRulesAssembly(testUproject, "TestStubEditor", false, false, false, null, false, new TestLogger());
 				TargetRules testTargetRules = generatedTestRulesAssembly.CreateTargetRules("TestStubEditor", UnrealTargetPlatform.Win64, UnrealTargetConfiguration.Development, null, testUproject, null, new TestLogger());
 				testTargetRules.WindowsPlatform.Compiler = WindowsCompiler.VisualStudio2022;
-				CppCompileWarnings baseWarnings = new CppCompileWarnings(testTargetRules, null);
+				CppCompileWarnings baseWarnings = new(testTargetRules, null);
 
 				IEnumerable<string> args = null!;
 
@@ -310,7 +310,7 @@ namespace UnrealBuildTool.Tests
 					Assert.IsTrue(args.Count() == 1);
 					Assert.IsTrue(args.Contains("/w44668"));
 
-					CppCompileEnvironment preprocessOnlyCompileEnvironment = new CppCompileEnvironment(newCompileEnvironment);
+					CppCompileEnvironment preprocessOnlyCompileEnvironment = new(newCompileEnvironment);
 					preprocessOnlyCompileEnvironment.bPreprocessOnly = true;
 					args = baseWarnings.GenerateWarningCommandLineArgs(preprocessOnlyCompileEnvironment, typeof(VCToolChain));
 
@@ -321,9 +321,9 @@ namespace UnrealBuildTool.Tests
 
 				{
 					// DeterministicWarningLevel
-					VersionNumber newVersionNumber = new VersionNumber(18, 9, 9);
+					VersionNumber newVersionNumber = new(18, 9, 9);
 
-					CppCompileEnvironment deterministicCompileEnvironment = new CppCompileEnvironment(UnrealTargetPlatform.Win64, CppConfiguration.Debug, new UnrealArchitectures(UnrealArch.X64), null!);
+					CppCompileEnvironment deterministicCompileEnvironment = new(UnrealTargetPlatform.Win64, CppConfiguration.Debug, new UnrealArchitectures(UnrealArch.X64), null!);
 					deterministicCompileEnvironment.bDeterministic = false;
 
 					args = baseWarnings.GenerateWarningCommandLineArgs(deterministicCompileEnvironment, typeof(VCToolChain), newVersionNumber);
@@ -364,19 +364,19 @@ namespace UnrealBuildTool.Tests
 			// Generate ephemeral TargetRules & uproject files
 			using (SafeTestDirectory testRootDirectory = TestDirs.Create())
 			using (SafeTestDirectory testSourceDirectory = SafeTestDirectory.CreateTestDirectory(Path.Combine(testRootDirectory.TemporaryDirectory, "Source")))
-			using (SafeTestFile testEditorTargetTestFile = new SafeTestFile(File.ReadAllText(Path.Combine(testStubsDirectory.FullName, "TargetStubs", "TestStubEditor.Target.ubttest")), "TestStubEditor.Target.cs", testSourceDirectory.TemporaryDirectory))
-			using (SafeTestFile testTargetTestFile = new SafeTestFile(File.ReadAllText(Path.Combine(testStubsDirectory.FullName, "TargetStubs", "TestStub.Target.ubttest")), "TestStub.Target.cs", testSourceDirectory.TemporaryDirectory))
-			using (SafeTestFile testUprojectTestFile = new SafeTestFile(File.ReadAllText(Path.Combine(testStubsDirectory.FullName, "UprojectStubs", "TestStub.uproject.ubttest")), "TestStub.uproject", testRootDirectory.TemporaryDirectory))
+			using (SafeTestFile testEditorTargetTestFile = new(File.ReadAllText(Path.Combine(testStubsDirectory.FullName, "TargetStubs", "TestStubEditor.Target.ubttest")), "TestStubEditor.Target.cs", testSourceDirectory.TemporaryDirectory))
+			using (SafeTestFile testTargetTestFile = new(File.ReadAllText(Path.Combine(testStubsDirectory.FullName, "TargetStubs", "TestStub.Target.ubttest")), "TestStub.Target.cs", testSourceDirectory.TemporaryDirectory))
+			using (SafeTestFile testUprojectTestFile = new(File.ReadAllText(Path.Combine(testStubsDirectory.FullName, "UprojectStubs", "TestStub.uproject.ubttest")), "TestStub.uproject", testRootDirectory.TemporaryDirectory))
 			{
-				CppCompileEnvironment newCompileEnvironment = new CppCompileEnvironment(UnrealTargetPlatform.Win64, CppConfiguration.Development, new UnrealArchitectures(UnrealArch.X64), null!);
+				CppCompileEnvironment newCompileEnvironment = new(UnrealTargetPlatform.Win64, CppConfiguration.Development, new UnrealArchitectures(UnrealArch.X64), null!);
 
-				FileReference testUproject = new FileReference(testUprojectTestFile.TemporaryFile);
+				FileReference testUproject = new(testUprojectTestFile.TemporaryFile);
 
 				// Generate rules assembly & subsequent targets
 				RulesAssembly generatedTestRulesAssembly = RulesCompiler.CreateTargetRulesAssembly(testUproject, "TestStubEditor", false, false, false, null, false, new TestLogger());
 				TargetRules testTargetRules = generatedTestRulesAssembly.CreateTargetRules("TestStubEditor", UnrealTargetPlatform.Win64, UnrealTargetConfiguration.Development, null, testUproject, null, new TestLogger());
 				testTargetRules.WindowsPlatform.Compiler = WindowsCompiler.VisualStudio2022;
-				CppCompileWarnings baseWarnings = new CppCompileWarnings(testTargetRules, null);
+				CppCompileWarnings baseWarnings = new(testTargetRules, null);
 
 				IEnumerable<string> args = null!;
 				{
@@ -384,9 +384,9 @@ namespace UnrealBuildTool.Tests
 					{
 						testTargetRules.WindowsPlatform.Compiler = WindowsCompiler.Clang;
 						baseWarnings.DeterministicWarningLevel = WarningLevel.Off;
-						VersionNumber newVersionNumber = new VersionNumber(18, 9, 9);
+						VersionNumber newVersionNumber = new(18, 9, 9);
 
-						CppCompileEnvironment deterministicCompileEnvironment = new CppCompileEnvironment(UnrealTargetPlatform.Win64, CppConfiguration.Debug, new UnrealArchitectures(UnrealArch.X64), null!);
+						CppCompileEnvironment deterministicCompileEnvironment = new(UnrealTargetPlatform.Win64, CppConfiguration.Debug, new UnrealArchitectures(UnrealArch.X64), null!);
 						deterministicCompileEnvironment.bDeterministic = false;
 
 						args = baseWarnings.GenerateWarningCommandLineArgs(deterministicCompileEnvironment, typeof(ClangToolChain), newVersionNumber);
@@ -422,9 +422,9 @@ namespace UnrealBuildTool.Tests
 					{
 						testTargetRules.WindowsPlatform.Compiler = WindowsCompiler.Clang;
 						baseWarnings.ProfileInstructWarningLevel = WarningLevel.Off;
-						VersionNumber newVersionNumber = new VersionNumber(18, 9, 9);
+						VersionNumber newVersionNumber = new(18, 9, 9);
 
-						CppCompileEnvironment pgoOptimizedCompileEnvironment = new CppCompileEnvironment(UnrealTargetPlatform.Win64, CppConfiguration.Debug, new UnrealArchitectures(UnrealArch.X64), null!);
+						CppCompileEnvironment pgoOptimizedCompileEnvironment = new(UnrealTargetPlatform.Win64, CppConfiguration.Debug, new UnrealArchitectures(UnrealArch.X64), null!);
 						pgoOptimizedCompileEnvironment.bPGOOptimize = false;
 
 						args = baseWarnings.GenerateWarningCommandLineArgs(pgoOptimizedCompileEnvironment, typeof(ClangToolChain), newVersionNumber);
@@ -446,8 +446,8 @@ namespace UnrealBuildTool.Tests
 						testTargetRules.WindowsPlatform.Compiler = WindowsCompiler.Clang;
 						baseWarnings.CastFunctionTypeMismatchWarningLevel = WarningLevel.Off;
 
-						VersionNumber newVersionNumber = new VersionNumber(19, 0, 0);
-						CppCompileEnvironment androidCompileEnvironment = new CppCompileEnvironment(UnrealTargetPlatform.Android, CppConfiguration.Development, new UnrealArchitectures(UnrealArch.X64), null!);
+						VersionNumber newVersionNumber = new(19, 0, 0);
+						CppCompileEnvironment androidCompileEnvironment = new(UnrealTargetPlatform.Android, CppConfiguration.Development, new UnrealArchitectures(UnrealArch.X64), null!);
 						args = baseWarnings.GenerateWarningCommandLineArgs(androidCompileEnvironment, typeof(ClangToolChain), newVersionNumber);
 
 						Assert.IsFalse(args.Any());
@@ -470,11 +470,11 @@ namespace UnrealBuildTool.Tests
 					{
 						testTargetRules.WindowsPlatform.Compiler = WindowsCompiler.Clang;
 						baseWarnings.InvalidUnevaluatedStringWarningLevel = WarningLevel.Off;
-						CppCompileEnvironment androidCompileEnvironment = new CppCompileEnvironment(UnrealTargetPlatform.Android, CppConfiguration.Development, new UnrealArchitectures(UnrealArch.X64), null!);
+						CppCompileEnvironment androidCompileEnvironment = new(UnrealTargetPlatform.Android, CppConfiguration.Development, new UnrealArchitectures(UnrealArch.X64), null!);
 
 						// Exclusion cases
 						{
-							VersionNumber newVersionNumber = new VersionNumber(17, 0, 2);
+							VersionNumber newVersionNumber = new(17, 0, 2);
 							args = baseWarnings.GenerateWarningCommandLineArgs(androidCompileEnvironment, typeof(ClangToolChain), newVersionNumber);
 
 							Assert.IsTrue(args.Count() == 1 && args.Contains("-Wno-shadow"));
@@ -493,7 +493,7 @@ namespace UnrealBuildTool.Tests
 
 						// Inclusion cases
 						{
-							VersionNumber newVersionNumber = new VersionNumber(17, 1);
+							VersionNumber newVersionNumber = new(17, 1);
 							androidCompileEnvironment.CppStandard = CppStandardVersion.Cpp20;
 
 							args = baseWarnings.GenerateWarningCommandLineArgs(androidCompileEnvironment, typeof(ClangToolChain), newVersionNumber);
@@ -514,10 +514,10 @@ namespace UnrealBuildTool.Tests
 					{
 						testTargetRules.WindowsPlatform.Compiler = WindowsCompiler.Clang;
 						baseWarnings.UnusedValueWarningLevel = WarningLevel.Off;
-						VersionNumber newVersionNumber = new VersionNumber(18, 9, 9);
+						VersionNumber newVersionNumber = new(18, 9, 9);
 
-						CppCompileEnvironment invalidCompileEnvironment = new CppCompileEnvironment(UnrealTargetPlatform.Win64, CppConfiguration.Debug, new UnrealArchitectures(UnrealArch.X64), null!);
-						CppCompileEnvironment validCompileEnvironment = new CppCompileEnvironment(UnrealTargetPlatform.Win64, CppConfiguration.Shipping, new UnrealArchitectures(UnrealArch.X64), null!);
+						CppCompileEnvironment invalidCompileEnvironment = new(UnrealTargetPlatform.Win64, CppConfiguration.Debug, new UnrealArchitectures(UnrealArch.X64), null!);
+						CppCompileEnvironment validCompileEnvironment = new(UnrealTargetPlatform.Win64, CppConfiguration.Shipping, new UnrealArchitectures(UnrealArch.X64), null!);
 
 						args = baseWarnings.GenerateWarningCommandLineArgs(invalidCompileEnvironment, typeof(ClangToolChain), newVersionNumber, StaticAnalyzer.None);
 						Assert.IsFalse(args.Any());
@@ -541,27 +541,27 @@ namespace UnrealBuildTool.Tests
 			// Generate ephemeral TargetRules & uproject files
 			using (SafeTestDirectory testRootDirectory = TestDirs.Create())
 			using (SafeTestDirectory testSourceDirectory = SafeTestDirectory.CreateTestDirectory(Path.Combine(testRootDirectory.TemporaryDirectory, "Source")))
-			using (SafeTestFile testEditorTargetTestFile = new SafeTestFile(File.ReadAllText(Path.Combine(testStubsDirectory.FullName, "TargetStubs", "TestStubEditor.Target.ubttest")), "TestStubEditor.Target.cs", testSourceDirectory.TemporaryDirectory))
-			using (SafeTestFile testTargetTestFile = new SafeTestFile(File.ReadAllText(Path.Combine(testStubsDirectory.FullName, "TargetStubs", "TestStub.Target.ubttest")), "TestStub.Target.cs", testSourceDirectory.TemporaryDirectory))
-			using (SafeTestFile testUprojectTestFile = new SafeTestFile(File.ReadAllText(Path.Combine(testStubsDirectory.FullName, "UprojectStubs", "TestStub.uproject.ubttest")), "TestStub.uproject", testRootDirectory.TemporaryDirectory))
+			using (SafeTestFile testEditorTargetTestFile = new(File.ReadAllText(Path.Combine(testStubsDirectory.FullName, "TargetStubs", "TestStubEditor.Target.ubttest")), "TestStubEditor.Target.cs", testSourceDirectory.TemporaryDirectory))
+			using (SafeTestFile testTargetTestFile = new(File.ReadAllText(Path.Combine(testStubsDirectory.FullName, "TargetStubs", "TestStub.Target.ubttest")), "TestStub.Target.cs", testSourceDirectory.TemporaryDirectory))
+			using (SafeTestFile testUprojectTestFile = new(File.ReadAllText(Path.Combine(testStubsDirectory.FullName, "UprojectStubs", "TestStub.uproject.ubttest")), "TestStub.uproject", testRootDirectory.TemporaryDirectory))
 			{
-				FileReference testUproject = new FileReference(testUprojectTestFile.TemporaryFile);
+				FileReference testUproject = new(testUprojectTestFile.TemporaryFile);
 				RulesAssembly generatedTestRulesAssembly = RulesCompiler.CreateTargetRulesAssembly(testUproject, "TestStubEditor", false, false, false, null, false, new TestLogger());
 				TargetRules testTargetRules = generatedTestRulesAssembly.CreateTargetRules("TestStubEditor", UnrealTargetPlatform.Win64, UnrealTargetConfiguration.Development, null, testUproject, null, new TestLogger());
 				testTargetRules.WindowsPlatform.Compiler = WindowsCompiler.VisualStudio2022;
-				CppCompileWarnings baseWarnings = new CppCompileWarnings(testTargetRules, null);
+				CppCompileWarnings baseWarnings = new(testTargetRules, null);
 
-				CppCompileEnvironment newCompileEnvironment = new CppCompileEnvironment(UnrealTargetPlatform.Win64, CppConfiguration.Development, new UnrealArchitectures(UnrealArch.X64), null!);
+				CppCompileEnvironment newCompileEnvironment = new(UnrealTargetPlatform.Win64, CppConfiguration.Development, new UnrealArchitectures(UnrealArch.X64), null!);
 				IEnumerable<string> args = null!;
 
 				// ConfigurationFilterAttribute
 				{
 					testTargetRules.WindowsPlatform.Compiler = WindowsCompiler.Clang;
 					baseWarnings.UnusedValueWarningLevel = WarningLevel.Off;
-					VersionNumber newVersionNumber = new VersionNumber(18, 9, 9);
+					VersionNumber newVersionNumber = new(18, 9, 9);
 
-					CppCompileEnvironment invalidCompileEnvironment = new CppCompileEnvironment(UnrealTargetPlatform.Win64, CppConfiguration.Debug, new UnrealArchitectures(UnrealArch.X64), null!);
-					CppCompileEnvironment validCompileEnvironment = new CppCompileEnvironment(UnrealTargetPlatform.Win64, CppConfiguration.Shipping, new UnrealArchitectures(UnrealArch.X64), null!);
+					CppCompileEnvironment invalidCompileEnvironment = new(UnrealTargetPlatform.Win64, CppConfiguration.Debug, new UnrealArchitectures(UnrealArch.X64), null!);
+					CppCompileEnvironment validCompileEnvironment = new(UnrealTargetPlatform.Win64, CppConfiguration.Shipping, new UnrealArchitectures(UnrealArch.X64), null!);
 
 					args = baseWarnings.GenerateWarningCommandLineArgs(invalidCompileEnvironment, typeof(ClangToolChain), newVersionNumber);
 					Assert.IsFalse(args.Any());
@@ -584,26 +584,26 @@ namespace UnrealBuildTool.Tests
 			// Generate ephemeral TargetRules & uproject files
 			using (SafeTestDirectory testRootDirectory = TestDirs.Create())
 			using (SafeTestDirectory testSourceDirectory = SafeTestDirectory.CreateTestDirectory(Path.Combine(testRootDirectory.TemporaryDirectory, "Source")))
-			using (SafeTestFile testEditorTargetTestFile = new SafeTestFile(File.ReadAllText(Path.Combine(testStubsDirectory.FullName, "TargetStubs", "TestStubEditor.Target.ubttest")), "TestStubEditor.Target.cs", testSourceDirectory.TemporaryDirectory))
-			using (SafeTestFile testTargetTestFile = new SafeTestFile(File.ReadAllText(Path.Combine(testStubsDirectory.FullName, "TargetStubs", "TestStub.Target.ubttest")), "TestStub.Target.cs", testSourceDirectory.TemporaryDirectory))
-			using (SafeTestFile testUprojectTestFile = new SafeTestFile(File.ReadAllText(Path.Combine(testStubsDirectory.FullName, "UprojectStubs", "TestStub.uproject.ubttest")), "TestStub.uproject", testRootDirectory.TemporaryDirectory))
+			using (SafeTestFile testEditorTargetTestFile = new(File.ReadAllText(Path.Combine(testStubsDirectory.FullName, "TargetStubs", "TestStubEditor.Target.ubttest")), "TestStubEditor.Target.cs", testSourceDirectory.TemporaryDirectory))
+			using (SafeTestFile testTargetTestFile = new(File.ReadAllText(Path.Combine(testStubsDirectory.FullName, "TargetStubs", "TestStub.Target.ubttest")), "TestStub.Target.cs", testSourceDirectory.TemporaryDirectory))
+			using (SafeTestFile testUprojectTestFile = new(File.ReadAllText(Path.Combine(testStubsDirectory.FullName, "UprojectStubs", "TestStub.uproject.ubttest")), "TestStub.uproject", testRootDirectory.TemporaryDirectory))
 			{
-				FileReference testUproject = new FileReference(testUprojectTestFile.TemporaryFile);
+				FileReference testUproject = new(testUprojectTestFile.TemporaryFile);
 				RulesAssembly generatedTestRulesAssembly = RulesCompiler.CreateTargetRulesAssembly(testUproject, "TestStubEditor", false, false, false, null, false, new TestLogger());
 				TargetRules testTargetRules = generatedTestRulesAssembly.CreateTargetRules("TestStubEditor", UnrealTargetPlatform.Win64, UnrealTargetConfiguration.Development, null, testUproject, null, new TestLogger());
 				testTargetRules.WindowsPlatform.Compiler = WindowsCompiler.VisualStudio2022;
-				CppCompileWarnings baseWarnings = new CppCompileWarnings(testTargetRules, null);
+				CppCompileWarnings baseWarnings = new(testTargetRules, null);
 
-				CppCompileEnvironment newCompileEnvironment = new CppCompileEnvironment(UnrealTargetPlatform.Win64, CppConfiguration.Development, new UnrealArchitectures(UnrealArch.X64), null!);
+				CppCompileEnvironment newCompileEnvironment = new(UnrealTargetPlatform.Win64, CppConfiguration.Development, new UnrealArchitectures(UnrealArch.X64), null!);
 				IEnumerable<string> args = null!;
 
 				// CppStandardConstrainedFilterAttribute
 				{
 					testTargetRules.WindowsPlatform.Compiler = WindowsCompiler.Clang;
 					baseWarnings.AmbiguousReversedOperatorWarningLevel = WarningLevel.Off;
-					VersionNumber newVersionNumber = new VersionNumber(18, 9, 9);
+					VersionNumber newVersionNumber = new(18, 9, 9);
 
-					CppCompileEnvironment cppStandardCompileEnvironment = new CppCompileEnvironment(UnrealTargetPlatform.Win64, CppConfiguration.Debug, new UnrealArchitectures(UnrealArch.X64), null!);
+					CppCompileEnvironment cppStandardCompileEnvironment = new(UnrealTargetPlatform.Win64, CppConfiguration.Debug, new UnrealArchitectures(UnrealArch.X64), null!);
 #pragma warning disable CS0618 // Type or member is obsolete
 					cppStandardCompileEnvironment.CppStandard = CppStandardVersion.Cpp17;
 #pragma warning restore CS0618 // Type or member is obsolete
@@ -635,24 +635,24 @@ namespace UnrealBuildTool.Tests
 			// Generate ephemeral TargetRules & uproject files
 			using (SafeTestDirectory testRootDirectory = TestDirs.Create())
 			using (SafeTestDirectory testSourceDirectory = SafeTestDirectory.CreateTestDirectory(Path.Combine(testRootDirectory.TemporaryDirectory, "Source")))
-			using (SafeTestFile testEditorTargetTestFile = new SafeTestFile(File.ReadAllText(Path.Combine(testStubsDirectory.FullName, "TargetStubs", "TestStubEditor.Target.ubttest")), "TestStubEditor.Target.cs", testSourceDirectory.TemporaryDirectory))
-			using (SafeTestFile testTargetTestFile = new SafeTestFile(File.ReadAllText(Path.Combine(testStubsDirectory.FullName, "TargetStubs", "TestStub.Target.ubttest")), "TestStub.Target.cs", testSourceDirectory.TemporaryDirectory))
-			using (SafeTestFile testUprojectTestFile = new SafeTestFile(File.ReadAllText(Path.Combine(testStubsDirectory.FullName, "UprojectStubs", "TestStub.uproject.ubttest")), "TestStub.uproject", testRootDirectory.TemporaryDirectory))
+			using (SafeTestFile testEditorTargetTestFile = new(File.ReadAllText(Path.Combine(testStubsDirectory.FullName, "TargetStubs", "TestStubEditor.Target.ubttest")), "TestStubEditor.Target.cs", testSourceDirectory.TemporaryDirectory))
+			using (SafeTestFile testTargetTestFile = new(File.ReadAllText(Path.Combine(testStubsDirectory.FullName, "TargetStubs", "TestStub.Target.ubttest")), "TestStub.Target.cs", testSourceDirectory.TemporaryDirectory))
+			using (SafeTestFile testUprojectTestFile = new(File.ReadAllText(Path.Combine(testStubsDirectory.FullName, "UprojectStubs", "TestStub.uproject.ubttest")), "TestStub.uproject", testRootDirectory.TemporaryDirectory))
 			{
-				FileReference testUproject = new FileReference(testUprojectTestFile.TemporaryFile);
+				FileReference testUproject = new(testUprojectTestFile.TemporaryFile);
 				RulesAssembly generatedTestRulesAssembly = RulesCompiler.CreateTargetRulesAssembly(testUproject, "TestStubEditor", false, false, false, null, false, new TestLogger());
 				TargetRules testTargetRules = generatedTestRulesAssembly.CreateTargetRules("TestStubEditor", UnrealTargetPlatform.Win64, UnrealTargetConfiguration.Development, null, testUproject, null, new TestLogger());
 				testTargetRules.WindowsPlatform.Compiler = WindowsCompiler.VisualStudio2022;
-				CppCompileWarnings baseWarnings = new CppCompileWarnings(testTargetRules, null);
+				CppCompileWarnings baseWarnings = new(testTargetRules, null);
 
-				CppCompileEnvironment newCompileEnvironment = new CppCompileEnvironment(UnrealTargetPlatform.Win64, CppConfiguration.Development, new UnrealArchitectures(UnrealArch.X64), null!);
+				CppCompileEnvironment newCompileEnvironment = new(UnrealTargetPlatform.Win64, CppConfiguration.Development, new UnrealArchitectures(UnrealArch.X64), null!);
 				IEnumerable<string> args = null!;
 
 				// Specialized VC toolchain check;  UnsafeTypeCast
 				{
 					RulesAssembly generatedTestGameRulesAssembly = RulesCompiler.CreateTargetRulesAssembly(testUproject, "TestStub", false, false, false, null, false, new TestLogger());
 					TargetRules testTargetRulesNonWin64 = generatedTestGameRulesAssembly.CreateTargetRules("TestStub", UnrealTargetPlatform.Android, UnrealTargetConfiguration.Development, null, testUproject, null, new TestLogger());
-					CppCompileWarnings specializedWarnings = new CppCompileWarnings(testTargetRulesNonWin64, null);
+					CppCompileWarnings specializedWarnings = new(testTargetRulesNonWin64, null);
 
 					specializedWarnings.UnsafeTypeCastWarningLevel = WarningLevel.Warning;
 					args = specializedWarnings.GenerateWarningCommandLineArgs(newCompileEnvironment, typeof(VCToolChain));
@@ -688,7 +688,7 @@ namespace UnrealBuildTool.Tests
 					testTargetRules.WindowsPlatform.Compiler = WindowsCompiler.Clang;
 					baseWarnings.ShadowVariableWarningLevel = WarningLevel.Error;
 
-					VersionNumber newVersionNumber = new VersionNumber(16, 9, 5);
+					VersionNumber newVersionNumber = new(16, 9, 5);
 
 					args = baseWarnings.GenerateWarningCommandLineArgs(newCompileEnvironment, typeof(ClangToolChain), newVersionNumber);
 					Assert.IsTrue(args.Count() == 1);
@@ -722,7 +722,7 @@ namespace UnrealBuildTool.Tests
 					testTargetRules.WindowsPlatform.Compiler = WindowsCompiler.Clang;
 					baseWarnings.Shorten64To32WarningLevel = WarningLevel.Off;
 
-					VersionNumber newVersionNumber = new VersionNumber(18, 9, 9);
+					VersionNumber newVersionNumber = new(18, 9, 9);
 					args = baseWarnings.GenerateWarningCommandLineArgs(newCompileEnvironment, typeof(ClangToolChain), newVersionNumber);
 					Assert.IsFalse(args.Any());
 
@@ -745,17 +745,17 @@ namespace UnrealBuildTool.Tests
 			// Generate ephemeral TargetRules & uproject files
 			using (SafeTestDirectory testRootDirectory = TestDirs.Create())
 			using (SafeTestDirectory testSourceDirectory = SafeTestDirectory.CreateTestDirectory(Path.Combine(testRootDirectory.TemporaryDirectory, "Source")))
-			using (SafeTestFile testEditorTargetTestFile = new SafeTestFile(File.ReadAllText(Path.Combine(testStubsDirectory.FullName, "TargetStubs", "TestStubEditor.Target.ubttest")), "TestStubEditor.Target.cs", testSourceDirectory.TemporaryDirectory))
-			using (SafeTestFile testTargetTestFile = new SafeTestFile(File.ReadAllText(Path.Combine(testStubsDirectory.FullName, "TargetStubs", "TestStub.Target.ubttest")), "TestStub.Target.cs", testSourceDirectory.TemporaryDirectory))
-			using (SafeTestFile testUprojectTestFile = new SafeTestFile(File.ReadAllText(Path.Combine(testStubsDirectory.FullName, "UprojectStubs", "TestStub.uproject.ubttest")), "TestStub.uproject", testRootDirectory.TemporaryDirectory))
+			using (SafeTestFile testEditorTargetTestFile = new(File.ReadAllText(Path.Combine(testStubsDirectory.FullName, "TargetStubs", "TestStubEditor.Target.ubttest")), "TestStubEditor.Target.cs", testSourceDirectory.TemporaryDirectory))
+			using (SafeTestFile testTargetTestFile = new(File.ReadAllText(Path.Combine(testStubsDirectory.FullName, "TargetStubs", "TestStub.Target.ubttest")), "TestStub.Target.cs", testSourceDirectory.TemporaryDirectory))
+			using (SafeTestFile testUprojectTestFile = new(File.ReadAllText(Path.Combine(testStubsDirectory.FullName, "UprojectStubs", "TestStub.uproject.ubttest")), "TestStub.uproject", testRootDirectory.TemporaryDirectory))
 			{
-				FileReference testUproject = new FileReference(testUprojectTestFile.TemporaryFile);
+				FileReference testUproject = new(testUprojectTestFile.TemporaryFile);
 				RulesAssembly generatedTestRulesAssembly = RulesCompiler.CreateTargetRulesAssembly(testUproject, "TestStubEditor", false, false, false, null, false, new TestLogger());
 				TargetRules testTargetRules = generatedTestRulesAssembly.CreateTargetRules("TestStubEditor", UnrealTargetPlatform.Win64, UnrealTargetConfiguration.Development, null, testUproject, null, new TestLogger());
 				testTargetRules.WindowsPlatform.Compiler = WindowsCompiler.VisualStudio2022;
-				CppCompileWarnings baseWarnings = new CppCompileWarnings(testTargetRules, null);
+				CppCompileWarnings baseWarnings = new(testTargetRules, null);
 
-				CppCompileEnvironment newCompileEnvironment = new CppCompileEnvironment(UnrealTargetPlatform.Win64, CppConfiguration.Development, new UnrealArchitectures(UnrealArch.X64), null!);
+				CppCompileEnvironment newCompileEnvironment = new(UnrealTargetPlatform.Win64, CppConfiguration.Development, new UnrealArchitectures(UnrealArch.X64), null!);
 				IEnumerable<string> args = null!;
 
 				// VC toolchain check;
diff --git a/Engine/Source/Programs/Shared/UnrealBuildTool.Tests/JsonExtensionsTests.cs b/Engine/Source/Programs/Shared/UnrealBuildTool.Tests/JsonExtensionsTests.cs
--- a/Engine/Source/Programs/Shared/UnrealBuildTool.Tests/JsonExtensionsTests.cs
+++ b/Engine/Source/Programs/Shared/UnrealBuildTool.Tests/JsonExtensionsTests.cs
@@ -21,7 +21,7 @@ namespace UnrealBuildTool.Tests
 		[TestMethod]
 		public void TestGetEnumArrayFieldWithDeprecatedFallbackReturnsNullWhenBothAbsent()
 		{
-			JsonObject json = new JsonObject();
+			JsonObject json = new();
 			TestEnum[]? values = json.GetEnumArrayFieldWithDeprecatedFallback<TestEnum>("New", "Old");
 			Assert.IsNull(values);
 		}
@@ -29,7 +29,7 @@ namespace UnrealBuildTool.Tests
 		[TestMethod]
 		public void TestGetEnumArrayFieldWithDeprecatedFallbackThrowsWhenBothPresent()
 		{
-			JsonObject json = new JsonObject();
+			JsonObject json = new();
 			json.AddOrSetFieldValue("New", [TestEnum.A, TestEnum.B]);
 			json.AddOrSetFieldValue("Old", [TestEnum.A, TestEnum.B]);
 			BuildLogEventException ex = Assert.ThrowsExactly<BuildLogEventException>(
@@ -41,7 +41,7 @@ namespace UnrealBuildTool.Tests
 		[TestMethod]
 		public void TestGetEnumArrayFieldWithDeprecatedFallbackPicksUpNewValues()
 		{
-			JsonObject json = new JsonObject();
+			JsonObject json = new();
 			json.AddOrSetFieldValue("New", [TestEnum.A, TestEnum.B]);
 			TestEnum[]? values = json.GetEnumArrayFieldWithDeprecatedFallback<TestEnum>("New", "Old");
 			CollectionAssert.AreEquivalent(values, new[] { TestEnum.A, TestEnum.B });
@@ -50,7 +50,7 @@ namespace UnrealBuildTool.Tests
 		[TestMethod]
 		public void TestGetEnumArrayFieldWithDeprecatedFallbackPicksUpFallbackValues()
 		{
-			JsonObject json = new JsonObject();
+			JsonObject json = new();
 			json.AddOrSetFieldValue("Old", [TestEnum.A, TestEnum.B]);
 			TestEnum[]? values = json.GetEnumArrayFieldWithDeprecatedFallback<TestEnum>("New", "Old");
 			CollectionAssert.AreEquivalent(values, new[] { TestEnum.A, TestEnum.B });
@@ -59,7 +59,7 @@ namespace UnrealBuildTool.Tests
 		[TestMethod]
 		public void TestGetEnumArrayFieldWithDeprecatedFallbackFailsOnInvalidNewValue()
 		{
-			JsonObject json = new JsonObject();
+			JsonObject json = new();
 			json.AddOrSetFieldValue("New", ["A", "B", "C", "BadItem1", "BadItem2"]);
 			BuildLogEventException ex = Assert.ThrowsExactly<BuildLogEventException>(
 				() => json.GetEnumArrayFieldWithDeprecatedFallback<TestEnum>("New", "Old"));
@@ -71,7 +71,7 @@ namespace UnrealBuildTool.Tests
 		[TestMethod]
 		public void TestGetEnumArrayFieldWithDeprecatedFallbackFailsOnInvalidFallbackValue()
 		{
-			JsonObject json = new JsonObject();
+			JsonObject json = new();
 			json.AddOrSetFieldValue("Old", ["A", "B", "C", "BadItem1", "BadItem2"]);
 			BuildLogEventException ex = Assert.ThrowsExactly<BuildLogEventException>(
 				() => json.GetEnumArrayFieldWithDeprecatedFallback<TestEnum>("New", "Old"));
@@ -83,7 +83,7 @@ namespace UnrealBuildTool.Tests
 		[TestMethod]
 		public void TestGetEnumArrayFieldWithDeprecatedFallbackFailsOnInvalidDataType()
 		{
-			JsonObject json = new JsonObject();
+			JsonObject json = new();
 			json.AddOrSetFieldValue("New", "A");
 			BuildLogEventException ex = Assert.ThrowsExactly<BuildLogEventException>(
 				() => json.GetEnumArrayFieldWithDeprecatedFallback<TestEnum>("New", "Old"));
diff --git a/Engine/Source/Programs/Shared/UnrealBuildTool.Tests/PluginDescriptorTests.cs b/Engine/Source/Programs/Shared/UnrealBuildTool.Tests/PluginDescriptorTests.cs
--- a/Engine/Source/Programs/Shared/UnrealBuildTool.Tests/PluginDescriptorTests.cs
+++ b/Engine/Source/Programs/Shared/UnrealBuildTool.Tests/PluginDescriptorTests.cs
@@ -150,7 +150,7 @@ namespace UnrealBuildTool.Tests
 			Debug.Assert(!String.IsNullOrEmpty(extension));
 			Debug.Assert(extension == ".uplugin");
 			string inputFile = Path.Join(s_tempDirectory?.FullName, fileName);
-			FileReference inputFileReference = new FileReference(inputFile);
+			FileReference inputFileReference = new(inputFile);
 			File.WriteAllText(inputFile, fileContent);
 			return inputFileReference;
 		}
@@ -329,7 +329,7 @@ namespace UnrealBuildTool.Tests
 				return verbatimString;
 			}
 			string[] lines = verbatimString.Split([Environment.NewLine], StringSplitOptions.None);
-			StringBuilder jsonStringBuilder = new StringBuilder();
+			StringBuilder jsonStringBuilder = new();
 			int numberOfTabs = 0;
 			bool bFoundTabLevel = false;
 			foreach (string line in lines)
diff --git a/Engine/Source/Programs/Shared/UnrealBuildTool.Tests/PreprocessorExpressionTests.cs b/Engine/Source/Programs/Shared/UnrealBuildTool.Tests/PreprocessorExpressionTests.cs
--- a/Engine/Source/Programs/Shared/UnrealBuildTool.Tests/PreprocessorExpressionTests.cs
+++ b/Engine/Source/Programs/Shared/UnrealBuildTool.Tests/PreprocessorExpressionTests.cs
@@ -123,7 +123,7 @@ namespace UnrealBuildTool.Tests
 		/// <param name="expectedResult">The expected value of the expression</param>
 		static void RunTest(string expression, long? expectedResult)
 		{
-			using TokenReader reader = new TokenReader(expression);
+			using TokenReader reader = new(expression);
 
 			List<Token> tokens = [];
 			while (reader.MoveNext())
diff --git a/Engine/Source/Programs/Shared/UnrealBuildTool.Tests/PreprocessorStateTests.cs b/Engine/Source/Programs/Shared/UnrealBuildTool.Tests/PreprocessorStateTests.cs
--- a/Engine/Source/Programs/Shared/UnrealBuildTool.Tests/PreprocessorStateTests.cs
+++ b/Engine/Source/Programs/Shared/UnrealBuildTool.Tests/PreprocessorStateTests.cs
@@ -13,14 +13,14 @@ namespace UnrealBuildTool.Tests
 		[TestMethod]
 		public void Run()
 		{
-			PreprocessorState baseState = new PreprocessorState();
+			PreprocessorState baseState = new();
 			baseState.DefineMacro(ParseMacro("FOO", null, "123"));
 			baseState.DefineMacro(ParseMacro("BAR", ["Arg1", "Arg2"], "Arg1 + Arg2"));
 			Assert.AreEqual("m BAR(Arg1, Arg2)=Arg1 + Arg2\nm FOO=123\n", FormatState(baseState));
 
 			// Modify the state and revert it
 			{
-				PreprocessorState state = new PreprocessorState(baseState);
+				PreprocessorState state = new(baseState);
 				Assert.AreEqual(FormatState(baseState), FormatState(state));
 
 				PreprocessorTransform transform1 = state.BeginCapture();
@@ -30,7 +30,7 @@ namespace UnrealBuildTool.Tests
 				Assert.AreEqual("b Active\nm BAR(Arg1, Arg2)=Arg1 + Arg2\nm FOO2=FOO()\n", FormatState(state));
 				Assert.AreEqual("+b Active\n+m FOO undef\n+m FOO2=FOO()\n", FormatTransform(transform1));
 
-				PreprocessorState state2 = new PreprocessorState(baseState);
+				PreprocessorState state2 = new(baseState);
 				state2.TryToApply(transform1);
 				Assert.AreEqual(FormatState(state), FormatState(state2));
 
@@ -48,7 +48,7 @@ namespace UnrealBuildTool.Tests
 
 			// Check the tracking of branches
 			{
-				PreprocessorState state = new PreprocessorState(baseState);
+				PreprocessorState state = new(baseState);
 				PreprocessorTransform transform1 = state.BeginCapture();
 				state.IsCurrentBranchActive();
 				Assert.AreEqual("=b Active\n", FormatTransform(transform1));
@@ -71,7 +71,7 @@ namespace UnrealBuildTool.Tests
 
 		static string FormatState(PreprocessorState state)
 		{
-			StringBuilder result = new StringBuilder();
+			StringBuilder result = new();
 			foreach (PreprocessorBranch branch in state.CurrentBranches)
 			{
 				result.AppendFormat("b {0}\n", branch);
@@ -85,7 +85,7 @@ namespace UnrealBuildTool.Tests
 
 		static string FormatTransform(PreprocessorTransform transform)
 		{
-			StringBuilder result = new StringBuilder();
+			StringBuilder result = new();
 			if (transform.RequireTopmostActive.HasValue)
 			{
 				result.AppendFormat("=b {0}\n", transform.RequireTopmostActive.Value ? "Active" : "0");
@@ -127,7 +127,7 @@ namespace UnrealBuildTool.Tests
 		{
 			List<Token> tokens = [];
 
-			using TokenReader reader = new TokenReader(value);
+			using TokenReader reader = new(value);
 			while (reader.MoveNext())
 			{
 				tokens.Add(reader.Current);
diff --git a/Engine/Source/Programs/Shared/UnrealBuildTool.Tests/PreprocessorTests.cs b/Engine/Source/Programs/Shared/UnrealBuildTool.Tests/PreprocessorTests.cs
--- a/Engine/Source/Programs/Shared/UnrealBuildTool.Tests/PreprocessorTests.cs
+++ b/Engine/Source/Programs/Shared/UnrealBuildTool.Tests/PreprocessorTests.cs
@@ -225,9 +225,9 @@ namespace UnrealBuildTool.Tests
 		{
 			string[] lines = fragment.Split('\n');
 
-			SourceFile file = new SourceFile(FileItem.GetItemByPath(Path.Combine(Path.GetTempPath(), "temp.cpp")), TokenReader.GetNullTerminatedByteArray(fragment));
+			SourceFile file = new(FileItem.GetItemByPath(Path.Combine(Path.GetTempPath(), "temp.cpp")), TokenReader.GetNullTerminatedByteArray(fragment));
 
-			Preprocessor instance = new Preprocessor();
+			Preprocessor instance = new();
 
 			string? result;
 			try
@@ -241,7 +241,7 @@ namespace UnrealBuildTool.Tests
 					{
 						if (instance.IsCurrentBranchActive())
 						{
-							StringBuilder sourceText = new StringBuilder();
+							StringBuilder sourceText = new();
 
 							int lastLineIndex = (markupIdx + 1 < file.Markup.Length) ? file.Markup[markupIdx + 1].LineNumber - 1 : lines.Length;
 							for (int lineIndex = markup.LineNumber - 1; lineIndex < lastLineIndex; lineIndex++)
@@ -251,7 +251,7 @@ namespace UnrealBuildTool.Tests
 
 							List<Token> tokens = [];
 
-							using TokenReader reader = new TokenReader(sourceText.ToString());
+							using TokenReader reader = new(sourceText.ToString());
 							while (reader.MoveNext())
 							{
 								tokens.Add(reader.Current);
diff --git a/Engine/Source/Programs/Shared/UnrealBuildTool.Tests/ProjectDescriptorTests.cs b/Engine/Source/Programs/Shared/UnrealBuildTool.Tests/ProjectDescriptorTests.cs
--- a/Engine/Source/Programs/Shared/UnrealBuildTool.Tests/ProjectDescriptorTests.cs
+++ b/Engine/Source/Programs/Shared/UnrealBuildTool.Tests/ProjectDescriptorTests.cs
@@ -17,7 +17,7 @@ namespace UnrealBuildTool.Tests
 			string uprojectPath = Path.Combine(TestDirs.Stubs.FullName, "UprojectStubs", "BadModuleAllowList.uproject.ubttest");
 
 			using SafeTestDirectory testRootDirectory = TestDirs.Create();
-			using SafeTestFile uprojectFile = new SafeTestFile(File.ReadAllText(uprojectPath), "TestStub.uproject", testRootDirectory.TemporaryDirectory);
+			using SafeTestFile uprojectFile = new(File.ReadAllText(uprojectPath), "TestStub.uproject", testRootDirectory.TemporaryDirectory);
 
 			CheckExceptionContainsError(uprojectFile, "Editro", "TargetAllowList");
 		}
@@ -28,7 +28,7 @@ namespace UnrealBuildTool.Tests
 			string uprojectPath = Path.Combine(TestDirs.Stubs.FullName, "UprojectStubs", "BadPluginAllowList.uproject.ubttest");
 
 			using SafeTestDirectory testRootDirectory = TestDirs.Create();
-			using SafeTestFile uprojectFile = new SafeTestFile(File.ReadAllText(uprojectPath), "TestStub.uproject", testRootDirectory.TemporaryDirectory);
+			using SafeTestFile uprojectFile = new(File.ReadAllText(uprojectPath), "TestStub.uproject", testRootDirectory.TemporaryDirectory);
 
 			CheckExceptionContainsError(uprojectFile, "Shpipping", "TargetConfigurationAllowList");
 		}
diff --git a/Engine/Source/Programs/Shared/UnrealBuildTool.Tests/SourceFileTests.cs b/Engine/Source/Programs/Shared/UnrealBuildTool.Tests/SourceFileTests.cs
--- a/Engine/Source/Programs/Shared/UnrealBuildTool.Tests/SourceFileTests.cs
+++ b/Engine/Source/Programs/Shared/UnrealBuildTool.Tests/SourceFileTests.cs
@@ -38,7 +38,7 @@ namespace UnrealBuildTool.Tests
 			ConcurrentBag<SourceFile> sourceFiles = [];
 			using (GlobalTracer.Instance.BuildSpan("Scanning source files").StartActive())
 			{
-				using (ThreadPoolWorkQueue queue = new ThreadPoolWorkQueue())
+				using (ThreadPoolWorkQueue queue = new())
 				{
 					foreach (DirectoryReference baseDirectory in baseDirectories)
 					{
@@ -52,7 +52,7 @@ namespace UnrealBuildTool.Tests
 
 			using (GlobalTracer.Instance.BuildSpan("Writing source file data").StartActive())
 			{
-				using (BinaryArchiveWriter writer = new BinaryArchiveWriter(tempDataFile))
+				using (BinaryArchiveWriter writer = new(tempDataFile))
 				{
 					writer.WriteList(sourceFiles.ToList(), x => x.Write(writer));
 				}
@@ -61,7 +61,7 @@ namespace UnrealBuildTool.Tests
 			List<SourceFile>? readSourceFiles = [];
 			using (GlobalTracer.Instance.BuildSpan("Reading source file data").StartActive())
 			{
-				using (BinaryArchiveReader reader = new BinaryArchiveReader(tempDataFile))
+				using (BinaryArchiveReader reader = new(tempDataFile))
 				{
 					readSourceFiles = reader.ReadList(() => new SourceFile(reader));
 				}
diff --git a/Engine/Source/Programs/Shared/UnrealBuildTool.Tests/TargetRulesTests.cs b/Engine/Source/Programs/Shared/UnrealBuildTool.Tests/TargetRulesTests.cs
--- a/Engine/Source/Programs/Shared/UnrealBuildTool.Tests/TargetRulesTests.cs
+++ b/Engine/Source/Programs/Shared/UnrealBuildTool.Tests/TargetRulesTests.cs
@@ -20,10 +20,10 @@ namespace UnrealBuildTool.Tests
 			// Generate ephemeral TargetRules & uproject files
 			using (SafeTestDirectory testRootDirectory = TestDirs.Create())
 			using (SafeTestDirectory testSourceDirectory = SafeTestDirectory.CreateTestDirectory(Path.Combine(testRootDirectory.TemporaryDirectory, "Source")))
-			using (SafeTestFile testEditorTargetTestFile = new SafeTestFile(File.ReadAllText(Path.Combine(testStubsDirectory.FullName, "TargetStubs", "TestStubEditor.Target.ubttest")), "TestStubEditor.Target.cs", testSourceDirectory.TemporaryDirectory))
-			using (SafeTestFile testUprojectTestFile = new SafeTestFile(File.ReadAllText(Path.Combine(testStubsDirectory.FullName, "UprojectStubs", "TestStub.uproject.ubttest")), "TestStub.uproject", testRootDirectory.TemporaryDirectory))
+			using (SafeTestFile testEditorTargetTestFile = new(File.ReadAllText(Path.Combine(testStubsDirectory.FullName, "TargetStubs", "TestStubEditor.Target.ubttest")), "TestStubEditor.Target.cs", testSourceDirectory.TemporaryDirectory))
+			using (SafeTestFile testUprojectTestFile = new(File.ReadAllText(Path.Combine(testStubsDirectory.FullName, "UprojectStubs", "TestStub.uproject.ubttest")), "TestStub.uproject", testRootDirectory.TemporaryDirectory))
 			{
-				FileReference testUproject = new FileReference(testUprojectTestFile.TemporaryFile);
+				FileReference testUproject = new(testUprojectTestFile.TemporaryFile);
 
 				// Generate rules assembly & subsequent targets
 				RulesAssembly generatedTestRulesAssembly = RulesCompiler.CreateTargetRulesAssembly(testUproject, "TestStubEditor", false, false, false, null, false, new TestLogger());
diff --git a/Engine/Source/Programs/Shared/UnrealBuildTool.Tests/TestUtilities/SafeTestDirectory.cs b/Engine/Source/Programs/Shared/UnrealBuildTool.Tests/TestUtilities/SafeTestDirectory.cs
--- a/Engine/Source/Programs/Shared/UnrealBuildTool.Tests/TestUtilities/SafeTestDirectory.cs
+++ b/Engine/Source/Programs/Shared/UnrealBuildTool.Tests/TestUtilities/SafeTestDirectory.cs
@@ -18,7 +18,7 @@ namespace UnrealBuildTool.Tests.TestUtilities
 
 		internal static SafeTestDirectory CreateTestDirectory(string absoluteDirectoryPath)
 		{
-			SafeTestDirectory returnDirectory = new SafeTestDirectory(absoluteDirectoryPath);
+			SafeTestDirectory returnDirectory = new(absoluteDirectoryPath);
 			if (!Directory.Exists(absoluteDirectoryPath))
 			{
 				Directory.CreateDirectory(absoluteDirectoryPath);
diff --git a/Engine/Source/Programs/Shared/UnrealBuildTool.Tests/TestUtilities/SafeTestFile.cs b/Engine/Source/Programs/Shared/UnrealBuildTool.Tests/TestUtilities/SafeTestFile.cs
--- a/Engine/Source/Programs/Shared/UnrealBuildTool.Tests/TestUtilities/SafeTestFile.cs
+++ b/Engine/Source/Programs/Shared/UnrealBuildTool.Tests/TestUtilities/SafeTestFile.cs
@@ -16,7 +16,7 @@ namespace UnrealBuildTool.Tests.TestUtilities
 			{
 				TemporaryFile = Path.Combine(outputFolder, outputFileName);
 
-				using (StreamWriter writer = new StreamWriter(TemporaryFile))
+				using (StreamWriter writer = new(TemporaryFile))
 				{
 					writer.Write(content);
 				}
@@ -36,6 +36,6 @@ namespace UnrealBuildTool.Tests.TestUtilities
 		}
 
 		public static implicit operator FileReference(SafeTestFile file)
-			=> new FileReference(file.TemporaryFile);
+			=> new(file.TemporaryFile);
 	}
 }
diff --git a/Engine/Source/Programs/Shared/UnrealBuildTool.Tests/TestUtilities/StaticTargetRulesInitializer.cs b/Engine/Source/Programs/Shared/UnrealBuildTool.Tests/TestUtilities/StaticTargetRulesInitializer.cs
--- a/Engine/Source/Programs/Shared/UnrealBuildTool.Tests/TestUtilities/StaticTargetRulesInitializer.cs
+++ b/Engine/Source/Programs/Shared/UnrealBuildTool.Tests/TestUtilities/StaticTargetRulesInitializer.cs
@@ -46,13 +46,13 @@ namespace UnrealBuildTool.Tests.TestUtilities
 			UEBuildPlatform? _;
 			if (!UEBuildPlatform.TryGetBuildPlatform(UnrealTargetPlatform.Win64, out _))
 			{
-				MicrosoftPlatformSDK sdk = new MicrosoftPlatformSDK(logger);
+				MicrosoftPlatformSDK sdk = new(logger);
 				UEBuildPlatform.RegisterBuildPlatform(new WindowsPlatform(UnrealTargetPlatform.Win64, sdk, logger), logger);
 			}
 
 			if (!UEBuildPlatform.TryGetBuildPlatform(UnrealTargetPlatform.Android, out _))
 			{
-				AndroidPlatformSDK asdk = new AndroidPlatformSDK(logger);
+				AndroidPlatformSDK asdk = new(logger);
 				UEBuildPlatform.RegisterBuildPlatform(new AndroidPlatform(asdk, logger), logger);
 			}
 
diff --git a/Engine/Source/Programs/Shared/UnrealBuildTool.Tests/TokenReaderTests.cs b/Engine/Source/Programs/Shared/UnrealBuildTool.Tests/TokenReaderTests.cs
--- a/Engine/Source/Programs/Shared/UnrealBuildTool.Tests/TokenReaderTests.cs
+++ b/Engine/Source/Programs/Shared/UnrealBuildTool.Tests/TokenReaderTests.cs
@@ -135,7 +135,7 @@ namespace UnrealBuildTool.Tests
 		{
 			List<Token> tokens = [];
 
-			using TokenReader reader = new TokenReader(inputText);
+			using TokenReader reader = new(inputText);
 			while (reader.MoveNext())
 			{
 				tokens.Add(reader.Current);
@@ -147,7 +147,7 @@ namespace UnrealBuildTool.Tests
 
 		static string FormatToken(Token token)
 		{
-			StringBuilder result = new StringBuilder();
+			StringBuilder result = new();
 			result.AppendFormat(token.Type.ToString());
 			if (token.Type == TokenType.Identifier || token.Type == TokenType.Number || token.Type == TokenType.Character || token.Type == TokenType.String)
 			{
@@ -158,7 +158,7 @@ namespace UnrealBuildTool.Tests
 
 		public static void RunLineNumberTest(string inputText, int expectedLineNumber)
 		{
-			using TokenReader reader = new TokenReader(inputText);
+			using TokenReader reader = new(inputText);
 			while (reader.MoveNext())
 			{
 				if (reader.Current.Type == TokenType.Plus)
diff --git a/Engine/Source/Programs/Shared/UnrealBuildTool.Tests/WarningLevelDefaultAttributeTests.cs b/Engine/Source/Programs/Shared/UnrealBuildTool.Tests/WarningLevelDefaultAttributeTests.cs
--- a/Engine/Source/Programs/Shared/UnrealBuildTool.Tests/WarningLevelDefaultAttributeTests.cs
+++ b/Engine/Source/Programs/Shared/UnrealBuildTool.Tests/WarningLevelDefaultAttributeTests.cs
@@ -117,8 +117,8 @@ namespace UnrealBuildTool.Tests
 				new VersionWarningLevelDefaultAttribute(WarningLevel.Warning, BuildSettingsVersion.V1, BuildSettingsVersion.V2),
 				new VersionWarningLevelDefaultAttribute(WarningLevel.Error, BuildSettingsVersion.V3, BuildSettingsVersion.V4)
 			];
-			TestBuildSettingsProvider testBuildSettingsProvider = new TestBuildSettingsProvider();
-			BuildSystemContext newBuildSystemContext = new BuildSystemContext(testBuildSettingsProvider);
+			TestBuildSettingsProvider testBuildSettingsProvider = new();
+			BuildSystemContext newBuildSystemContext = new(testBuildSettingsProvider);
 
 			Assert.AreEqual(WarningLevel.Warning, VersionWarningLevelDefaultAttribute.ResolveVersionWarningLevelDefault(attributes, newBuildSystemContext));
 			testBuildSettingsProvider._testBuildSettings = BuildSettingsVersion.V2;
diff --git a/Engine/Source/Programs/UnrealBuildTool/Configuration/UEBuildModule.cs b/Engine/Source/Programs/UnrealBuildTool/Configuration/UEBuildModule.cs
--- a/Engine/Source/Programs/UnrealBuildTool/Configuration/UEBuildModule.cs
+++ b/Engine/Source/Programs/UnrealBuildTool/Configuration/UEBuildModule.cs
@@ -478,7 +478,7 @@ namespace UnrealBuildTool
 		{
 			_ = InEnumerableStrings.TryGetNonEnumeratedCount(out int estimatedCount);
 			HashSet<string> hashSet = new(estimatedCount);
-			foreach (string s in InEnumerableStrings.Where(str => !string.IsNullOrWhiteSpace(str)))
+			foreach (string s in InEnumerableStrings.Where(str => !String.IsNullOrWhiteSpace(str)))
 			{
 				hashSet.Add(ExpandPathVariables(s, null, null));
 			}
diff --git a/Engine/Source/Programs/UnrealBuildTool/Modes/IWYUMode.cs b/Engine/Source/Programs/UnrealBuildTool/Modes/IWYUMode.cs
--- a/Engine/Source/Programs/UnrealBuildTool/Modes/IWYUMode.cs
+++ b/Engine/Source/Programs/UnrealBuildTool/Modes/IWYUMode.cs
@@ -2016,7 +2016,7 @@ namespace UnrealBuildTool
 							else if (LinesToAdd.Count != 0 && (!IsCpp || LastSeenIncludeBeforeCode != -1))
 							{
 								string LineToAdd = LinesToAdd.First();
-								if (string.Compare(LineToAdd, OldLine, StringComparison.Ordinal) < 0)
+								if (String.Compare(LineToAdd, OldLine, StringComparison.Ordinal) < 0)
 								{
 									NewLines.Add(LineToAdd);
 									LinesToAdd.Remove(LineToAdd);
diff --git a/Engine/Source/Programs/UnrealBuildTool/Platform/Android/AndroidToolChain.cs b/Engine/Source/Programs/UnrealBuildTool/Platform/Android/AndroidToolChain.cs
--- a/Engine/Source/Programs/UnrealBuildTool/Platform/Android/AndroidToolChain.cs
+++ b/Engine/Source/Programs/UnrealBuildTool/Platform/Android/AndroidToolChain.cs
@@ -1467,7 +1467,7 @@ namespace UnrealBuildTool
 		public string GetBuildToolsVersion()
 		{
 			string HomePath = Environment.ExpandEnvironmentVariables("%ANDROID_HOME%");
-			if (string.IsNullOrEmpty(HomePath))
+			if (String.IsNullOrEmpty(HomePath))
 			{
 				throw new BuildException("Failed to find %ANDROID_HOME%. Ensure environment is configured properly");
 			}
diff --git a/Engine/Source/Programs/UnrealBuildTool/Preprocessor/Identifier.cs b/Engine/Source/Programs/UnrealBuildTool/Preprocessor/Identifier.cs
--- a/Engine/Source/Programs/UnrealBuildTool/Preprocessor/Identifier.cs
+++ b/Engine/Source/Programs/UnrealBuildTool/Preprocessor/Identifier.cs
@@ -59,7 +59,7 @@ namespace UnrealBuildTool
 		/// <returns>Value indicating which identifier should sort first</returns>
 		public int CompareTo(Identifier? other)
 		{
-			return other is null ? 1 : string.Compare(_name, other._name, StringComparison.Ordinal);
+			return other is null ? 1 : String.Compare(_name, other._name, StringComparison.Ordinal);
 		}
 
 		/// <summary>
diff --git a/Engine/Source/Programs/UnrealBuildTool/ProjectFiles/VisualStudioCode/VSCodeProjectFileGenerator.cs b/Engine/Source/Programs/UnrealBuildTool/ProjectFiles/VisualStudioCode/VSCodeProjectFileGenerator.cs
--- a/Engine/Source/Programs/UnrealBuildTool/ProjectFiles/VisualStudioCode/VSCodeProjectFileGenerator.cs
+++ b/Engine/Source/Programs/UnrealBuildTool/ProjectFiles/VisualStudioCode/VSCodeProjectFileGenerator.cs
@@ -290,7 +290,7 @@ namespace UnrealBuildTool
 			{
 				Projects = new List<ProjectFile>(AllProjectFiles);
 			}
-			Projects.Sort((A, B) => { return string.Compare(A.ProjectFilePath.GetFileName(), B.ProjectFilePath.GetFileName(), StringComparison.Ordinal); });
+			Projects.Sort((A, B) => { return String.Compare(A.ProjectFilePath.GetFileName(), B.ProjectFilePath.GetFileName(), StringComparison.Ordinal); });
 
 			ProjectData ProjectData = GatherProjectData(Projects);
 
diff --git a/Engine/Source/Programs/UnrealBuildTool/ProjectFiles/Xcode/XcodeProjectFileGenerator.cs b/Engine/Source/Programs/UnrealBuildTool/ProjectFiles/Xcode/XcodeProjectFileGenerator.cs
--- a/Engine/Source/Programs/UnrealBuildTool/ProjectFiles/Xcode/XcodeProjectFileGenerator.cs
+++ b/Engine/Source/Programs/UnrealBuildTool/ProjectFiles/Xcode/XcodeProjectFileGenerator.cs
@@ -395,7 +395,7 @@ namespace UnrealBuildTool
 						return TypeA.CompareTo(TypeB);
 					}
 
-					return string.Compare(TargetA.Name, TargetB.Name, StringComparison.Ordinal);
+					return String.Compare(TargetA.Name, TargetB.Name, StringComparison.Ordinal);
 				});
 
 				foreach (XcodeProjectXcconfig.XcodeProjectFile XcodeProject in BuildableProjects)
