Horde: Configurable log parser ignore pattern discovery

Add project/stream IgnorePatternFiles config and feature
flag to expand discovery beyond Engine/. Adds regex timeout
protection, invalid pattern handling, and path traversal
guards. Wires new config to job agent via gRPC.

diff --git a/Engine/Source/Programs/Horde/Drivers/JobDriver/Execution/JobExecutor.cs b/Engine/Source/Programs/Horde/Drivers/JobDriver/Execution/JobExecutor.cs
--- a/Engine/Source/Programs/Horde/Drivers/JobDriver/Execution/JobExecutor.cs
+++ b/Engine/Source/Programs/Horde/Drivers/JobDriver/Execution/JobExecutor.cs
@@ -1814,48 +1814,59 @@ namespace JobDriver.Execution
 			return result;
 		}
 
-		static void AddRestrictedDirs(List<DirectoryReference> directories, string subFolder)
+		async Task LoadIgnorePatternsAsync(LogEventParser parser, DirectoryReference workspaceDir, ILogger logger)
 		{
-			int numDirs = directories.Count;
-			for (int idx = 0; idx < numDirs; idx++)
+			List<FileReference> patternFiles;
+			if (Batch.EnableExpandedIgnorePatterns)
 			{
-				DirectoryReference subDir = DirectoryReference.Combine(directories[idx], subFolder);
-				if (DirectoryReference.Exists(subDir))
+				patternFiles = LogEventParser.DiscoverIgnorePatternFiles(workspaceDir, EnginePath);
+			}
+			else
+			{
+				patternFiles = LogEventParser.DiscoverLegacyIgnorePatternFiles(workspaceDir, EnginePath);
+			}
+
+			await parser.ReadIgnorePatternsAsync(patternFiles);
+
+			// Read additional config-specified ignore pattern files (expanded mode only)
+			if (Batch.EnableExpandedIgnorePatterns)
+			{
+				HashSet<FileReference> discoveredFiles = new HashSet<FileReference>(patternFiles);
+				List<FileReference> configFiles = ValidateConfigPatternFiles(workspaceDir, discoveredFiles, logger);
+				if (configFiles.Count > 0)
 				{
-					directories.AddRange(DirectoryReference.EnumerateDirectories(subDir));
+					await parser.ReadIgnorePatternsAsync(configFiles);
 				}
 			}
+
+			logger.LogInformation("Loaded {PatternCount} ignore patterns from workspace", parser.IgnorePatterns.Count);
 		}
 
-		async Task<List<string>> ReadIgnorePatternsAsync(DirectoryReference workspaceDir, ILogger logger)
+		List<FileReference> ValidateConfigPatternFiles(DirectoryReference workspaceDir, HashSet<FileReference> discoveredFiles, ILogger logger)
 		{
-			List<DirectoryReference> baseDirs = new List<DirectoryReference>();
-			baseDirs.Add(DirectoryReference.Combine(workspaceDir, EnginePath));
-			AddRestrictedDirs(baseDirs, "Restricted");
-			AddRestrictedDirs(baseDirs, "Platforms");
+			string workspacePrefix = workspaceDir.FullName.TrimEnd(Path.DirectorySeparatorChar) + Path.DirectorySeparatorChar;
+			List<FileReference> validFiles = new List<FileReference>();
 
-			List<string> ignorePatternLines = new List<string>();
-			foreach (DirectoryReference baseDir in baseDirs)
+			foreach (string relativePath in Batch.IgnorePatternFiles)
 			{
-				FileReference ignorePatternFile = FileReference.Combine(baseDir, "Build", "Horde", "IgnorePatterns.txt");
-				if (FileReference.Exists(ignorePatternFile))
+				string normalizedPath = Path.GetFullPath(Path.Combine(workspaceDir.FullName, relativePath));
+				if (!normalizedPath.StartsWith(workspacePrefix, StringComparison.OrdinalIgnoreCase))
 				{
-					logger.LogInformation("Reading ignore patterns from {File}...", ignorePatternFile);
-					ignorePatternLines.AddRange(await FileReference.ReadAllLinesAsync(ignorePatternFile));
+					logger.LogWarning("Ignoring config-specified pattern file outside workspace: {Path}", relativePath);
+					continue;
 				}
-			}
 
-			HashSet<string> ignorePatterns = new HashSet<string>(StringComparer.Ordinal);
-			foreach (string line in ignorePatternLines)
-			{
-				string trimLine = line.Trim();
-				if (trimLine.Length > 0 && trimLine[0] != '#')
+				FileReference configPatternFile = new FileReference(normalizedPath);
+				if (discoveredFiles.Contains(configPatternFile))
 				{
-					ignorePatterns.Add(trimLine);
+					logger.LogDebug("Config-specified file already discovered by auto-scan, skipping: {File}", configPatternFile);
+					continue;
 				}
+
+				validFiles.Add(configPatternFile);
 			}
 
-			return ignorePatterns.ToList();
+			return validFiles;
 		}
 
 		static FileReference GetCleanupScript(DirectoryReference workspaceDir)
@@ -1892,7 +1903,7 @@ namespace JobDriver.Execution
 		/// </summary>
 		protected static async Task ExecuteLeaseCleanupScriptAsync(DirectoryReference workspaceDir, ILogger logger)
 		{
-			using (LogParser parser = new LogParser(logger, new List<string>()))
+			using (LogParser parser = new LogParser(logger))
 			{
 				FileReference leaseCleanupScript = GetLeaseCleanupScript(workspaceDir);
 				await ExecuteCleanupScriptAsync(leaseCleanupScript, parser, logger);
@@ -2194,11 +2205,10 @@ namespace JobDriver.Execution
 				FileUtils.ForceDeleteDirectoryContents(telemetryDir);
 			}
 
-			List<string> ignorePatterns = await ReadIgnorePatternsAsync(workspaceDir, jobLogger);
-
 			int exitCode;
-			using (LogParser filter = new LogParser(jobLogger, ignorePatterns))
+			using (LogParser filter = new LogParser(jobLogger))
 			{
+				await LoadIgnorePatternsAsync(filter, workspaceDir, jobLogger);
 				await ExecuteCleanupScriptAsync(cleanupScript, filter, jobLogger);
 				try
 				{
diff --git a/Engine/Source/Programs/Horde/Drivers/JobDriver/Parser/LogParser.cs b/Engine/Source/Programs/Horde/Drivers/JobDriver/Parser/LogParser.cs
--- a/Engine/Source/Programs/Horde/Drivers/JobDriver/Parser/LogParser.cs
+++ b/Engine/Source/Programs/Horde/Drivers/JobDriver/Parser/LogParser.cs
@@ -12,6 +12,18 @@ namespace JobDriver.Parser
 	/// </summary>
 	public class LogParser : LogEventParser
 	{
+		/// <summary>
+		/// Constructor
+		/// </summary>
+		/// <param name="logger">The logger to receive parsed output messages</param>
+		/// <param name="logEventSinks">Additional sinks to receive log events</param>
+		public LogParser(ILogger logger, List<ILogEventSink>? logEventSinks = null)
+			: base(logger, logEventSinks)
+		{
+			AddMatchers(LogMatcherCollections.AutomationUtilsMatchers);
+			AddMatchers(LogMatcherCollections.UBTMatchers);
+		}
+
 		/// <summary>
 		/// Constructor
 		/// </summary>
@@ -23,7 +35,14 @@ namespace JobDriver.Parser
 		{
 			foreach (string ignorePattern in ignorePatterns)
 			{
-				IgnorePatterns.Add(new Regex(ignorePattern, RegexOptions.Compiled));
+				try
+				{
+					IgnorePatterns.Add(new IgnorePatternEntry(new Regex(ignorePattern, RegexOptions.Compiled, IgnorePatternTimeout), "(inline)"));
+				}
+				catch (ArgumentException ex)
+				{
+					logger.LogWarning("Invalid ignore pattern: {Pattern} ({Message})", ignorePattern, ex.Message);
+				}
 			}
 
 			AddMatchers(LogMatcherCollections.AutomationUtilsMatchers);
diff --git a/Engine/Source/Programs/Horde/Plugins/Build/HordeServer.Build.Tests/Streams/StreamConfigIgnorePatternTests.cs b/Engine/Source/Programs/Horde/Plugins/Build/HordeServer.Build.Tests/Streams/StreamConfigIgnorePatternTests.cs
new file mode 100644
--- /dev/null
+++ b/Engine/Source/Programs/Horde/Plugins/Build/HordeServer.Build.Tests/Streams/StreamConfigIgnorePatternTests.cs
@@ -0,0 +1,93 @@
+// Copyright Epic Games, Inc. All Rights Reserved.
+
+using EpicGames.Horde.Streams;
+using HordeServer.Artifacts;
+using HordeServer.Projects;
+using HordeServer.Streams;
+
+namespace HordeServer.Tests.Streams
+{
+	[TestClass]
+	public class StreamConfigIgnorePatternTests
+	{
+		[TestMethod]
+		public void TestIgnorePatternFiles_BothProjectAndStream()
+		{
+			ProjectConfig project = new ProjectConfig();
+			project.IgnorePatternFiles = new List<string> { "project-patterns.txt", "shared-patterns.txt" };
+
+			StreamConfig stream = new StreamConfig();
+			stream.IgnorePatternFiles = new List<string> { "stream-patterns.txt" };
+
+			stream.PostLoad(new StreamId("test-stream"), project, Array.Empty<ArtifactTypeConfig>(), null);
+
+			Assert.IsNotNull(stream.IgnorePatternFiles);
+			Assert.AreEqual(3, stream.IgnorePatternFiles.Count);
+			// Project-level files should come first
+			Assert.AreEqual("project-patterns.txt", stream.IgnorePatternFiles[0]);
+			Assert.AreEqual("shared-patterns.txt", stream.IgnorePatternFiles[1]);
+			// Stream-level files appended after
+			Assert.AreEqual("stream-patterns.txt", stream.IgnorePatternFiles[2]);
+		}
+
+		[TestMethod]
+		public void TestIgnorePatternFiles_ProjectOnly()
+		{
+			ProjectConfig project = new ProjectConfig();
+			project.IgnorePatternFiles = new List<string> { "project-patterns.txt" };
+
+			StreamConfig stream = new StreamConfig();
+			// stream.IgnorePatternFiles is null by default
+
+			stream.PostLoad(new StreamId("test-stream"), project, Array.Empty<ArtifactTypeConfig>(), null);
+
+			Assert.IsNotNull(stream.IgnorePatternFiles);
+			Assert.AreEqual(1, stream.IgnorePatternFiles.Count);
+			Assert.AreEqual("project-patterns.txt", stream.IgnorePatternFiles[0]);
+		}
+
+		[TestMethod]
+		public void TestIgnorePatternFiles_StreamOnly()
+		{
+			ProjectConfig project = new ProjectConfig();
+			// project.IgnorePatternFiles is null by default
+
+			StreamConfig stream = new StreamConfig();
+			stream.IgnorePatternFiles = new List<string> { "stream-patterns.txt" };
+
+			stream.PostLoad(new StreamId("test-stream"), project, Array.Empty<ArtifactTypeConfig>(), null);
+
+			Assert.IsNotNull(stream.IgnorePatternFiles);
+			Assert.AreEqual(1, stream.IgnorePatternFiles.Count);
+			Assert.AreEqual("stream-patterns.txt", stream.IgnorePatternFiles[0]);
+		}
+
+		[TestMethod]
+		public void TestIgnorePatternFiles_NeitherDefined()
+		{
+			ProjectConfig project = new ProjectConfig();
+			StreamConfig stream = new StreamConfig();
+
+			stream.PostLoad(new StreamId("test-stream"), project, Array.Empty<ArtifactTypeConfig>(), null);
+
+			Assert.IsNull(stream.IgnorePatternFiles);
+		}
+
+		[TestMethod]
+		public void TestIgnorePatternFiles_EmptyProjectList()
+		{
+			ProjectConfig project = new ProjectConfig();
+			project.IgnorePatternFiles = new List<string>();
+
+			StreamConfig stream = new StreamConfig();
+			stream.IgnorePatternFiles = new List<string> { "stream-patterns.txt" };
+
+			stream.PostLoad(new StreamId("test-stream"), project, Array.Empty<ArtifactTypeConfig>(), null);
+
+			// Empty project list (Count == 0) should not affect stream's list
+			Assert.IsNotNull(stream.IgnorePatternFiles);
+			Assert.AreEqual(1, stream.IgnorePatternFiles.Count);
+			Assert.AreEqual("stream-patterns.txt", stream.IgnorePatternFiles[0]);
+		}
+	}
+}
diff --git a/Engine/Source/Programs/Horde/Plugins/Build/HordeServer.Build.Tests/Stubs/Services/PerforceServiceStub.cs b/Engine/Source/Programs/Horde/Plugins/Build/HordeServer.Build.Tests/Stubs/Services/PerforceServiceStub.cs
--- a/Engine/Source/Programs/Horde/Plugins/Build/HordeServer.Build.Tests/Stubs/Services/PerforceServiceStub.cs
+++ b/Engine/Source/Programs/Horde/Plugins/Build/HordeServer.Build.Tests/Stubs/Services/PerforceServiceStub.cs
@@ -33,7 +33,8 @@ namespace HordeServer.Tests.Stubs.Services
 			}
 		}
 
-		public class Commit : ICommit
+		[System.Diagnostics.CodeAnalysis.SuppressMessage("Design", "CA1034")]
+		public sealed class Commit : ICommit
 		{
 			public StreamId StreamId { get; }
 
diff --git a/Engine/Source/Programs/Horde/Plugins/Build/HordeServer.Build/BuildServerConfig.cs b/Engine/Source/Programs/Horde/Plugins/Build/HordeServer.Build/BuildServerConfig.cs
--- a/Engine/Source/Programs/Horde/Plugins/Build/HordeServer.Build/BuildServerConfig.cs
+++ b/Engine/Source/Programs/Horde/Plugins/Build/HordeServer.Build/BuildServerConfig.cs
@@ -165,6 +165,12 @@ namespace HordeServer
 		/// </summary>
 		public long BlockCacheSizeBytes => StringUtils.ParseBytesString(BlockCacheSize);
 
+		/// <summary>
+		/// Whether to enable expanded IgnorePatterns discovery (workspace-wide scanning + config-specified files).
+		/// When false (default), only Engine directory is scanned for ignore patterns.
+		/// </summary>
+		public bool EnableExpandedIgnorePatterns { get; set; }
+
 		/// <summary>
 		/// Options for the commit service
 		/// </summary>
diff --git a/Engine/Source/Programs/Horde/Plugins/Build/HordeServer.Build/Jobs/JobRpcService.cs b/Engine/Source/Programs/Horde/Plugins/Build/HordeServer.Build/Jobs/JobRpcService.cs
--- a/Engine/Source/Programs/Horde/Plugins/Build/HordeServer.Build/Jobs/JobRpcService.cs
+++ b/Engine/Source/Programs/Horde/Plugins/Build/HordeServer.Build/Jobs/JobRpcService.cs
@@ -59,12 +59,13 @@ namespace HordeServer.Jobs
 		readonly HttpClient _httpClient;
 		readonly IClock _clock;
 		readonly IOptionsSnapshot<BuildConfig> _buildConfig;
+		readonly IOptions<BuildServerConfig> _buildServerConfig;
 		readonly ILogger _logger;
 
 		/// <summary>
 		/// Constructor
 		/// </summary>
-		public JobRpcService(IAclService aclService, JobService jobService, AgentService agentService, PoolService poolService, ConformTaskSource conformTaskSource, IArtifactCollection artifactCollection, IJobCollection jobCollection, ILogCollection logCollection, IGraphCollection graphs, ITestDataCollection testData, ITestDataCollectionV2 testDataV2, IJobStepRefCollection jobStepRefCollection, ITemplateCollection templateCollection, HttpClient httpClient, IClock clock, IOptionsSnapshot<BuildConfig> buildConfig, ILogger<JobRpcService> logger)
+		public JobRpcService(IAclService aclService, JobService jobService, AgentService agentService, PoolService poolService, ConformTaskSource conformTaskSource, IArtifactCollection artifactCollection, IJobCollection jobCollection, ILogCollection logCollection, IGraphCollection graphs, ITestDataCollection testData, ITestDataCollectionV2 testDataV2, IJobStepRefCollection jobStepRefCollection, ITemplateCollection templateCollection, HttpClient httpClient, IClock clock, IOptionsSnapshot<BuildConfig> buildConfig, IOptions<BuildServerConfig> buildServerConfig, ILogger<JobRpcService> logger)
 		{
 			_aclService = aclService;
 			_jobService = jobService;
@@ -82,6 +83,7 @@ namespace HordeServer.Jobs
 			_httpClient = httpClient;
 			_clock = clock;
 			_buildConfig = buildConfig;
+			_buildServerConfig = buildServerConfig;
 			_logger = logger;
 		}
 
@@ -258,6 +260,13 @@ namespace HordeServer.Jobs
 			}
 			response.ValidAgentTypes.Add(streamConfig.AgentTypes.Keys);
 
+			bool expandedIgnorePatterns = _buildServerConfig.Value.EnableExpandedIgnorePatterns;
+			response.EnableExpandedIgnorePatterns = expandedIgnorePatterns;
+			if (expandedIgnorePatterns && streamConfig.IgnorePatternFiles != null)
+			{
+				response.IgnorePatternFiles.AddRange(streamConfig.IgnorePatternFiles);
+			}
+
 			return response;
 		}
 
diff --git a/Engine/Source/Programs/Horde/Plugins/Build/HordeServer.Build/Projects/ProjectConfig.cs b/Engine/Source/Programs/Horde/Plugins/Build/HordeServer.Build/Projects/ProjectConfig.cs
--- a/Engine/Source/Programs/Horde/Plugins/Build/HordeServer.Build/Projects/ProjectConfig.cs
+++ b/Engine/Source/Programs/Horde/Plugins/Build/HordeServer.Build/Projects/ProjectConfig.cs
@@ -96,6 +96,13 @@ namespace HordeServer.Projects
 		/// </summary>
 		public Dictionary<string, WorkspaceConfig> WorkspaceTypes { get; set; } = [];
 
+		/// <summary>
+		/// Additional ignore pattern files applied to all streams in this project.
+		/// Streams can append additional files via their own IgnorePatternFiles property.
+		/// Paths are relative to workspace root.
+		/// </summary>
+		public List<string>? IgnorePatternFiles { get; set; }
+
 		/// <summary>
 		/// Telemetry store for Horde data for this project
 		/// </summary>
diff --git a/Engine/Source/Programs/Horde/Plugins/Build/HordeServer.Build/Streams/StreamConfig.cs b/Engine/Source/Programs/Horde/Plugins/Build/HordeServer.Build/Streams/StreamConfig.cs
--- a/Engine/Source/Programs/Horde/Plugins/Build/HordeServer.Build/Streams/StreamConfig.cs
+++ b/Engine/Source/Programs/Horde/Plugins/Build/HordeServer.Build/Streams/StreamConfig.cs
@@ -166,6 +166,14 @@ namespace HordeServer.Streams
 		/// </summary>
 		public string EnginePath { get; set; } = "Engine";
 
+		/// <summary>
+		/// Additional file paths (relative to workspace root) containing ignore patterns for the log parser.
+		/// Each file should contain one regex per line, with lines starting with '#' treated as comments.
+		/// Inherits from ProjectConfig and appends stream-level entries.
+		/// </summary>
+		[ConfigMergeStrategy(ConfigMergeStrategy.Append)]
+		public List<string>? IgnorePatternFiles { get; set; }
+
 		/// <summary>
 		/// The perforce cluster containing the stream
 		/// </summary>
@@ -348,6 +356,18 @@ namespace HordeServer.Streams
 				WorkspaceTypes.TryAdd(wid, wc);
 			}
 
+			if (projectConfig.IgnorePatternFiles != null && projectConfig.IgnorePatternFiles.Count > 0)
+			{
+				if (IgnorePatternFiles == null)
+				{
+					IgnorePatternFiles = new List<string>(projectConfig.IgnorePatternFiles);
+				}
+				else
+				{
+					IgnorePatternFiles.InsertRange(0, projectConfig.IgnorePatternFiles);
+				}
+			}
+
 			ConfigObject.MergeDefaults(AgentTypes.Select(x => (x.Key, x.Value.Base, x.Value)));
 			ConfigObject.MergeDefaults(WorkspaceTypes.Select(x => (x.Key, x.Value.Base, x.Value)));
 			ConfigObject.MergeDefaults(Templates.Select(x => (x.Id, x.Base, x)));
diff --git a/Engine/Source/Programs/Shared/EpicGames.Core.Tests/LogEventParserTests.cs b/Engine/Source/Programs/Shared/EpicGames.Core.Tests/LogEventParserTests.cs
--- a/Engine/Source/Programs/Shared/EpicGames.Core.Tests/LogEventParserTests.cs
+++ b/Engine/Source/Programs/Shared/EpicGames.Core.Tests/LogEventParserTests.cs
@@ -2,7 +2,10 @@
 
 using System;
 using System.Collections.Generic;
+using System.IO;
 using System.Linq;
+using System.Text.RegularExpressions;
+using System.Threading.Tasks;
 using Microsoft.Extensions.Logging;
 using Microsoft.VisualStudio.TestTools.UnitTesting;
 
@@ -286,7 +289,7 @@ public class LogEventParserTests
 		Assert.AreEqual(1, lowest.MatchCalls);
 		// 0 because of lowest 3 line count requirement
 		Assert.AreEqual(0, capture.Events.Count);
-		
+
 		parser.WriteLine("\t" + LineToLog);
 		Assert.AreEqual(2, highest.MatchCalls);
 		Assert.AreEqual(2, lowest.MatchCalls);
@@ -298,27 +301,27 @@ public class LogEventParserTests
 		Assert.AreEqual(3, lowest.MatchCalls);
 		// 3 because lowest 3 line count requirement is met
 		Assert.AreEqual(3, capture.Events.Count);
-		
+
 		parser.WriteLine("\t" + LineToLog);
 		Assert.AreEqual(4, highest.MatchCalls);
 		// 3 because highest it incomplete hanging
 		Assert.AreEqual(3, lowest.MatchCalls);
 		// 4 because of highest was incomplete hanging
 		Assert.AreEqual(4, capture.Events.Count);
-		
+
 		parser.WriteLine(LineToLog);
 		// 6 because highest gets called again once complete
 		Assert.AreEqual(6, highest.MatchCalls);
 		Assert.AreEqual(4, lowest.MatchCalls);
 		// 4 because of lowest 3 line count requirement
 		Assert.AreEqual(4, capture.Events.Count);
-		
+
 		parser.WriteLine(LineToLog);
 		Assert.AreEqual(7, highest.MatchCalls);
 		Assert.AreEqual(5, lowest.MatchCalls);
 		// 4 because of lowest 3 line count requirement
 		Assert.AreEqual(4, capture.Events.Count);
-		
+
 		parser.WriteLine(LineToLog);
 		Assert.AreEqual(9, highest.MatchCalls);
 		Assert.AreEqual(7, lowest.MatchCalls);
@@ -332,6 +335,339 @@ public class LogEventParserTests
 		Assert.AreEqual(7, capture.Events.Count);
 	}
 
+	#region Ignore Pattern Tests
+
+	[TestMethod]
+	public void TestDiscoverIgnorePatternFiles_EngineDir()
+	{
+		using TempWorkspace workspace = new();
+		workspace.CreatePatternFile("Engine/Build/Horde/IgnorePatterns.txt", ".*test.*");
+
+		List<FileReference> files = LogEventParser.DiscoverIgnorePatternFiles(workspace.Dir);
+
+		Assert.AreEqual(1, files.Count);
+		Assert.IsTrue(files[0].FullName.Contains("Engine", StringComparison.Ordinal));
+	}
+
+	[TestMethod]
+	public void TestDiscoverIgnorePatternFiles_ProjectDir()
+	{
+		using TempWorkspace workspace = new();
+		workspace.CreatePatternFile("Engine/Build/Horde/IgnorePatterns.txt", ".*engine.*");
+		workspace.CreatePatternFile("MyGame/Build/Horde/IgnorePatterns.txt", ".*game.*");
+
+		List<FileReference> files = LogEventParser.DiscoverIgnorePatternFiles(workspace.Dir);
+
+		Assert.AreEqual(2, files.Count);
+		Assert.IsTrue(files.Any(f => f.FullName.Contains("Engine", StringComparison.Ordinal)));
+		Assert.IsTrue(files.Any(f => f.FullName.Contains("MyGame", StringComparison.Ordinal)));
+	}
+
+	[TestMethod]
+	public void TestDiscoverIgnorePatternFiles_RestrictedAndPlatforms()
+	{
+		using TempWorkspace workspace = new();
+		workspace.CreatePatternFile("Engine/Build/Horde/IgnorePatterns.txt", ".*engine.*");
+		workspace.CreatePatternFile("Engine/Restricted/ConsoleX/Build/Horde/IgnorePatterns.txt", ".*consolex.*");
+		workspace.CreatePatternFile("MyGame/Platforms/PS5/Build/Horde/IgnorePatterns.txt", ".*ps5.*");
+
+		List<FileReference> files = LogEventParser.DiscoverIgnorePatternFiles(workspace.Dir);
+
+		Assert.AreEqual(3, files.Count);
+		Assert.IsTrue(files.Any(f => f.FullName.Contains(Path.Combine("Engine", "Build"), StringComparison.Ordinal)));
+		Assert.IsTrue(files.Any(f => f.FullName.Contains("ConsoleX", StringComparison.Ordinal)));
+		Assert.IsTrue(files.Any(f => f.FullName.Contains("PS5", StringComparison.Ordinal)));
+	}
+
+	[TestMethod]
+	public void TestDiscoverIgnorePatternFiles_ProjectRestrictedSubdirs()
+	{
+		using TempWorkspace workspace = new();
+		workspace.CreatePatternFile("MyGame/Build/Horde/IgnorePatterns.txt", ".*game.*");
+		workspace.CreatePatternFile("MyGame/Restricted/Internal/Build/Horde/IgnorePatterns.txt", ".*internal.*");
+		workspace.CreatePatternFile("MyGame/Platforms/Switch/Build/Horde/IgnorePatterns.txt", ".*switch.*");
+
+		List<FileReference> files = LogEventParser.DiscoverIgnorePatternFiles(workspace.Dir);
+
+		Assert.AreEqual(3, files.Count);
+		Assert.IsTrue(files.Any(f => f.FullName.Contains(Path.Combine("MyGame", "Build"), StringComparison.Ordinal)));
+		Assert.IsTrue(files.Any(f => f.FullName.Contains("Internal", StringComparison.Ordinal)));
+		Assert.IsTrue(files.Any(f => f.FullName.Contains("Switch", StringComparison.Ordinal)));
+	}
+
+	[TestMethod]
+	public void TestDiscoverIgnorePatternFiles_EmptyWorkspace()
+	{
+		using TempWorkspace workspace = new();
+
+		List<FileReference> files = LogEventParser.DiscoverIgnorePatternFiles(workspace.Dir);
+
+		Assert.AreEqual(0, files.Count);
+	}
+
+	[TestMethod]
+	public void TestDiscoverIgnorePatternFiles_EngineDirNotDuplicated()
+	{
+		using TempWorkspace workspace = new();
+		workspace.CreatePatternFile("Engine/Build/Horde/IgnorePatterns.txt", ".*test.*");
+
+		List<FileReference> files = LogEventParser.DiscoverIgnorePatternFiles(workspace.Dir);
+
+		// Engine should appear only once even though it's both the explicit engine dir and a top-level dir
+		int engineCount = files.Count(f => f.FullName.Contains(Path.Combine("Engine", "Build", "Horde", "IgnorePatterns.txt"), StringComparison.Ordinal));
+		Assert.AreEqual(1, engineCount);
+	}
+
+	[TestMethod]
+	public void TestDiscoverIgnorePatternFiles_CustomEnginePath()
+	{
+		using TempWorkspace workspace = new();
+		workspace.CreatePatternFile("CustomEngine/Build/Horde/IgnorePatterns.txt", ".*custom.*");
+		workspace.CreatePatternFile("Engine/Build/Horde/IgnorePatterns.txt", ".*default.*");
+
+		List<FileReference> files = LogEventParser.DiscoverIgnorePatternFiles(workspace.Dir, "CustomEngine");
+
+		// Both should be discovered - CustomEngine as the engine dir, Engine as a top-level dir
+		Assert.AreEqual(2, files.Count);
+		Assert.IsTrue(files.Any(f => f.FullName.Contains("CustomEngine", StringComparison.Ordinal)));
+		Assert.IsTrue(files.Any(f => f.FullName.Contains(Path.Combine("Engine", "Build"), StringComparison.Ordinal)));
+	}
+
+	[TestMethod]
+	public async Task TestReadIgnorePatterns_ValidPatternsAsync()
+	{
+		using TempWorkspace workspace = new();
+		FileReference file = workspace.CreatePatternFile("patterns.txt",
+			".*warning.*\n.*error X4000.*\n.*LNK4099.*");
+
+		LoggerCapture capture = new();
+		using LogEventParser parser = new(capture);
+		await parser.ReadIgnorePatternsAsync(file);
+
+		Assert.AreEqual(3, parser.IgnorePatterns.Count);
+		// Verify source file is tracked on each entry
+		foreach (IgnorePatternEntry entry in parser.IgnorePatterns)
+		{
+			Assert.AreEqual(file.FullName, entry.SourceFile);
+		}
+	}
+
+	[TestMethod]
+	public async Task TestReadIgnorePatterns_CommentsAndBlankLinesAsync()
+	{
+		using TempWorkspace workspace = new();
+		FileReference file = workspace.CreatePatternFile("patterns.txt",
+			"# This is a comment\n" +
+			"\n" +
+			".*warning.*\n" +
+			"   # Indented comment\n" +
+			"  \n" +
+			".*error.*");
+
+		LoggerCapture capture = new();
+		using LogEventParser parser = new(capture);
+		await parser.ReadIgnorePatternsAsync(file);
+
+		// Only non-comment, non-blank lines should be loaded
+		Assert.AreEqual(2, parser.IgnorePatterns.Count);
+	}
+
+	[TestMethod]
+	public async Task TestReadIgnorePatterns_InvalidRegexAsync()
+	{
+		using TempWorkspace workspace = new();
+		FileReference file = workspace.CreatePatternFile("patterns.txt",
+			".*valid_pattern.*\n" +
+			"[invalid_unclosed_bracket\n" +
+			".*another_valid.*");
+
+		LoggerCapture capture = new();
+		using LogEventParser parser = new(capture);
+		await parser.ReadIgnorePatternsAsync(file);
+
+		// Invalid regex should be skipped, valid ones still loaded
+		Assert.AreEqual(2, parser.IgnorePatterns.Count);
+		// Warning should have been logged about the invalid pattern
+		Assert.IsTrue(capture.Events.Any(e => e.Level == LogLevel.Warning));
+	}
+
+	[TestMethod]
+	public void TestIgnorePattern_SuppressesMatchedError()
+	{
+		LoggerCapture capture = new();
+		TestMatcher matcher = new(LogEventPriority.Highest, LogLevel.Error);
+
+		using (LogEventParser parser = new(capture))
+		{
+			parser.AddMatchers([matcher]);
+			parser.IgnorePatterns.Add(new IgnorePatternEntry(new Regex(".*hello.*", RegexOptions.Compiled), "test"));
+			parser.WriteLine(LineToLog); // LineToLog is "hello"
+		}
+
+		// The error should be suppressed - line is logged at Information level (not Error)
+		Assert.AreEqual(1, capture.Events.Count);
+		Assert.AreEqual(LogLevel.Information, capture.Events[0].Level);
+	}
+
+	[TestMethod]
+	public void TestIgnorePattern_DoesNotSuppressNonMatchingLine()
+	{
+		LoggerCapture capture = new();
+		TestMatcher matcher = new(LogEventPriority.Highest, LogLevel.Error);
+
+		using (LogEventParser parser = new(capture))
+		{
+			parser.AddMatchers([matcher]);
+			parser.IgnorePatterns.Add(new IgnorePatternEntry(new Regex(".*goodbye.*", RegexOptions.Compiled), "test"));
+			parser.WriteLine(LineToLog); // LineToLog is "hello", doesn't match ".*goodbye.*"
+		}
+
+		// The error should NOT be suppressed since the ignore pattern doesn't match
+		Assert.AreEqual(1, capture.Events.Count);
+		Assert.AreEqual(LogLevel.Error, capture.Events[0].Level);
+	}
+
+	[TestMethod]
+	public void TestIgnorePattern_MultiplePatterns()
+	{
+		LoggerCapture capture = new();
+		TestMatcher matcher = new(LogEventPriority.Highest, LogLevel.Error);
+
+		using (LogEventParser parser = new(capture))
+		{
+			parser.AddMatchers([matcher]);
+			parser.IgnorePatterns.Add(new IgnorePatternEntry(new Regex(".*unrelated.*", RegexOptions.Compiled), "test"));
+			parser.IgnorePatterns.Add(new IgnorePatternEntry(new Regex(".*hello.*", RegexOptions.Compiled), "test"));
+			parser.IgnorePatterns.Add(new IgnorePatternEntry(new Regex(".*other.*", RegexOptions.Compiled), "test"));
+			parser.WriteLine(LineToLog);
+		}
+
+		// Second pattern should match, suppressing the error
+		Assert.AreEqual(1, capture.Events.Count);
+		Assert.AreEqual(LogLevel.Information, capture.Events[0].Level);
+	}
+
+	[TestMethod]
+	public void TestIgnorePattern_StructuredProperties()
+	{
+		LoggerCapture capture = new();
+		TestMatcher matcher = new(LogEventPriority.Highest, LogLevel.Error);
+
+		using (LogEventParser parser = new(capture))
+		{
+			parser.AddMatchers([matcher]);
+			parser.IgnorePatterns.Add(new IgnorePatternEntry(new Regex(".*hello.*", RegexOptions.Compiled), "/workspace/Engine/Build/Horde/IgnorePatterns.txt"));
+			parser.WriteLine(LineToLog);
+		}
+
+		Assert.AreEqual(1, capture.Events.Count);
+		Assert.AreEqual(LogLevel.Information, capture.Events[0].Level);
+
+		// Verify structured properties are present
+		Assert.IsNotNull(capture.Events[0].Properties);
+		KeyValuePair<string, object?>? fileProperty = capture.Events[0].Properties?.FirstOrDefault(p => p.Key == "IgnorePatternFile");
+		KeyValuePair<string, object?>? patternProperty = capture.Events[0].Properties?.FirstOrDefault(p => p.Key == "IgnorePattern");
+		Assert.IsNotNull(fileProperty);
+		Assert.AreEqual("/workspace/Engine/Build/Horde/IgnorePatterns.txt", fileProperty.Value.Value?.ToString());
+		Assert.IsNotNull(patternProperty);
+		Assert.AreEqual(".*hello.*", patternProperty.Value.Value?.ToString());
+	}
+
+	[TestMethod]
+	public void TestIgnorePattern_RegexTimeoutHandled()
+	{
+		LoggerCapture capture = new();
+		TestMatcher matcher = new(LogEventPriority.Highest, LogLevel.Error);
+
+		using (LogEventParser parser = new(capture))
+		{
+			parser.AddMatchers([matcher]);
+			// Pathological pattern with very short timeout to trigger RegexMatchTimeoutException
+			parser.IgnorePatterns.Add(new IgnorePatternEntry(new Regex(@"(a+)+$", RegexOptions.None, TimeSpan.FromTicks(1)), "test"));
+			// String that causes catastrophic backtracking: many 'a's followed by non-'a'
+			parser.WriteLine(new string('a', 50) + "!");
+		}
+
+		// The error should NOT be suppressed because the pattern either timed out or didn't match.
+		// In either case, the line should appear as an error event.
+		Assert.IsTrue(capture.Events.Any(e => e.Level == LogLevel.Error));
+	}
+
+	[TestMethod]
+	public async Task TestReadIgnorePatternsFromWorkspaceAsync()
+	{
+		using TempWorkspace workspace = new();
+		workspace.CreatePatternFile("Engine/Build/Horde/IgnorePatterns.txt", ".*engine_warning.*");
+		workspace.CreatePatternFile("MyGame/Build/Horde/IgnorePatterns.txt", ".*game_warning.*");
+
+		LoggerCapture capture = new();
+		using LogEventParser parser = new(capture);
+		await parser.ReadIgnorePatternsFromWorkspaceAsync(workspace.Dir);
+
+		Assert.AreEqual(2, parser.IgnorePatterns.Count);
+	}
+
+	[TestMethod]
+	public async Task TestReadIgnorePatternsFromWorkspace_WithRestrictedDirsAsync()
+	{
+		using TempWorkspace workspace = new();
+		workspace.CreatePatternFile("Engine/Build/Horde/IgnorePatterns.txt", ".*engine.*");
+		workspace.CreatePatternFile("Engine/Restricted/Internal/Build/Horde/IgnorePatterns.txt", ".*internal.*");
+		workspace.CreatePatternFile("MyGame/Build/Horde/IgnorePatterns.txt", ".*game.*");
+
+		LoggerCapture capture = new();
+		using LogEventParser parser = new(capture);
+		await parser.ReadIgnorePatternsFromWorkspaceAsync(workspace.Dir);
+
+		Assert.AreEqual(3, parser.IgnorePatterns.Count);
+	}
+
+	[TestMethod]
+	public async Task TestReadIgnorePatterns_MissingAdditionalFilesAsync()
+	{
+		LoggerCapture capture = new();
+		using LogEventParser parser = new(capture);
+		FileReference missingFile = new(Path.Combine(Path.GetTempPath(), Guid.NewGuid().ToString("N"), "nonexistent.txt"));
+		await parser.ReadIgnorePatternsAsync([missingFile]);
+
+		// No patterns should be loaded
+		Assert.AreEqual(0, parser.IgnorePatterns.Count);
+		// Warning should have been logged about missing file
+		Assert.IsTrue(capture.Events.Any(e => e.Level == LogLevel.Warning));
+	}
+
+	[TestMethod]
+	public async Task TestReadIgnorePatterns_AdditionalFilesLoadedAsync()
+	{
+		using TempWorkspace workspace = new();
+		FileReference file1 = workspace.CreatePatternFile("extra/patterns1.txt", ".*pattern1.*");
+		FileReference file2 = workspace.CreatePatternFile("extra/patterns2.txt", ".*pattern2.*\n.*pattern3.*");
+
+		LoggerCapture capture = new();
+		using LogEventParser parser = new(capture);
+		await parser.ReadIgnorePatternsAsync([file1, file2]);
+
+		Assert.AreEqual(3, parser.IgnorePatterns.Count);
+	}
+
+	[TestMethod]
+	public async Task TestReadIgnorePatterns_MixedValidAndMissingAdditionalFilesAsync()
+	{
+		using TempWorkspace workspace = new();
+		FileReference validFile = workspace.CreatePatternFile("extra/patterns.txt", ".*valid.*");
+		FileReference missingFile = new(Path.Combine(workspace.Dir.FullName, "nonexistent.txt"));
+
+		LoggerCapture capture = new();
+		using LogEventParser parser = new(capture);
+		await parser.ReadIgnorePatternsAsync([validFile, missingFile]);
+
+		// Valid file patterns loaded, missing file logged as warning
+		Assert.AreEqual(1, parser.IgnorePatterns.Count);
+		Assert.IsTrue(capture.Events.Any(e => e.Level == LogLevel.Warning));
+	}
+
+	#endregion
+
 	private const string LogLineProperty = "LogLine";
 	private const string LineToLog = "hello";
 
@@ -384,7 +720,7 @@ public class LogEventParserTests
 
 		public string GetMatcherName() => Priority.ToString();
 	}
-	
+
 	private class TestHangingLinesMatcher(LogEventPriority priority, LogLevel level, int requiredLineCount = 1) : ILogEventMatcher
 	{
 		public LogEventMatch? Match(ILogCursor cursor)
@@ -408,4 +744,38 @@ public class LogEventParserTests
 
 		public string GetMatcherName() => Priority.ToString();
 	}
+
+	/// <summary>
+	/// Helper class for creating temporary workspace directory structures for testing
+	/// </summary>
+	private class TempWorkspace : IDisposable
+	{
+		public DirectoryReference Dir { get; }
+
+		public TempWorkspace()
+		{
+			Dir = new DirectoryReference(Path.Combine(Path.GetTempPath(), "HordeTest_" + Guid.NewGuid().ToString("N")));
+			Directory.CreateDirectory(Dir.FullName);
+		}
+
+		public FileReference CreatePatternFile(string relativePath, string contents)
+		{
+			string fullPath = Path.Combine(Dir.FullName, relativePath);
+			Directory.CreateDirectory(Path.GetDirectoryName(fullPath)!);
+			File.WriteAllText(fullPath, contents);
+			return new FileReference(fullPath);
+		}
+
+		public void Dispose()
+		{
+			try
+			{
+				Directory.Delete(Dir.FullName, true);
+			}
+			catch
+			{
+				// Best effort cleanup
+			}
+		}
+	}
 }
diff --git a/Engine/Source/Programs/Shared/EpicGames.Core/LogEventParser.cs b/Engine/Source/Programs/Shared/EpicGames.Core/LogEventParser.cs
--- a/Engine/Source/Programs/Shared/EpicGames.Core/LogEventParser.cs
+++ b/Engine/Source/Programs/Shared/EpicGames.Core/LogEventParser.cs
@@ -121,6 +121,13 @@ namespace EpicGames.Core
 		string GetMatcherName() => GetType().Name;
 	}
 
+	/// <summary>
+	/// An ignore pattern with metadata about its source file
+	/// </summary>
+	/// <param name="Pattern">The compiled regex pattern</param>
+	/// <param name="SourceFile">Absolute path to the file that defined this pattern</param>
+	public readonly record struct IgnorePatternEntry(Regex Pattern, string SourceFile);
+
 	/// <summary>
 	/// Turns raw text output into structured logging events
 	/// </summary>
@@ -131,6 +138,16 @@ namespace EpicGames.Core
 	/// </remarks>
 	public partial class LogEventParser : IDisposable
 	{
+		/// <summary>
+		/// Timeout applied to each ignore pattern regex match.
+		/// </summary>
+		static readonly TimeSpan s_ignorePatternTimeout = TimeSpan.FromMilliseconds(200);
+
+		/// <summary>
+		/// Timeout applied to each ignore pattern regex match. Exposed for derived classes.
+		/// </summary>
+		protected static TimeSpan IgnorePatternTimeout => s_ignorePatternTimeout;
+
 		/// <summary>
 		/// A list of enabled event matchers for this parser.
 		/// </summary>
@@ -139,7 +156,7 @@ namespace EpicGames.Core
 		/// <summary>
 		/// List of patterns to ignore
 		/// </summary>
-		public List<Regex> IgnorePatterns { get; } = [];
+		public List<IgnorePatternEntry> IgnorePatterns { get; } = [];
 
 		/// <summary>
 		/// A list of matchers all matchers, including disabled matchers.
@@ -301,6 +318,7 @@ namespace EpicGames.Core
 		/// </summary>
 		/// <param name="rootDir"></param>
 		/// <returns></returns>
+		[Obsolete("Use ReadIgnorePatternsFromWorkspaceAsync or DiscoverIgnorePatternFiles instead")]
 		public async Task ReadIgnorePatternsAsync(DirectoryReference rootDir)
 		{
 			Stopwatch timer = Stopwatch.StartNew();
@@ -344,19 +362,28 @@ namespace EpicGames.Core
 		{
 			string[] lines = await FileReference.ReadAllLinesAsync(ignorePatternFile);
 
-			List<Regex> patterns = [];
-			foreach (string line in lines)
+			List<IgnorePatternEntry> entries = [];
+			for (int idx = 0; idx < lines.Length; idx++)
 			{
-				string trimLine = line.Trim();
+				string trimLine = lines[idx].Trim();
 				if (trimLine.Length > 0 && trimLine[0] != '#')
 				{
-					patterns.Add(new Regex(trimLine));
+					try
+					{
+						Regex regex = new(trimLine, RegexOptions.Compiled, s_ignorePatternTimeout);
+						entries.Add(new IgnorePatternEntry(regex, ignorePatternFile.FullName));
+					}
+					catch (ArgumentException ex)
+					{
+						_logger.LogWarning("Invalid regex in {File} line {LineNumber}: {Pattern} ({Message})",
+							ignorePatternFile, idx + 1, trimLine, ex.Message);
+					}
 				}
 			}
 
 			lock (IgnorePatterns)
 			{
-				IgnorePatterns.AddRange(patterns);
+				IgnorePatterns.AddRange(entries);
 			}
 		}
 
@@ -373,6 +400,189 @@ namespace EpicGames.Core
 			}
 		}
 
+		static IgnorePatternEntry? FindMatchingIgnorePattern(List<IgnorePatternEntry> ignorePatterns, string line, ILogger logger)
+		{
+			foreach (IgnorePatternEntry entry in ignorePatterns)
+			{
+				try
+				{
+					if (entry.Pattern.IsMatch(line))
+					{
+						return entry;
+					}
+				}
+				catch (RegexMatchTimeoutException)
+				{
+					logger.LogWarning("Ignore pattern timed out on line (pattern: {Pattern}): {Line}",
+						entry.Pattern.ToString(), line.Length > 200 ? line[..200] + "..." : line);
+				}
+			}
+			return null;
+		}
+
+		/// <summary>
+		/// Writes a line that was suppressed by an ignore pattern, including structured properties
+		/// identifying which pattern file and regex caused the suppression.
+		/// </summary>
+		void WriteIgnoredLine(string line, IgnorePatternEntry ignoredBy)
+		{
+			KeyValuePair<string, object?>[] state =
+			[
+				new(MessageTemplate.FormatPropertyName, "{IgnoredLine}"),
+				new("IgnoredLine", line),
+				new("IgnorePatternFile", ignoredBy.SourceFile),
+				new("IgnorePattern", ignoredBy.Pattern.ToString()),
+			];
+
+			_logger.Log(LogLevel.Information, KnownLogEvents.None, state, null, (s, _) => line);
+		}
+
+		/// <summary>
+		/// Discover all IgnorePatterns.txt files in a workspace by scanning Engine, project,
+		/// and Restricted/Platforms subdirectories.
+		/// </summary>
+		public static List<FileReference> DiscoverIgnorePatternFiles(DirectoryReference workspaceDir, string enginePath = "Engine")
+		{
+			HashSet<DirectoryReference> scannedDirs = [];
+			List<DirectoryReference> baseDirs = [];
+
+			// Engine directory
+			DirectoryReference engineDir = DirectoryReference.Combine(workspaceDir, enginePath);
+			scannedDirs.Add(engineDir);
+			baseDirs.Add(engineDir);
+
+			// All top-level workspace directories
+			if (DirectoryReference.Exists(workspaceDir))
+			{
+				foreach (DirectoryReference subDir in DirectoryReference.EnumerateDirectories(workspaceDir))
+				{
+					if (scannedDirs.Add(subDir))
+					{
+						baseDirs.Add(subDir);
+					}
+				}
+			}
+
+			// Restricted/Platforms under each base directory
+			int count = baseDirs.Count;
+			for (int idx = 0; idx < count; idx++)
+			{
+				foreach (string subFolder in new[] { "Restricted", "Platforms" })
+				{
+					DirectoryReference subDir = DirectoryReference.Combine(baseDirs[idx], subFolder);
+					if (DirectoryReference.Exists(subDir))
+					{
+						foreach (DirectoryReference child in DirectoryReference.EnumerateDirectories(subDir))
+						{
+							if (scannedDirs.Add(child))
+							{
+								baseDirs.Add(child);
+							}
+						}
+					}
+				}
+			}
+
+			// Find all existing IgnorePatterns.txt files
+			List<FileReference> files = [];
+			foreach (DirectoryReference baseDir in baseDirs)
+			{
+				FileReference file = FileReference.Combine(baseDir, "Build", "Horde", "IgnorePatterns.txt");
+				if (FileReference.Exists(file))
+				{
+					files.Add(file);
+				}
+			}
+			return files;
+		}
+
+		/// <summary>
+		/// Legacy ignore pattern file discovery - scans only Engine + Engine/Restricted/* + Engine/Platforms/*.
+		/// Use <see cref="DiscoverIgnorePatternFiles"/> for expanded workspace-wide scanning.
+		/// </summary>
+		public static List<FileReference> DiscoverLegacyIgnorePatternFiles(DirectoryReference workspaceDir, string enginePath = "Engine")
+		{
+			DirectoryReference engineDir = DirectoryReference.Combine(workspaceDir, enginePath);
+			List<DirectoryReference> baseDirs = [engineDir];
+			AddRestrictedDirs(baseDirs, "Restricted");
+			AddRestrictedDirs(baseDirs, "Platforms");
+
+			List<FileReference> files = [];
+			foreach (DirectoryReference baseDir in baseDirs)
+			{
+				FileReference file = FileReference.Combine(baseDir, "Build", "Horde", "IgnorePatterns.txt");
+				if (FileReference.Exists(file))
+				{
+					files.Add(file);
+				}
+			}
+			return files;
+		}
+
+		/// <summary>
+		/// Read ignore patterns from the workspace, scanning Engine and project directories
+		/// </summary>
+		public async Task ReadIgnorePatternsFromWorkspaceAsync(DirectoryReference workspaceDir, string enginePath = "Engine")
+		{
+			Stopwatch timer = Stopwatch.StartNew();
+
+			List<FileReference> files = DiscoverIgnorePatternFiles(workspaceDir, enginePath);
+
+			List<(FileReference, Task)> tasks = [];
+			foreach (FileReference file in files)
+			{
+				_logger.LogDebug("Reading ignore patterns from {File}...", file);
+				tasks.Add((file, ReadIgnorePatternsAsync(file)));
+			}
+
+			foreach ((FileReference file, Task task) in tasks)
+			{
+				try
+				{
+					await task;
+				}
+				catch (Exception ex)
+				{
+					_logger.LogError(ex, "Exception reading patterns from {File}: {Message}", file, ex.Message);
+				}
+			}
+
+			_logger.LogDebug("Read {PatternCount} ignore patterns from {FileCount} files in {TimeMs}ms",
+				IgnorePatterns.Count, tasks.Count, timer.ElapsedMilliseconds);
+		}
+
+		/// <summary>
+		/// Read ignore patterns from additional file paths
+		/// </summary>
+		public async Task ReadIgnorePatternsAsync(IEnumerable<FileReference> additionalFiles)
+		{
+			List<(FileReference, Task)> tasks = [];
+			foreach (FileReference file in additionalFiles)
+			{
+				if (FileReference.Exists(file))
+				{
+					_logger.LogDebug("Reading ignore patterns from {File}...", file);
+					tasks.Add((file, ReadIgnorePatternsAsync(file)));
+				}
+				else
+				{
+					_logger.LogWarning("Ignore pattern file not found: {File}", file);
+				}
+			}
+
+			foreach ((FileReference file, Task task) in tasks)
+			{
+				try
+				{
+					await task;
+				}
+				catch (Exception ex)
+				{
+					_logger.LogError(ex, "Exception reading patterns from {File}: {Message}", file, ex.Message);
+				}
+			}
+		}
+
 		/// <summary>
 		/// Writes a line to the event filter
 		/// </summary>
@@ -660,11 +870,11 @@ namespace EpicGames.Core
 					else if (_matchingEnabled >= 0 && !disableMatching)
 					{
 						(ILogEventMatcher? matcher, LogEventMatch? match) = MatchEvent();
-						
+
 						// Bail out if we need more data, or match is incomplete
 						if (_buffer.Length < 1024 && !bFlush && (_buffer.NeedMoreData || match is { IsComplete: false }))
 						{
-							if (!_buffer.NeedMoreData && !IgnorePatterns.Any(ignorePattern => ignorePattern.IsMatch(line)))
+							if (!_buffer.NeedMoreData && FindMatchingIgnorePattern(IgnorePatterns, line, _logger) == null)
 							{
 								WriteEvents(match!.Events);
 								Debug.Assert(matcher != null);
@@ -675,12 +885,20 @@ namespace EpicGames.Core
 						}
 
 						// Report the error to the listeners
-						// If we did match something, check if it's not negated by an ignore pattern. We typically have relatively few errors and many more ignore patterns than matchers, so it's quicker 
+						// If we did match something, check if it's not negated by an ignore pattern. We typically have relatively few errors and many more ignore patterns than matchers, so it's quicker
 						// to check them in response to an identified error than to treat them as matchers of their own.
-						if (match != null && !IgnorePatterns.Any(ignorePattern => ignorePattern.IsMatch(line)))
+						if (match != null)
 						{
-							WriteEvents(match.Events);
-							_buffer.Advance(match.Events.Count);
+							IgnorePatternEntry? ignoredBy = FindMatchingIgnorePattern(IgnorePatterns, line, _logger);
+							if (ignoredBy == null)
+							{
+								WriteEvents(match.Events);
+								_buffer.Advance(match.Events.Count);
+								continue;
+							}
+
+							WriteIgnoredLine(line, ignoredBy.Value);
+							_buffer.MoveNext();
 							continue;
 						}
 					}
