Horde: Fix issue with log parser coercion

diff --git a/Engine/Source/Programs/Shared/EpicGames.Core/LogEventParser.cs b/Engine/Source/Programs/Shared/EpicGames.Core/LogEventParser.cs
--- a/Engine/Source/Programs/Shared/EpicGames.Core/LogEventParser.cs
+++ b/Engine/Source/Programs/Shared/EpicGames.Core/LogEventParser.cs
@@ -126,7 +126,13 @@ namespace EpicGames.Core
 	/// </summary>
 	/// <param name="Pattern">The compiled regex pattern</param>
 	/// <param name="SourceFile">Absolute path to the file that defined this pattern</param>
-	public readonly record struct IgnorePatternEntry(Regex Pattern, string SourceFile);
+	public readonly record struct IgnorePatternEntry(Regex Pattern, string SourceFile)
+	{
+		/// <summary>
+		/// Implicit conversion to Regex for backward compatibility with code that expects IgnorePatterns to contain Regex objects.
+		/// </summary>
+		public static implicit operator Regex(IgnorePatternEntry entry) => entry.Pattern;
+	}
 
 	/// <summary>
 	/// Turns raw text output into structured logging events
