Horde: Add backwards compatibility to log event parser

diff --git a/Engine/Source/Programs/Shared/EpicGames.Core/LogEventParser.cs b/Engine/Source/Programs/Shared/EpicGames.Core/LogEventParser.cs
--- a/Engine/Source/Programs/Shared/EpicGames.Core/LogEventParser.cs
+++ b/Engine/Source/Programs/Shared/EpicGames.Core/LogEventParser.cs
@@ -132,6 +132,28 @@ namespace EpicGames.Core
 		/// Implicit conversion to Regex for backward compatibility with code that expects IgnorePatterns to contain Regex objects.
 		/// </summary>
 		public static implicit operator Regex(IgnorePatternEntry entry) => entry.Pattern;
+
+		/// <summary>
+		/// Implicit conversion from Regex for backward compatibility with code that adds Regex objects directly.
+		/// </summary>
+		public static implicit operator IgnorePatternEntry(Regex pattern) => new(pattern, "unknown");
+	}
+
+	/// <summary>
+	/// Extension methods for backward compatibility with code that passes Regex collections to IgnorePatternEntry lists.
+	/// </summary>
+	public static class IgnorePatternEntryExtensions
+	{
+		/// <summary>
+		/// Add a range of Regex objects as ignore pattern entries.
+		/// </summary>
+		public static void AddRange(this List<IgnorePatternEntry> list, IEnumerable<Regex> patterns)
+		{
+			foreach (Regex regex in patterns)
+			{
+				list.Add(new IgnorePatternEntry(regex, "unknown"));
+			}
+		}
 	}
 
 	/// <summary>
