Horde: Apply ignore patterns to structured JSON log events

Structured (JSON) log events bypassed ignore pattern checking
entirely, meaning IgnorePatterns.txt files could not suppress
build health issues from structured logging. This adds ignore
pattern checking for JSON events at Warning level or above in
both the WriteLine and WriteData/AddLine code paths.

diff --git a/Engine/Source/Programs/Shared/EpicGames.Core.Tests/LogEventParserTests.cs b/Engine/Source/Programs/Shared/EpicGames.Core.Tests/LogEventParserTests.cs
--- a/Engine/Source/Programs/Shared/EpicGames.Core.Tests/LogEventParserTests.cs
+++ b/Engine/Source/Programs/Shared/EpicGames.Core.Tests/LogEventParserTests.cs
@@ -668,6 +668,122 @@ public class LogEventParserTests
 
 	#endregion
 
+	#region Structured JSON Ignore Pattern Tests
+
+	// Minimal valid structured JSON log events for testing
+	private const string JsonErrorLine = "{\"level\":\"Error\",\"message\":\"LogCompile: Error: C1234: undefined symbol 'foo'\",\"properties\":{\"id\":100},\"format\":\"{message}\"}";
+	private const string JsonWarningLine = "{\"level\":\"Warning\",\"message\":\"LogCompile: Warning: variable 'x' is unused\",\"properties\":{\"id\":100},\"format\":\"{message}\"}";
+	private const string JsonInfoLine = "{\"level\":\"Information\",\"message\":\"LogCompile: Display: Build succeeded\",\"properties\":{},\"format\":\"{message}\"}";
+
+	[TestMethod]
+	public void TestIgnorePattern_SuppressesStructuredJsonError()
+	{
+		LoggerCapture capture = new();
+
+		using (LogEventParser parser = new(capture))
+		{
+			parser.IgnorePatterns.Add(new IgnorePatternEntry(new Regex(".*C1234.*", RegexOptions.Compiled), "test.txt"));
+			parser.WriteLine(JsonErrorLine);
+		}
+
+		// The error should be suppressed - logged at Information instead of Error
+		Assert.AreEqual(1, capture.Events.Count);
+		Assert.AreEqual(LogLevel.Information, capture.Events[0].Level);
+	}
+
+	[TestMethod]
+	public void TestIgnorePattern_SuppressesStructuredJsonWarning()
+	{
+		LoggerCapture capture = new();
+
+		using (LogEventParser parser = new(capture))
+		{
+			parser.IgnorePatterns.Add(new IgnorePatternEntry(new Regex(".*variable.*unused.*", RegexOptions.Compiled), "test.txt"));
+			parser.WriteLine(JsonWarningLine);
+		}
+
+		// The warning should be suppressed
+		Assert.AreEqual(1, capture.Events.Count);
+		Assert.AreEqual(LogLevel.Information, capture.Events[0].Level);
+	}
+
+	[TestMethod]
+	public void TestIgnorePattern_DoesNotSuppressNonMatchingStructuredJsonError()
+	{
+		LoggerCapture capture = new();
+
+		using (LogEventParser parser = new(capture))
+		{
+			parser.IgnorePatterns.Add(new IgnorePatternEntry(new Regex(".*unrelated_pattern.*", RegexOptions.Compiled), "test.txt"));
+			parser.WriteLine(JsonErrorLine);
+		}
+
+		// The error should NOT be suppressed since the pattern doesn't match
+		Assert.AreEqual(1, capture.Events.Count);
+		Assert.AreEqual(LogLevel.Error, capture.Events[0].Level);
+	}
+
+	[TestMethod]
+	public void TestIgnorePattern_DoesNotCheckStructuredJsonInformationEvents()
+	{
+		LoggerCapture capture = new();
+
+		using (LogEventParser parser = new(capture))
+		{
+			// Pattern matches the info message but should NOT be checked (optimization: only Warning+ checked)
+			parser.IgnorePatterns.Add(new IgnorePatternEntry(new Regex(".*Build succeeded.*", RegexOptions.Compiled), "test.txt"));
+			parser.WriteLine(JsonInfoLine);
+		}
+
+		// Information events pass through without ignore pattern checking
+		Assert.AreEqual(1, capture.Events.Count);
+		Assert.AreEqual(LogLevel.Information, capture.Events[0].Level);
+	}
+
+	[TestMethod]
+	public void TestIgnorePattern_StructuredJsonViaWriteData()
+	{
+		LoggerCapture capture = new();
+
+		using (LogEventParser parser = new(capture))
+		{
+			parser.IgnorePatterns.Add(new IgnorePatternEntry(new Regex(".*C1234.*", RegexOptions.Compiled), "test.txt"));
+			// WriteData goes through AddLine path (byte-based), not WriteLine
+			byte[] data = System.Text.Encoding.UTF8.GetBytes(JsonErrorLine + "\n");
+			parser.WriteData(data);
+		}
+
+		// The error should be suppressed through the WriteData/AddLine path too
+		Assert.AreEqual(1, capture.Events.Count);
+		Assert.AreEqual(LogLevel.Information, capture.Events[0].Level);
+	}
+
+	[TestMethod]
+	public void TestIgnorePattern_StructuredJsonPreservesIgnoreMetadata()
+	{
+		LoggerCapture capture = new();
+
+		using (LogEventParser parser = new(capture))
+		{
+			parser.IgnorePatterns.Add(new IgnorePatternEntry(new Regex(".*C1234.*", RegexOptions.Compiled), "/workspace/Engine/Build/Horde/IgnorePatterns.txt"));
+			parser.WriteLine(JsonErrorLine);
+		}
+
+		Assert.AreEqual(1, capture.Events.Count);
+		Assert.AreEqual(LogLevel.Information, capture.Events[0].Level);
+
+		// Verify structured properties tracking the ignore source are present
+		Assert.IsNotNull(capture.Events[0].Properties);
+		KeyValuePair<string, object?>? fileProperty = capture.Events[0].Properties?.FirstOrDefault(p => p.Key == "IgnorePatternFile");
+		KeyValuePair<string, object?>? patternProperty = capture.Events[0].Properties?.FirstOrDefault(p => p.Key == "IgnorePattern");
+		Assert.IsNotNull(fileProperty);
+		Assert.AreEqual("/workspace/Engine/Build/Horde/IgnorePatterns.txt", fileProperty.Value.Value?.ToString());
+		Assert.IsNotNull(patternProperty);
+		Assert.AreEqual(".*C1234.*", patternProperty.Value.Value?.ToString());
+	}
+
+	#endregion
+
 	private const string LogLineProperty = "LogLine";
 	private const string LineToLog = "hello";
 
diff --git a/Engine/Source/Programs/Shared/EpicGames.Core/LogEventParser.cs b/Engine/Source/Programs/Shared/EpicGames.Core/LogEventParser.cs
--- a/Engine/Source/Programs/Shared/EpicGames.Core/LogEventParser.cs
+++ b/Engine/Source/Programs/Shared/EpicGames.Core/LogEventParser.cs
@@ -465,6 +465,27 @@ namespace EpicGames.Core
 			_logger.Log(LogLevel.Information, KnownLogEvents.None, state, null, (s, _) => line);
 		}
 
+		/// <summary>
+		/// Checks whether a structured JSON log event should be suppressed by an ignore pattern.
+		/// Only events at Warning level or above are checked. If a matching pattern is found,
+		/// the event is re-logged at Information level (preventing build health issue creation).
+		/// </summary>
+		/// <returns>True if the event was suppressed, false if it should be logged normally</returns>
+		bool TryIgnoreStructuredEvent(JsonLogEvent jsonEvent)
+		{
+			if (jsonEvent.Level >= LogLevel.Warning && IgnorePatterns.Count > 0)
+			{
+				string message = jsonEvent.GetRenderedMessage().ToString();
+				IgnorePatternEntry? ignoredBy = FindMatchingIgnorePattern(IgnorePatterns, message, _logger);
+				if (ignoredBy != null)
+				{
+					WriteIgnoredLine(message, ignoredBy.Value);
+					return true;
+				}
+			}
+			return false;
+		}
+
 		/// <summary>
 		/// Discover all IgnorePatterns.txt files in a workspace by scanning Engine, project,
 		/// and Restricted/Platforms subdirectories.
@@ -634,7 +655,10 @@ namespace EpicGames.Core
 					if (JsonLogEvent.TryParse(data, out jsonEvent))
 					{
 						ProcessData(true);
-						_logger.Log(jsonEvent.Level, jsonEvent.EventId, jsonEvent, null, JsonLogEvent.Format);
+						if (!TryIgnoreStructuredEvent(jsonEvent))
+						{
+							_logger.Log(jsonEvent.Level, jsonEvent.EventId, jsonEvent, null, JsonLogEvent.Format);
+						}
 						return;
 					}
 
@@ -817,7 +841,10 @@ namespace EpicGames.Core
 				if (JsonLogEvent.TryParse(data, out jsonEvent))
 				{
 					ProcessData(true);
-					_logger.LogJsonLogEvent(jsonEvent);
+					if (!TryIgnoreStructuredEvent(jsonEvent))
+					{
+						_logger.LogJsonLogEvent(jsonEvent);
+					}
 					return;
 				}
 			}
