Updated the logeventparser pop and suspend regex to handle JSON input.

#Synced-CL 49593614

diff --git a/Engine/Source/Programs/Shared/EpicGames.Core/LogEventParser.cs b/Engine/Source/Programs/Shared/EpicGames.Core/LogEventParser.cs
--- a/Engine/Source/Programs/Shared/EpicGames.Core/LogEventParser.cs
+++ b/Engine/Source/Programs/Shared/EpicGames.Core/LogEventParser.cs
@@ -275,7 +275,7 @@ namespace EpicGames.Core
 		/// 
 		/// <c>MatcherName</c> is either the literal name of a matcher, or the string <c>*</c> for all matchers.
 		/// </remarks>
-		[GeneratedRegex("^<-- logeventparser (layer=(?<LayerName>\\S+)? )?((?<Action>suspend|resume) (?<MatcherName>\\S+)|(?<Action>pop)) -->$", RegexOptions.IgnoreCase | RegexOptions.ExplicitCapture)]
+		[GeneratedRegex("^(?:<|\\\\u003C)-- logeventparser (layer=(?<LayerName>\\S+)? )?((?<Action>suspend|resume) (?<MatcherName>\\S+)|(?<Action>pop)) --(?:>|\\\\u003E)$", RegexOptions.IgnoreCase | RegexOptions.ExplicitCapture)]
 		public static partial Regex MatcherToggleRegex();
 
 		/// <summary>
