[Horde][Analytics] Introducing standard const for invalid commit id.

#Synced-CL 50915067

diff --git a/Engine/Source/Programs/Shared/EpicGames.Horde/Telemetry/HordeContextTelemetryRecord.cs b/Engine/Source/Programs/Shared/EpicGames.Horde/Telemetry/HordeContextTelemetryRecord.cs
--- a/Engine/Source/Programs/Shared/EpicGames.Horde/Telemetry/HordeContextTelemetryRecord.cs
+++ b/Engine/Source/Programs/Shared/EpicGames.Horde/Telemetry/HordeContextTelemetryRecord.cs
@@ -16,6 +16,11 @@ namespace EpicGames.Horde.Telemetry
 	/// </summary>
 	public record HordeContextTelemetryRecord : CorrelatableTelemetryRecord
 	{
+		/// <summary>
+		/// The value for an invalid commit id.
+		/// </summary>
+		public const int InvalidCommitId = -1;
+
 		#region -- Public Properties --
 
 		/// <summary>
@@ -167,13 +172,13 @@ namespace EpicGames.Horde.Telemetry
 			}
 			else
 			{
-				CommitIdOrdered = -1;
+				CommitIdOrdered = InvalidCommitId;
 			}
 		}
 
 		#endregion -- Public Api --
 
-		#region -- Protectied Api --
+		#region -- Protected Api --
 
 		/// <inheritdoc/>
 		protected override void WriteProperties(Hashtable hashtable)
