Cloud Storage - When outputting zen cli output to stdout include directives to disable log event matching in Horde so that error messages with "error" and "warning" are not categorized as issues if retried

#Synced-CL 50250157

diff --git a/Engine/Source/Programs/AutomationTool/BuildGraph/Tasks/CreateCloudArtifactTask.cs b/Engine/Source/Programs/AutomationTool/BuildGraph/Tasks/CreateCloudArtifactTask.cs
--- a/Engine/Source/Programs/AutomationTool/BuildGraph/Tasks/CreateCloudArtifactTask.cs
+++ b/Engine/Source/Programs/AutomationTool/BuildGraph/Tasks/CreateCloudArtifactTask.cs
@@ -727,7 +727,18 @@ namespace AutomationTool.Tasks
 				// pass the access token via environment variable to avoid showing it on the cli
 				string cmdline = $"{traceOption}builds upload --url {host} --namespace {namespaceId} --bucket {bucketId} --build-id {buildId} --build-part-name \"{partName}\" --build-part-id {partId} --metadata-path \"{metadataJsonPath}\" --manifest-path \"{manifestFile}\" --access-token-env UE-CloudDataCacheAccessToken --create-build --plain-progress --allow-multipart={allowMultipart} {http2Options} {boostWorkersOption} --clean={forceClean} {verifyOption} {verboseOption} \"{baseDir}\"";
 				Log.Logger.LogInformation("Using metadata '{Metadata}'", metadataString);
-				CommandUtils.RunAndLog(CommandUtils.CmdEnv, zenExe.FullName,cmdline , Options: CommandUtils.ERunOptions.Default, EnvVars: new Dictionary<string, string> {{"UE-CloudDataCacheAccessToken", accessToken}});
+
+				try
+				{
+					// suppress any log processing done by horde unless we explicitly opt in to it
+					Log.Logger.LogInformation("<-- logeventparser suspend * -->");
+					CommandUtils.RunAndLog(CommandUtils.CmdEnv, zenExe.FullName, cmdline, Options: CommandUtils.ERunOptions.Default, EnvVars: new Dictionary<string, string> { { "UE-CloudDataCacheAccessToken", accessToken } });
+				}
+				finally
+				{
+					// resume log processing
+					Log.Logger.LogInformation("<-- logeventparser pop -->");
+				}
 			}
 			finally
 			{
diff --git a/Engine/Source/Programs/Shared/EpicGames.ProjectStore/CloudStorage/CloudStorage.cs b/Engine/Source/Programs/Shared/EpicGames.ProjectStore/CloudStorage/CloudStorage.cs
--- a/Engine/Source/Programs/Shared/EpicGames.ProjectStore/CloudStorage/CloudStorage.cs
+++ b/Engine/Source/Programs/Shared/EpicGames.ProjectStore/CloudStorage/CloudStorage.cs
@@ -561,47 +561,70 @@ namespace EpicGames.ProjectStore
 				zenLogger = optionsOverride?.EchoLogger;
 			}
 
-			const int WarnAtAttempts = 2;
-			// increasingly long time to wait before retrying
-			int[] waitDurationSeconds = [60, 300, 900, 3600];
-			int maxRetryAttempts = waitDurationSeconds.Length;
-
-			for (int retryAttempt = 0; retryAttempt < maxRetryAttempts; retryAttempt++)
+			if (_options.Value.EchoZenCliToStdout && _options.Value.WriteLogEventStatements)
 			{
-				int exitCode = await RunProcessAsync(zenExe.FullName, cmdline, env: env, echoStdout: _options.Value.EchoZenCliToStdout, echoLogger: zenLogger, progress: progress, cancellationToken: cancellationToken);
+				// if we are echoing zen output to stdout we instruct the horde log event parsing to not run its usually categorization of errors and instead allow us to control errors
+				Console.WriteLine("<-- logeventparser suspend * -->");
+			}
 
-				int waitDuration = waitDurationSeconds[retryAttempt];
+			try
+			{
+				const int WarnAtAttempts = 2;
+				// increasingly long time to wait before retrying
+				int[] waitDurationSeconds = [60, 300, 900, 3600];
+				int maxRetryAttempts = waitDurationSeconds.Length;
 
-				ZenExitCodes zenExitCode = AsZenExitCodeValue(exitCode);
-				// retry on zen exit codes that indicate transient issues
-				if (zenExitCode is ZenExitCodes.HttpServiceUnavailable or ZenExitCodes.HttpBadGateway or ZenExitCodes.HttpGatewayTimeout)
+				for (int retryAttempt = 0; retryAttempt < maxRetryAttempts; retryAttempt++)
 				{
-					if (retryAttempt + 1 == maxRetryAttempts)
-					{
-						// give up
-						_logger.LogError("Zen returned http status code {ExitCode}({ExitCodeMessage}) after {RetryAttempts} retries. Giving up.", exitCode, AsZenExitCodeValue(exitCode), retryAttempt);
+					int exitCode = await RunProcessAsync(zenExe.FullName, cmdline, env: env, echoStdout: _options.Value.EchoZenCliToStdout, echoLogger: zenLogger, progress: progress, cancellationToken: cancellationToken);
 
-						// we return the exit code even though it's a failure to allow the caller to handle this issue that isn't fixed by retrying
-					}
-					else
+					int waitDuration = waitDurationSeconds[retryAttempt];
+
+					ZenExitCodes zenExitCode = AsZenExitCodeValue(exitCode);
+					// retry on zen exit codes that indicate transient issues
+					if (zenExitCode is ZenExitCodes.HttpServiceUnavailable or ZenExitCodes.HttpBadGateway or ZenExitCodes.HttpGatewayTimeout)
 					{
-						if (retryAttempt + 1 == WarnAtAttempts)
+						if (retryAttempt + 1 == maxRetryAttempts)
 						{
-							_logger.LogWarning("Zen returned http status code {ExitCode}({ExitCodeMessage}) after {RetryAttempts} retries. This is a indication of an infrastructure issue.", exitCode, AsZenExitCodeValue(exitCode), retryAttempt);
+							// give up
+							_logger.LogError("Zen returned http status code {ExitCode}({ExitCodeMessage}) after {RetryAttempts} retries. Giving up.",
+								exitCode, AsZenExitCodeValue(exitCode), retryAttempt);
+
+							// we return the exit code even though it's a failure to allow the caller to handle this issue that isn't fixed by retrying
 						}
 						else
 						{
-							_logger.LogInformation("Http Status code {ExitCode}({ExitCodeMessage}) returned from zen. Waiting {WaitDuration}s before trying.", exitCode, AsZenExitCodeValue(exitCode), waitDuration);
+							if (retryAttempt + 1 == WarnAtAttempts)
+							{
+								_logger.LogWarning(
+									"Zen returned http status code {ExitCode}({ExitCodeMessage}) after {RetryAttempts} retries. This is a indication of an infrastructure issue.",
+									exitCode, AsZenExitCodeValue(exitCode), retryAttempt);
+							}
+							else
+							{
+								_logger.LogInformation(
+									"Http Status code {ExitCode}({ExitCodeMessage}) returned from zen. Waiting {WaitDuration}s before trying.",
+									exitCode, AsZenExitCodeValue(exitCode), waitDuration);
+							}
+
+							await Task.Delay(TimeSpan.FromSeconds(waitDuration), cancellationToken);
+							continue; // retry
 						}
-						await Task.Delay(TimeSpan.FromSeconds(waitDuration), cancellationToken);
-						continue; // retry
 					}
+
+					return exitCode;
 				}
 
-				return exitCode;
+				throw new CloudStorageException($"Failed to run zen cli, gave up after {maxRetryAttempts} attempts");
+			}
+			finally
+			{
+				if (_options.Value.EchoZenCliToStdout && _options.Value.WriteLogEventStatements)
+				{
+					// resume normal log parsing
+					Console.WriteLine("<-- logeventparser pop -->");
+				}
 			}
-
-			throw new CloudStorageException($"Failed to run zen cli, gave up after {maxRetryAttempts} attempts");
 		}
 
 		private async Task<int> RunProcessAsync(string cmd, string arguments, string? workingDirectory = null, Dictionary<string, string>? env = null, IProgress<(float, string)>? progress = null, bool echoStdout = false, ILogger? echoLogger = null, CancellationToken cancellationToken = default)
diff --git a/Engine/Source/Programs/Shared/EpicGames.ProjectStore/CloudStorage/ICloudStorage.cs b/Engine/Source/Programs/Shared/EpicGames.ProjectStore/CloudStorage/ICloudStorage.cs
--- a/Engine/Source/Programs/Shared/EpicGames.ProjectStore/CloudStorage/ICloudStorage.cs
+++ b/Engine/Source/Programs/Shared/EpicGames.ProjectStore/CloudStorage/ICloudStorage.cs
@@ -481,6 +481,11 @@ namespace EpicGames.ProjectStore
 		/// </summary>
 		public bool EchoZenCliToLog { get; set; } = false;
 
+		/// <summary>
+		/// Used to control if log event directives are written into sources that it will consume
+		/// </summary>
+		public bool WriteLogEventStatements { get; set; } = true;
+
 		/// <summary>
 		/// Implements verification of the settings object for IValidatableObject
 		/// </summary>
