Cloud Storage - Using the old suspend / resume macros for log parsing control as we do not (currently) need per matcher control.

#Synced-CL 50250157

diff --git a/Engine/Source/Programs/AutomationTool/BuildGraph/Tasks/CreateCloudArtifactTask.cs b/Engine/Source/Programs/AutomationTool/BuildGraph/Tasks/CreateCloudArtifactTask.cs
--- a/Engine/Source/Programs/AutomationTool/BuildGraph/Tasks/CreateCloudArtifactTask.cs
+++ b/Engine/Source/Programs/AutomationTool/BuildGraph/Tasks/CreateCloudArtifactTask.cs
@@ -731,13 +731,13 @@ namespace AutomationTool.Tasks
 				try
 				{
 					// suppress any log processing done by horde unless we explicitly opt in to it
-					Log.Logger.LogInformation("<-- logeventparser suspend * -->");
+					Log.Logger.LogInformation("<-- Suspend Log Parsing -->");
 					CommandUtils.RunAndLog(CommandUtils.CmdEnv, zenExe.FullName, cmdline, Options: CommandUtils.ERunOptions.Default, EnvVars: new Dictionary<string, string> { { "UE-CloudDataCacheAccessToken", accessToken } });
 				}
 				finally
 				{
 					// resume log processing
-					Log.Logger.LogInformation("<-- logeventparser pop -->");
+					Log.Logger.LogInformation("<-- Resume Log Parsing -->");
 				}
 			}
 			finally
diff --git a/Engine/Source/Programs/Shared/EpicGames.ProjectStore/CloudStorage/CloudStorage.cs b/Engine/Source/Programs/Shared/EpicGames.ProjectStore/CloudStorage/CloudStorage.cs
--- a/Engine/Source/Programs/Shared/EpicGames.ProjectStore/CloudStorage/CloudStorage.cs
+++ b/Engine/Source/Programs/Shared/EpicGames.ProjectStore/CloudStorage/CloudStorage.cs
@@ -564,7 +564,7 @@ namespace EpicGames.ProjectStore
 			if (_options.Value.EchoZenCliToStdout && _options.Value.WriteLogEventStatements)
 			{
 				// if we are echoing zen output to stdout we instruct the horde log event parsing to not run its usually categorization of errors and instead allow us to control errors
-				Console.WriteLine("<-- logeventparser suspend * -->");
+				Console.WriteLine("<-- Suspend Log Parsing -->");
 			}
 
 			try
@@ -622,7 +622,7 @@ namespace EpicGames.ProjectStore
 				if (_options.Value.EchoZenCliToStdout && _options.Value.WriteLogEventStatements)
 				{
 					// resume normal log parsing
-					Console.WriteLine("<-- logeventparser pop -->");
+					Console.WriteLine("<-- Resume Log Parsing -->");
 				}
 			}
 		}
