EpicGames.Core: Add lock diagnostics to directory deletion on Windows

Includes which process ID and executable name that may be holding a lock when directory deletion fails.
Enumerates all files in a directory to see which one holds a lock.

diff --git a/Engine/Source/Programs/Shared/EpicGames.Core.Tests/FileUtilsTests.cs b/Engine/Source/Programs/Shared/EpicGames.Core.Tests/FileUtilsTests.cs
--- a/Engine/Source/Programs/Shared/EpicGames.Core.Tests/FileUtilsTests.cs
+++ b/Engine/Source/Programs/Shared/EpicGames.Core.Tests/FileUtilsTests.cs
@@ -85,6 +85,31 @@ namespace EpicGames.Core.Tests
 			return new DirectoryReference(tempDir);
 		}
 
+		[TestMethod]
+		public void ForceDeleteDirectory_LockedFile_IncludesLockInfo()
+		{
+			if (!OperatingSystem.IsWindows())
+			{
+				Assert.Inconclusive("This test only runs on Windows");
+				return;
+			}
+
+			string subDir = Path.Join(_tempDir.FullName, "LockedDir");
+			Directory.CreateDirectory(subDir);
+			string tempFile = Path.Join(subDir, "locked_file.txt");
+			File.WriteAllText(tempFile, "test content");
+
+			using FileStream lockedStream = new(tempFile, FileMode.Open, FileAccess.Read, FileShare.None);
+			WrappedFileOrDirectoryException ex = Assert.ThrowsExactly<WrappedFileOrDirectoryException>(() => FileUtils.ForceDeleteDirectory(subDir));
+
+			// Verify the exception includes process ID, image path, and command line of the locking process
+			string msg = ex.Message;
+			StringAssert.Contains(msg, "locked_file.txt", StringComparison.Ordinal);
+			StringAssert.Contains(msg, "Processes with open handles", StringComparison.Ordinal);
+			StringAssert.Contains(msg, $"{Environment.ProcessId}: ", StringComparison.Ordinal);
+			StringAssert.Contains(msg, "Command line:", StringComparison.Ordinal);
+		}
+
 		[TestMethod]
 		public void GetFileLockInfoWin32_ReturnsCommandLineAndParentInfo()
 		{
diff --git a/Engine/Source/Programs/Shared/EpicGames.Core/FileUtils.cs b/Engine/Source/Programs/Shared/EpicGames.Core/FileUtils.cs
--- a/Engine/Source/Programs/Shared/EpicGames.Core/FileUtils.cs
+++ b/Engine/Source/Programs/Shared/EpicGames.Core/FileUtils.cs
@@ -645,6 +645,7 @@ namespace EpicGames.Core
 		const int ERROR_FILE_NOT_FOUND = 2;
 		const int ERROR_PATH_NOT_FOUND = 3;
 		const int ERROR_ACCESS_DENIED = 5;
+		const int ERROR_SHARING_VIOLATION = 32;
 #pragma warning restore IDE1006 // Naming Styles
 
 		[SupportedOSPlatform("windows")]
@@ -709,8 +710,130 @@ namespace EpicGames.Core
 				int errorCode = Marshal.GetLastWin32Error();
 				if (errorCode != ERROR_FILE_NOT_FOUND && errorCode != ERROR_PATH_NOT_FOUND)
 				{
-					throw new WrappedFileOrDirectoryException(new Win32Exception(errorCode), "Unable to delete " + dirName);
+					Win32Exception ex = new(errorCode);
+					StringBuilder message = new($"Unable to delete {dirName} - {ex.Message}");
+
+					if (errorCode == ERROR_ACCESS_DENIED)
+					{
+						string? lockInfo = GetDirectoryLockInfo(dirName);
+						if (lockInfo != null)
+						{
+							message.Append('\n');
+							message.Append(lockInfo);
+						}
+					}
+
+					throw new WrappedFileOrDirectoryException(ex, message.ToString());
+				}
+			}
+		}
+
+		/// <summary>
+		/// Enumerates remaining files in a directory after RemoveDirectory fails and gathers lock info.
+		/// Returns null if no useful diagnostics could be gathered.
+		/// </summary>
+		[SupportedOSPlatform("windows")]
+		private static string? GetDirectoryLockInfo(string dirName)
+		{
+			const string RawPathPrefix = "\\\\?\\";
+			const int MaxFiles = 10;
+			const int MaxDirs = 5;
+
+			try
+			{
+				string searchPath = dirName;
+				if (!searchPath.StartsWith(RawPathPrefix, StringComparison.Ordinal))
+				{
+					searchPath = RawPathPrefix + searchPath;
+				}
+				if (!searchPath.EndsWith('\\'))
+				{
+					searchPath += "\\";
+				}
+
+				WIN32_FIND_DATA findData = new();
+				IntPtr hFind = FindFirstFileW(searchPath + "*", ref findData);
+				if (hFind == INVALID_HANDLE_VALUE)
+				{
+					return null;
+				}
+
+				List<string> remainingFiles = [];
+				List<string> remainingDirs = [];
+				try
+				{
+					for (; ; )
+					{
+						if (findData.cFileName != "." && findData.cFileName != "..")
+						{
+							string fullName = searchPath + findData.cFileName;
+							if ((findData.dwFileAttributes & FILE_ATTRIBUTE_DIRECTORY) != 0)
+							{
+								if (remainingDirs.Count < MaxDirs)
+								{
+									remainingDirs.Add(fullName);
+								}
+							}
+							else
+							{
+								if (remainingFiles.Count < MaxFiles)
+								{
+									remainingFiles.Add(fullName);
+								}
+							}
+						}
+
+						if (!FindNextFileW(hFind, ref findData))
+						{
+							break;
+						}
+					}
+				}
+				finally
+				{
+					FindClose(hFind);
+				}
+
+				if (remainingFiles.Count == 0 && remainingDirs.Count == 0)
+				{
+					return "Directory is empty but could not be removed; a process may be holding a handle to the directory itself.";
 				}
+
+				StringBuilder message = new();
+				if (remainingDirs.Count > 0)
+				{
+					message.Append($"Directory still contains {remainingDirs.Count} subdirector{(remainingDirs.Count == 1 ? "y" : "ies")}:");
+					foreach (string dir in remainingDirs)
+					{
+						message.Append($"\n  {dir}");
+					}
+				}
+
+				foreach (string file in remainingFiles)
+				{
+					// Strip \\?\ prefix for Restart Manager compatibility
+					string fileForLockInfo = file;
+					if (fileForLockInfo.StartsWith(RawPathPrefix, StringComparison.Ordinal))
+					{
+						fileForLockInfo = fileForLockInfo.Substring(RawPathPrefix.Length);
+					}
+
+					string? lockInfo = GetFileLockInfo(fileForLockInfo);
+					if (lockInfo != null)
+					{
+						if (message.Length > 0)
+						{
+							message.Append('\n');
+						}
+						message.Append($"File: {file}\n{lockInfo}");
+					}
+				}
+
+				return message.Length > 0 ? message.ToString() : null;
+			}
+			catch
+			{
+				return null;
 			}
 		}
 
@@ -756,7 +879,7 @@ namespace EpicGames.Core
 
 			// Get a useful error message about why the delete failed
 			StringBuilder message = new($"Unable to delete {fileName} - {ex.Message}");
-			if (ex.NativeErrorCode == ERROR_ACCESS_DENIED)
+			if (ex.NativeErrorCode == ERROR_ACCESS_DENIED || ex.NativeErrorCode == ERROR_SHARING_VIOLATION)
 			{
 				string? lockInfo = GetFileLockInfo(fileName);
 				if (lockInfo != null)
