EpicGames.Core: Improve error handling for ForceDeleteFileWin32 and add tests

This fixes a race condition during parallel directory cleanup where GetFileAttributesW returning a different error than DeleteFileW caused a spurious WrappedFileOrDirectoryException. Lock info is now always gathered regardless of error code.

diff --git a/Engine/Source/Programs/Shared/EpicGames.Core.Tests/FileUtilsTests.cs b/Engine/Source/Programs/Shared/EpicGames.Core.Tests/FileUtilsTests.cs
--- a/Engine/Source/Programs/Shared/EpicGames.Core.Tests/FileUtilsTests.cs
+++ b/Engine/Source/Programs/Shared/EpicGames.Core.Tests/FileUtilsTests.cs
@@ -3,6 +3,7 @@
 using System;
 using System.Collections.Generic;
 using System.IO;
+using System.Runtime.Versioning;
 using Microsoft.VisualStudio.TestTools.UnitTesting;
 
 namespace EpicGames.Core.Tests
@@ -78,22 +79,41 @@ namespace EpicGames.Core.Tests
 			return subDir;
 		}
 
-		private static DirectoryReference CreateTempDir()
+		internal static DirectoryReference CreateTempDir()
 		{
 			string tempDir = Path.Join(Path.GetTempPath(), "epicgames-core-tests-" + Guid.NewGuid().ToString()[..8]);
 			Directory.CreateDirectory(tempDir);
 			return new DirectoryReference(tempDir);
 		}
+	}
 
-		[TestMethod]
-		public void ForceDeleteDirectory_LockedFile_IncludesLockInfo()
+	[TestClass]
+	[SupportedOSPlatform("windows")]
+	public class FileUtilsWindowsTest
+	{
+		private readonly DirectoryReference _tempDir = FileUtilsTest.CreateTempDir();
+
+		[TestInitialize]
+		public void CheckWindows()
 		{
 			if (!OperatingSystem.IsWindows())
 			{
 				Assert.Inconclusive("This test only runs on Windows");
-				return;
 			}
+		}
+
+		[TestCleanup]
+		public void RemoveTempDir()
+		{
+			if (Directory.Exists(_tempDir.FullName))
+			{
+				Directory.Delete(_tempDir.FullName, true);
+			}
+		}
 
+		[TestMethod]
+		public void ForceDeleteDirectory_LockedFile_IncludesLockInfo()
+		{
 			string subDir = Path.Join(_tempDir.FullName, "LockedDir");
 			Directory.CreateDirectory(subDir);
 			string tempFile = Path.Join(subDir, "locked_file.txt");
@@ -111,14 +131,49 @@ namespace EpicGames.Core.Tests
 		}
 
 		[TestMethod]
-		public void GetFileLockInfoWin32_ReturnsCommandLineAndParentInfo()
+		public void ForceDeleteFileWin32_NormalFile_Deleted()
 		{
-			if (!OperatingSystem.IsWindows())
-			{
-				Assert.Inconclusive("This test only runs on Windows");
-				return;
-			}
+			string filePath = Path.Join(_tempDir.FullName, "normal.txt");
+			File.WriteAllText(filePath, "content");
+			FileUtils.ForceDeleteFileWin32(filePath);
+			Assert.IsFalse(File.Exists(filePath));
+		}
+
+		[TestMethod]
+		public void ForceDeleteFileWin32_NonExistentFile_NoThrow()
+		{
+			string filePath = Path.Join(_tempDir.FullName, "does_not_exist.txt");
+			FileUtils.ForceDeleteFileWin32(filePath);
+		}
 
+		[TestMethod]
+		public void ForceDeleteFileWin32_ReadOnlyFile_Deleted()
+		{
+			string filePath = Path.Join(_tempDir.FullName, "readonly.txt");
+			File.WriteAllText(filePath, "content");
+			File.SetAttributes(filePath, FileAttributes.ReadOnly);
+			FileUtils.ForceDeleteFileWin32(filePath);
+			Assert.IsFalse(File.Exists(filePath));
+		}
+
+		[TestMethod]
+		public void ForceDeleteFileWin32_LockedFile_ThrowsWithLockInfo()
+		{
+			string filePath = Path.Join(_tempDir.FullName, "locked.txt");
+			File.WriteAllText(filePath, "content");
+
+			using FileStream lockedStream = new(filePath, FileMode.Open, FileAccess.Read, FileShare.None);
+			WrappedFileOrDirectoryException ex = Assert.ThrowsExactly<WrappedFileOrDirectoryException>(() => FileUtils.ForceDeleteFileWin32(filePath));
+
+			string msg = ex.Message;
+			StringAssert.Contains(msg, "locked.txt", StringComparison.Ordinal);
+			StringAssert.Contains(msg, "Processes with open handles", StringComparison.Ordinal);
+			StringAssert.Contains(msg, $"{Environment.ProcessId}: ", StringComparison.Ordinal);
+		}
+
+		[TestMethod]
+		public void GetFileLockInfoWin32_ReturnsCommandLineAndParentInfo()
+		{
 			string tempFile = Path.Combine(_tempDir.FullName, "locked_file.txt");
 			File.WriteAllText(tempFile, "test content");
 
diff --git a/Engine/Source/Programs/Shared/EpicGames.Core/FileUtils.cs b/Engine/Source/Programs/Shared/EpicGames.Core/FileUtils.cs
--- a/Engine/Source/Programs/Shared/EpicGames.Core/FileUtils.cs
+++ b/Engine/Source/Programs/Shared/EpicGames.Core/FileUtils.cs
@@ -847,39 +847,27 @@ namespace EpicGames.Core
 		[SupportedOSPlatform("windows")]
 		public static void ForceDeleteFileWin32(string fileName)
 		{
-			// Try to delete the file normally
-			if (DeleteFileW(fileName))
+			if (TryDeleteFileWin32(fileName, out int errorCode))
 			{
 				return;
 			}
 
-			// Capture the exception for failing to delete the file
-			Win32Exception ex = new();
-
-			// Check the file exists and is not readonly
+			// If the file is read-only, clear the flag and retry
 			int attributes = GetFileAttributesW(fileName);
-			if (attributes == -1)
+			if (attributes != -1 && (attributes & (int)FileAttributes.ReadOnly) != 0)
 			{
-				int errorCode = Marshal.GetLastWin32Error();
-				if (errorCode == ERROR_PATH_NOT_FOUND || errorCode == ERROR_FILE_NOT_FOUND)
+				SetFileAttributesW(fileName, FILE_ATTRIBUTE_NORMAL);
+				if (TryDeleteFileWin32(fileName, out errorCode))
 				{
 					return;
 				}
 			}
-			else
-			{
-				if ((attributes & (int)FileAttributes.ReadOnly) != 0)
-				{
-					if (SetFileAttributesW(fileName, (int)FileAttributes.Normal) && DeleteFileW(fileName))
-					{
-						return;
-					}
-				}
-			}
 
-			// Get a useful error message about why the delete failed
-			StringBuilder message = new($"Unable to delete {fileName} - {ex.Message}");
-			if (ex.NativeErrorCode == ERROR_ACCESS_DENIED || ex.NativeErrorCode == ERROR_SHARING_VIOLATION)
+			// Build error message with lock info
+			Win32Exception ex = new(errorCode);
+			StringBuilder message = new($"Unable to delete {fileName} (error {errorCode}) - {ex.Message}");
+
+			if (errorCode is ERROR_ACCESS_DENIED or ERROR_SHARING_VIOLATION)
 			{
 				string? lockInfo = GetFileLockInfo(fileName);
 				if (lockInfo != null)
@@ -888,9 +876,26 @@ namespace EpicGames.Core
 					message.Append(lockInfo);
 				}
 			}
+
 			throw new WrappedFileOrDirectoryException(ex, message.ToString());
 		}
 
+		/// <summary>
+		/// Attempt to delete a file, returning true on success or if the file does not exist.
+		/// </summary>
+		[SupportedOSPlatform("windows")]
+		private static bool TryDeleteFileWin32(string fileName, out int errorCode)
+		{
+			if (DeleteFileW(fileName))
+			{
+				errorCode = 0;
+				return true;
+			}
+
+			errorCode = Marshal.GetLastWin32Error();
+			return errorCode is ERROR_FILE_NOT_FOUND or ERROR_PATH_NOT_FOUND;
+		}
+
 		#endregion
 
 		#region Win32 Restart Manager API
