Fixing an issue where types wrapped in WITH_EDITORONLY_DATA would not have that define reflected in the generated code.

#Synced-CL 51048589

diff --git a/Engine/Source/Programs/Shared/EpicGames.UHT/Parsers/UhtEnumParser.cs b/Engine/Source/Programs/Shared/EpicGames.UHT/Parsers/UhtEnumParser.cs
--- a/Engine/Source/Programs/Shared/EpicGames.UHT/Parsers/UhtEnumParser.cs
+++ b/Engine/Source/Programs/Shared/EpicGames.UHT/Parsers/UhtEnumParser.cs
@@ -70,10 +70,6 @@ namespace EpicGames.UHT.Parsers
 					enumObject.Outer?.AddChild(enumObject);
 
 					UhtCompilerDirective compilerDirective = topScope.HeaderParser.GetCurrentCompositeCompilerDirective();
-					if (compilerDirective.HasAnyFlags(UhtCompilerDirective.WithEditorOnlyData))
-					{
-						enumObject.DefineScope |= UhtDefineScope.EditorOnlyData;
-					}
 					enumObject.DefineScope |= compilerDirective.GetDefaultDefineScopes();
 
 					// Read base for enum class
diff --git a/Engine/Source/Programs/Shared/EpicGames.UHT/Parsers/UhtFunctionParser.cs b/Engine/Source/Programs/Shared/EpicGames.UHT/Parsers/UhtFunctionParser.cs
--- a/Engine/Source/Programs/Shared/EpicGames.UHT/Parsers/UhtFunctionParser.cs
+++ b/Engine/Source/Programs/Shared/EpicGames.UHT/Parsers/UhtFunctionParser.cs
@@ -149,7 +149,7 @@ namespace EpicGames.UHT.Parsers
 						function.FunctionFlags |= EFunctionFlags.EditorOnly;
 						function.DefineScope |= UhtDefineScope.Editor;
 					}
-					function.DefineScope |= compilerDirective.GetDefaultDefineScopes();
+					function.DefineScope |= (compilerDirective & ~(UhtCompilerDirective.WithEditorOnlyData | UhtCompilerDirective.WithEditor)).GetDefaultDefineScopes();
 
 					// If this is a UDELEGATE, parse the specifiers
 					StringView delegateMacro = new();
@@ -370,7 +370,7 @@ namespace EpicGames.UHT.Parsers
 						function.FunctionFlags |= EFunctionFlags.EditorOnly;
 						function.DefineScope |= UhtDefineScope.Editor;
 					}
-					function.DefineScope |= compilerDirective.GetDefaultDefineScopes();
+					function.DefineScope |= (compilerDirective & ~(UhtCompilerDirective.WithEditorOnlyData | UhtCompilerDirective.WithEditor)).GetDefaultDefineScopes();
 
 					specifierParser.ParseDeferred();
 					FinalizeFunctionSpecifiers(function);
diff --git a/Engine/Source/Programs/Shared/EpicGames.UHT/Parsers/UhtHeaderFileParser.cs b/Engine/Source/Programs/Shared/EpicGames.UHT/Parsers/UhtHeaderFileParser.cs
--- a/Engine/Source/Programs/Shared/EpicGames.UHT/Parsers/UhtHeaderFileParser.cs
+++ b/Engine/Source/Programs/Shared/EpicGames.UHT/Parsers/UhtHeaderFileParser.cs
@@ -212,6 +212,10 @@ namespace EpicGames.UHT.Parsers
 			{
 				defineScope |= UhtDefineScope.VerseVM;
 			}
+			if (compilerDirectives.HasAnyFlags(UhtCompilerDirective.WithEditorOnlyData))
+			{
+				defineScope |= UhtDefineScope.EditorOnlyData;
+			}
 			return defineScope;
 		}
 
diff --git a/Engine/Source/Programs/Shared/EpicGames.UHT/Parsers/UhtPropertyParser.cs b/Engine/Source/Programs/Shared/EpicGames.UHT/Parsers/UhtPropertyParser.cs
--- a/Engine/Source/Programs/Shared/EpicGames.UHT/Parsers/UhtPropertyParser.cs
+++ b/Engine/Source/Programs/Shared/EpicGames.UHT/Parsers/UhtPropertyParser.cs
@@ -688,7 +688,7 @@ namespace EpicGames.UHT.Parsers
 						tokenReader.LogError("UProperties should not be wrapped by WITH_EDITOR, use WITH_EDITORONLY_DATA instead.");
 					}
 				}
-				propertySettings.DefineScope |= compilerDirective.GetDefaultDefineScopes();
+				propertySettings.DefineScope |= (compilerDirective & ~(UhtCompilerDirective.WithEditorOnlyData | UhtCompilerDirective.WithEditor)).GetDefaultDefineScopes();
 			}
 
 			// Parse type information including UPARAM that might appear in template arguments
