Horde: Add time-based persistent flushing to ServerLogger

ServerLogger previously only flushed buffered log data to blob storage when the buffer exceeded 1 MB or at end-of-stream, leaving slow-streaming jobs without visible log updates for several minutes when the tail mechanism was inactive.

diff --git a/Engine/Source/Programs/Shared/EpicGames.Horde.Tests/Logs/ServerLoggerTests.cs b/Engine/Source/Programs/Shared/EpicGames.Horde.Tests/Logs/ServerLoggerTests.cs
new file mode 100644
--- /dev/null
+++ b/Engine/Source/Programs/Shared/EpicGames.Horde.Tests/Logs/ServerLoggerTests.cs
@@ -0,0 +1,544 @@
+// Copyright Epic Games, Inc. All Rights Reserved.
+
+using System;
+using System.Collections.Generic;
+using System.Linq;
+using System.Threading;
+using System.Threading.Channels;
+using System.Threading.Tasks;
+using EpicGames.Core;
+using EpicGames.Horde.Artifacts;
+using EpicGames.Horde.Compute;
+using EpicGames.Horde.Logs;
+using EpicGames.Horde.Projects;
+using EpicGames.Horde.Secrets;
+using EpicGames.Horde.Storage;
+using EpicGames.Horde.Storage.Bundles;
+using EpicGames.Horde.Tools;
+using Google.Protobuf.WellKnownTypes;
+using Grpc.Core;
+using Grpc.Net.Client;
+using Horde.Common.Rpc;
+using Microsoft.Extensions.Logging;
+using Microsoft.Extensions.Logging.Abstractions;
+using Microsoft.VisualStudio.TestTools.UnitTesting;
+
+namespace EpicGames.Horde.Tests
+{
+	[TestClass]
+	public class ServerLoggerTests
+	{
+		#region Test infrastructure
+
+		class GrpcHelpers
+		{
+			class StreamReader<T> : IAsyncStreamReader<T> where T : class
+			{
+				readonly ChannelReader<T> _reader;
+				T? _current;
+
+				public T Current => _current ?? throw new InvalidOperationException();
+
+				public StreamReader(ChannelReader<T> reader)
+					=> _reader = reader;
+
+				public async Task<bool> MoveNext(CancellationToken cancellationToken)
+				{
+					while (await _reader.WaitToReadAsync(cancellationToken))
+					{
+						if (_reader.TryRead(out _current))
+						{
+							return true;
+						}
+					}
+					return false;
+				}
+			}
+
+			class StreamWriter<T> : IClientStreamWriter<T>
+			{
+				readonly ChannelWriter<T> _writer;
+
+				public StreamWriter(ChannelWriter<T> writer)
+					=> _writer = writer;
+
+				public WriteOptions? WriteOptions { get => throw new NotImplementedException(); set => throw new NotImplementedException(); }
+
+				public Task CompleteAsync()
+				{
+					_writer.Complete();
+					return Task.CompletedTask;
+				}
+
+				public async Task WriteAsync(T message) => await _writer.WriteAsync(message);
+			}
+
+			public static AsyncDuplexStreamingCall<TRequest, TResponse> CreateDuplexCall<TRequest, TResponse>(Func<ChannelReader<TRequest>, ChannelWriter<TResponse>, Task> func)
+				where TRequest : class
+				where TResponse : class
+			{
+				Channel<TRequest> requests = Channel.CreateUnbounded<TRequest>();
+				Channel<TResponse> responses = Channel.CreateUnbounded<TResponse>();
+
+				BackgroundTask runner = BackgroundTask.StartNew(async ctx =>
+				{
+					await func(requests.Reader, responses.Writer);
+					responses.Writer.Complete();
+				});
+
+				return new AsyncDuplexStreamingCall<TRequest, TResponse>(new StreamWriter<TRequest>(requests.Writer), new StreamReader<TResponse>(responses.Reader), Task.FromResult<Metadata>(null!), () => Status.DefaultSuccess, () => null!, () => runner.DisposeAsync().AsTask().Wait());
+			}
+		}
+
+		class FakeLogRpcClient : LogRpc.LogRpcClient
+		{
+			public Dictionary<LogId, BlobLocator> Logs { get; } = [];
+			public Channel<RpcUpdateLogRequest> UpdateLogChannel { get; } = Channel.CreateUnbounded<RpcUpdateLogRequest>();
+			public Channel<RpcCreateLogEventsRequest> CreateLogEventsChannel { get; } = Channel.CreateUnbounded<RpcCreateLogEventsRequest>();
+			public Channel<RpcUpdateLogTailRequest> UpdateLogTailChannel { get; } = Channel.CreateUnbounded<RpcUpdateLogTailRequest>();
+
+			public override AsyncUnaryCall<RpcUpdateLogResponse> UpdateLogAsync(RpcUpdateLogRequest request, CallOptions options)
+			{
+				Logs[LogId.Parse(request.LogId)] = new BlobLocator(request.TargetLocator);
+				UpdateLogChannel.Writer.TryWrite(request);
+				return new AsyncUnaryCall<RpcUpdateLogResponse>(Task.FromResult(new RpcUpdateLogResponse()), Task.FromResult(new Metadata()), () => Status.DefaultSuccess, () => [], () => { });
+			}
+
+			public override AsyncUnaryCall<Empty> CreateLogEventsAsync(RpcCreateLogEventsRequest request, CallOptions options)
+			{
+				CreateLogEventsChannel.Writer.TryWrite(request);
+				return new AsyncUnaryCall<Empty>(Task.FromResult(new Empty()), Task.FromResult(new Metadata()), () => Status.DefaultSuccess, () => [], () => { });
+			}
+
+			public override AsyncDuplexStreamingCall<RpcUpdateLogTailRequest, RpcUpdateLogTailResponse> UpdateLogTail(CallOptions options)
+			{
+				return GrpcHelpers.CreateDuplexCall<RpcUpdateLogTailRequest, RpcUpdateLogTailResponse>(async (reader, writer) =>
+				{
+					while (await reader.WaitToReadAsync())
+					{
+						if (reader.TryRead(out RpcUpdateLogTailRequest? request))
+						{
+							UpdateLogTailChannel.Writer.TryWrite(request);
+						}
+					}
+				});
+			}
+		}
+
+		class FakeHordeClient : IHordeClient, IAsyncDisposable
+		{
+			public FakeLogRpcClient LogRpc { get; } = new FakeLogRpcClient();
+			public Dictionary<string, BundleStorageNamespace> StorageNamespaces { get; } = [];
+
+			public Uri ServerUrl => throw new NotImplementedException();
+
+			public IArtifactCollection Artifacts => throw new NotImplementedException();
+			public IComputeClient Compute => throw new NotImplementedException();
+			public IProjectCollection Projects => throw new NotImplementedException();
+			public ISecretCollection Secrets => throw new NotImplementedException();
+			public IToolCollection Tools => throw new NotImplementedException();
+
+			public event Action? OnAccessTokenStateChanged
+			{
+				add { }
+				remove { }
+			}
+
+			public Task<bool> LoginAsync(bool allowLogin, CancellationToken cancellationToken)
+				=> throw new NotImplementedException();
+
+			public HordeHttpClient CreateHttpClient()
+				=> throw new NotImplementedException();
+
+			public IStorageNamespace GetStorageNamespace(string relativePath, string? accessToken = null)
+			{
+				BundleStorageNamespace? storageNamespace;
+				if (!StorageNamespaces.TryGetValue(relativePath, out storageNamespace))
+				{
+					storageNamespace = BundleStorageNamespace.CreateInMemory(NullLogger.Instance);
+					StorageNamespaces.Add(relativePath, storageNamespace);
+				}
+				return storageNamespace;
+			}
+
+			public ValueTask DisposeAsync()
+			{
+				StorageNamespaces.Clear();
+				return default;
+			}
+
+			public Task<string?> GetAccessTokenAsync(bool interactive, CancellationToken cancellationToken = default)
+				=> throw new NotImplementedException();
+
+			public Task<GrpcChannel> CreateGrpcChannelAsync(CancellationToken cancellationToken = default)
+				=> throw new NotImplementedException();
+
+			public Task<TClient> CreateGrpcClientAsync<TClient>(CancellationToken cancellationToken = default) where TClient : ClientBase<TClient>
+			{
+				if (typeof(TClient) == typeof(LogRpc.LogRpcClient))
+				{
+					return Task.FromResult<TClient>((TClient)(object)LogRpc);
+				}
+				throw new NotImplementedException();
+			}
+
+			public bool HasValidAccessToken()
+				=> throw new NotImplementedException();
+
+			public IServerLogger CreateServerLogger(LogId logId, LogLevel minimumLevel = LogLevel.Information)
+				=> throw new NotImplementedException();
+		}
+
+		static List<T> DrainChannel<T>(Channel<T> channel)
+		{
+			List<T> items = [];
+			while (channel.Reader.TryRead(out T? item))
+			{
+				items.Add(item);
+			}
+			return items;
+		}
+
+		static List<RpcCreateLogEventRequest> DrainAllEvents(FakeLogRpcClient rpc)
+		{
+			return [.. DrainChannel(rpc.CreateLogEventsChannel).SelectMany(r => r.Events)];
+		}
+
+		#endregion
+
+		[TestMethod]
+		public async Task FinalFlushSetsCompleteFlagAsync()
+		{
+			await using FakeHordeClient hordeClient = new();
+			LogId logId = default;
+
+			await using (ServerLogger logger = new(hordeClient, logId, LogLevel.Information, NullLogger.Instance))
+			{
+				for (int idx = 0; idx < 10; idx++)
+				{
+					logger.LogInformation("Line {Number}", idx);
+				}
+			}
+
+			List<RpcUpdateLogRequest> requests = DrainChannel(hordeClient.LogRpc.UpdateLogChannel);
+			Assert.IsTrue(requests.Count >= 1);
+			Assert.IsTrue(requests[^1].Complete);
+			Assert.AreEqual(10, requests[^1].LineCount);
+		}
+
+		[TestMethod]
+		public async Task SmallDataPreservedOnDisposeAsync()
+		{
+			await using FakeHordeClient hordeClient = new();
+			LogId logId = default;
+
+			await using (ServerLogger logger = new(hordeClient, logId, LogLevel.Information, NullLogger.Instance))
+			{
+				logger.LogInformation("Single line");
+			}
+
+			List<RpcUpdateLogRequest> requests = DrainChannel(hordeClient.LogRpc.UpdateLogChannel);
+			Assert.AreEqual(1, requests.Count);
+			Assert.IsTrue(requests[0].Complete);
+
+			// Verify round-trip through storage
+			IStorageNamespace storageNamespace = hordeClient.GetStorageNamespace(logId);
+			LogNode file = await storageNamespace.CreateBlobRef(hordeClient.LogRpc.Logs[logId]).ReadBlobAsync<LogNode>();
+			LogIndexNode index = await file.IndexRef.ReadBlobAsync();
+
+			List<Utf8String> lines = [];
+			foreach (LogChunkRef block in index.PlainTextChunkRefs)
+			{
+				LogChunkNode text = await block.Target.ReadBlobAsync();
+				lines.AddRange(text.Lines);
+			}
+
+			Assert.AreEqual(1, lines.Count);
+			Assert.AreEqual(new Utf8String("Single line"), lines[0]);
+		}
+
+		[TestMethod]
+		public async Task SizeBasedFlushTriggersIntermediateUpdateAsync()
+		{
+			await using FakeHordeClient hordeClient = new();
+			LogId logId = default;
+
+			await using (ServerLogger logger = new(hordeClient, logId, LogLevel.Information, NullLogger.Instance))
+			{
+				for (int idx = 0; idx < 10000; idx++)
+				{
+					logger.LogInformation("This is a sufficiently long line of log output for testing purposes number {Number}", idx);
+				}
+			}
+
+			List<RpcUpdateLogRequest> requests = DrainChannel(hordeClient.LogRpc.UpdateLogChannel);
+			Assert.IsTrue(requests.Count >= 2);
+			Assert.IsTrue(requests.Any(r => !r.Complete));
+			Assert.IsTrue(requests[^1].Complete);
+
+			for (int i = 1; i < requests.Count; i++)
+			{
+				Assert.IsTrue(requests[i].LineCount >= requests[i - 1].LineCount);
+			}
+		}
+
+		[TestMethod]
+		public async Task WarningEventsReportedAsync()
+		{
+			await using FakeHordeClient hordeClient = new();
+			LogId logId = default;
+
+			await using (ServerLogger logger = new(hordeClient, logId, LogLevel.Information, NullLogger.Instance))
+			{
+				for (int idx = 0; idx < 5; idx++)
+				{
+					logger.LogWarning("Warning {Number}", idx);
+				}
+			}
+
+			List<RpcCreateLogEventRequest> allEvents = DrainAllEvents(hordeClient.LogRpc);
+			Assert.AreEqual(5, allEvents.Count);
+			Assert.IsTrue(allEvents.All(e => e.Severity == (int)LogEventSeverity.Warning));
+		}
+
+		[TestMethod]
+		public async Task ErrorAndCriticalEventsReportedAsync()
+		{
+			await using FakeHordeClient hordeClient = new();
+			LogId logId = default;
+
+			await using (ServerLogger logger = new(hordeClient, logId, LogLevel.Information, NullLogger.Instance))
+			{
+				for (int idx = 0; idx < 3; idx++)
+				{
+					logger.LogError("Error {Number}", idx);
+				}
+				for (int idx = 0; idx < 2; idx++)
+				{
+					logger.LogCritical("Critical {Number}", idx);
+				}
+			}
+
+			List<RpcCreateLogEventRequest> allEvents = DrainAllEvents(hordeClient.LogRpc);
+			Assert.AreEqual(5, allEvents.Count);
+			Assert.IsTrue(allEvents.All(e => e.Severity == (int)LogEventSeverity.Error));
+		}
+
+		[TestMethod]
+		public Task WarningCapEnforcedAt200Async() => VerifyEventCapAt200Async(LogLevel.Warning);
+
+		[TestMethod]
+		public Task ErrorCapEnforcedAt200Async() => VerifyEventCapAt200Async(LogLevel.Error);
+
+		static async Task VerifyEventCapAt200Async(LogLevel level)
+		{
+			await using FakeHordeClient hordeClient = new();
+			LogId logId = default;
+
+			await using (ServerLogger logger = new(hordeClient, logId, LogLevel.Information, NullLogger.Instance))
+			{
+				for (int idx = 0; idx < 250; idx++)
+				{
+					logger.Log(level, "Message {Number}", idx);
+				}
+			}
+
+			List<RpcCreateLogEventRequest> allEvents = DrainAllEvents(hordeClient.LogRpc);
+			Assert.AreEqual(200, allEvents.Count);
+		}
+
+		[TestMethod]
+		public async Task MixedWarningsAndErrorsHaveIndependentCapsAsync()
+		{
+			await using FakeHordeClient hordeClient = new();
+			LogId logId = default;
+
+			await using (ServerLogger logger = new(hordeClient, logId, LogLevel.Information, NullLogger.Instance))
+			{
+				for (int idx = 0; idx < 210; idx++)
+				{
+					logger.LogWarning("Warning {Number}", idx);
+				}
+				for (int idx = 0; idx < 210; idx++)
+				{
+					logger.LogError("Error {Number}", idx);
+				}
+			}
+
+			List<RpcCreateLogEventRequest> allEvents = DrainAllEvents(hordeClient.LogRpc);
+			int warningCount = allEvents.Count(e => e.Severity == (int)LogEventSeverity.Warning);
+			int errorCount = allEvents.Count(e => e.Severity == (int)LogEventSeverity.Error);
+
+			Assert.AreEqual(200, warningCount);
+			Assert.AreEqual(200, errorCount);
+		}
+
+		[TestMethod]
+		public async Task InformationLinesDoNotCreateEventsAsync()
+		{
+			await using FakeHordeClient hordeClient = new();
+			LogId logId = default;
+
+			await using (ServerLogger logger = new(hordeClient, logId, LogLevel.Information, NullLogger.Instance))
+			{
+				for (int idx = 0; idx < 100; idx++)
+				{
+					logger.LogInformation("Info {Number}", idx);
+				}
+			}
+
+			List<RpcCreateLogEventRequest> allEvents = DrainAllEvents(hordeClient.LogRpc);
+			Assert.AreEqual(0, allEvents.Count);
+		}
+
+		[TestMethod]
+		public async Task StopAsyncIdempotentAsync()
+		{
+			await using FakeHordeClient hordeClient = new();
+			LogId logId = default;
+
+			ServerLogger logger = new(hordeClient, logId, LogLevel.Information, NullLogger.Instance);
+
+			for (int idx = 0; idx < 10; idx++)
+			{
+				logger.LogInformation("Line {Number}", idx);
+			}
+
+			// Call StopAsync twice — should not throw
+			await logger.StopAsync();
+			await logger.StopAsync();
+
+			// DisposeAsync after StopAsync
+			await logger.DisposeAsync();
+
+			List<RpcUpdateLogRequest> requests = DrainChannel(hordeClient.LogRpc.UpdateLogChannel);
+			int completeCount = requests.Count(r => r.Complete);
+			Assert.AreEqual(1, completeCount);
+		}
+
+		[TestMethod]
+		public async Task StorageRoundTripAsync()
+		{
+			await using FakeHordeClient hordeClient = new();
+			LogId logId = default;
+
+			const int Count = 1000;
+			await using (ServerLogger logger = new(hordeClient, logId, LogLevel.Information, NullLogger.Instance))
+			{
+				for (int idx = 0; idx < Count; idx++)
+				{
+					logger.LogInformation("Testing {Number}", idx);
+				}
+			}
+
+			// Read the log back from storage
+			IStorageNamespace storageNamespace = hordeClient.GetStorageNamespace(logId);
+			LogNode file = await storageNamespace.CreateBlobRef(hordeClient.LogRpc.Logs[logId]).ReadBlobAsync<LogNode>();
+
+			// Check index text
+			List<Utf8String> extractedIndexText = [];
+			LogIndexNode index = await file.IndexRef.ReadBlobAsync();
+			foreach (LogChunkRef block in index.PlainTextChunkRefs)
+			{
+				LogChunkNode text = await block.Target.ReadBlobAsync();
+				extractedIndexText.AddRange(text.Lines);
+			}
+
+			Assert.AreEqual(Count, extractedIndexText.Count);
+			for (int idx = 0; idx < Count; idx++)
+			{
+				Assert.AreEqual(new Utf8String($"Testing {idx}"), extractedIndexText[idx]);
+			}
+
+			// Check body text
+			List<string> extractedBodyText = [];
+			foreach (LogChunkRef blockRef in file.TextChunkRefs)
+			{
+				LogChunkNode blockText = await blockRef.Target.ReadBlobAsync();
+				foreach (Utf8String line in blockText.Lines)
+				{
+					LogEvent logEvent = LogEvent.Read(line.Span);
+					extractedBodyText.Add(logEvent.ToString());
+				}
+			}
+
+			Assert.AreEqual(Count, extractedBodyText.Count);
+			for (int idx = 0; idx < Count; idx++)
+			{
+				Assert.AreEqual($"Testing {idx}", extractedBodyText[idx]);
+			}
+		}
+
+		[TestMethod]
+		[Timeout(15000)]
+		public async Task TimeBasedFlushTriggersBeforeDisposeAsync()
+		{
+			await using FakeHordeClient hordeClient = new();
+			LogId logId = default;
+
+			await using ServerLogger logger = new(hordeClient, logId, LogLevel.Information, NullLogger.Instance, flushInterval: TimeSpan.FromSeconds(1));
+
+			for (int idx = 0; idx < 10; idx++)
+			{
+				logger.LogInformation("Line {Number}", idx);
+			}
+
+			// Wait for an intermediate flush (before dispose)
+			using CancellationTokenSource cts = new(TimeSpan.FromSeconds(10));
+			RpcUpdateLogRequest request = await hordeClient.LogRpc.UpdateLogChannel.Reader.ReadAsync(cts.Token);
+			Assert.IsFalse(request.Complete);
+		}
+
+		[TestMethod]
+		[Timeout(15000)]
+		public async Task EmptyBufferDoesNotTriggerTimeFlushAsync()
+		{
+			await using FakeHordeClient hordeClient = new();
+			LogId logId = default;
+
+			await using (ServerLogger logger = new(hordeClient, logId, LogLevel.Information, NullLogger.Instance, flushInterval: TimeSpan.FromSeconds(1)))
+			{
+				// Write nothing, wait longer than flush interval
+				await Task.Delay(TimeSpan.FromSeconds(3));
+			}
+
+			List<RpcUpdateLogRequest> requests = DrainChannel(hordeClient.LogRpc.UpdateLogChannel);
+
+			// Only the final flush should be present (Complete == true), and only if data was written.
+			// With no data, there should still be a final Complete call from dispose.
+			Assert.IsTrue(requests.All(r => r.Complete));
+		}
+
+		[TestMethod]
+		public async Task MinimumLevelFilteringAsync()
+		{
+			await using FakeHordeClient hordeClient = new();
+			LogId logId = default;
+
+			await using ServerLogger logger = new(hordeClient, logId, LogLevel.Warning, NullLogger.Instance);
+
+			Assert.IsFalse(logger.IsEnabled(LogLevel.Trace));
+			Assert.IsFalse(logger.IsEnabled(LogLevel.Debug));
+			Assert.IsFalse(logger.IsEnabled(LogLevel.Information));
+			Assert.IsTrue(logger.IsEnabled(LogLevel.Warning));
+			Assert.IsTrue(logger.IsEnabled(LogLevel.Error));
+			Assert.IsTrue(logger.IsEnabled(LogLevel.Critical));
+		}
+
+		[TestMethod]
+		public async Task LogAfterStopDoesNotThrowAsync()
+		{
+			await using FakeHordeClient hordeClient = new();
+			LogId logId = default;
+
+			ServerLogger logger = new(hordeClient, logId, LogLevel.Information, NullLogger.Instance);
+			logger.LogInformation("Before stop");
+			await logger.StopAsync();
+
+			// Should not throw even though the channel is closed
+			logger.LogInformation("After stop");
+
+			await logger.DisposeAsync();
+		}
+	}
+}
diff --git a/Engine/Source/Programs/Shared/EpicGames.Horde/Logs/ServerLogger.cs b/Engine/Source/Programs/Shared/EpicGames.Horde/Logs/ServerLogger.cs
--- a/Engine/Source/Programs/Shared/EpicGames.Horde/Logs/ServerLogger.cs
+++ b/Engine/Source/Programs/Shared/EpicGames.Horde/Logs/ServerLogger.cs
@@ -2,6 +2,7 @@
 
 using System;
 using System.Collections.Generic;
+using System.Diagnostics;
 using System.Text;
 using System.Threading;
 using System.Threading.Channels;
@@ -31,8 +32,15 @@ namespace EpicGames.Horde.Logs
 		readonly IBlobWriter _writer;
 
 		int _bufferLength;
+		readonly Stopwatch _timeSinceFlush = Stopwatch.StartNew();
+		readonly TimeSpan _flushInterval;
 		string? _closedByStackTrace;
 
+		const int MaxErrors = 200;
+		const int MaxWarnings = 200;
+		int _numErrors;
+		int _numWarnings;
+
 		// Tailing task
 		readonly Task _tailTask;
 		AsyncEvent _tailTaskStop;
@@ -48,12 +56,14 @@ namespace EpicGames.Horde.Logs
 		/// <param name="logId">The log id to write to</param>
 		/// <param name="minimumLevel">Minimum level for output</param>
 		/// <param name="internalLogger">Logger for systemic messages</param>
-		public ServerLogger(IHordeClient hordeClient, LogId logId, LogLevel minimumLevel, ILogger internalLogger)
+		/// <param name="flushInterval">Interval for time-based persistent flush (default 30s)</param>
+		public ServerLogger(IHordeClient hordeClient, LogId logId, LogLevel minimumLevel, ILogger internalLogger, TimeSpan? flushInterval = null)
 		{
 			_hordeClient = hordeClient;
 			_logId = logId;
 			_minimumLevel = minimumLevel;
 			_internalLogger = internalLogger;
+			_flushInterval = flushInterval ?? TimeSpan.FromSeconds(30);
 			_builder = new LogBuilder(LogFormat.Json, internalLogger);
 			_store = _hordeClient.GetStorageNamespace(logId);
 			_writer = _store.CreateBlobWriter($"{logId}");
@@ -132,63 +142,21 @@ namespace EpicGames.Horde.Logs
 		/// <returns>Async task</returns>
 		async Task RunDataWriterAsync()
 		{
-			// Current position and line number in the log file
 			long packetOffset = 0;
 			int packetLineIndex = 0;
-
-			// Index of the next line to write to the log
 			int nextLineIndex = 0;
 
-			// Total number of errors and warnings
-			const int MaxErrors = 200;
-			int numErrors = 0;
-			const int MaxWarnings = 200;
-			int numWarnings = 0;
-
-			// Buffer for events read in a single iteration
 			ServerLogPacketBuilder writer = new();
 			List<RpcCreateLogEventRequest> events = [];
 
-			// Whether we've written the flush command
 			for (; ; )
 			{
 				events.Clear();
 
-				// Get the next data
-				Task waitTask = Task.Delay(TimeSpan.FromSeconds(2.0));
-				while (writer.PacketLength < writer.MaxPacketLength)
-				{
-					JsonLogEvent jsonLogEvent;
-					if (_dataChannel.Reader.TryRead(out jsonLogEvent))
-					{
-						int lineCount = writer.SanitizeAndWriteEvent(jsonLogEvent);
-						if (jsonLogEvent.LineIndex == 0)
-						{
-							if (jsonLogEvent.Level == LogLevel.Warning && ++numWarnings <= MaxWarnings)
-							{
-								AddEvent(jsonLogEvent.Data.Span, nextLineIndex, Math.Max(lineCount, jsonLogEvent.LineCount), LogEventSeverity.Warning, events);
-							}
-							else if ((jsonLogEvent.Level == LogLevel.Error || jsonLogEvent.Level == LogLevel.Critical) && ++numErrors <= MaxErrors)
-							{
-								AddEvent(jsonLogEvent.Data.Span, nextLineIndex, Math.Max(lineCount, jsonLogEvent.LineCount), LogEventSeverity.Error, events);
-							}
-						}
-						nextLineIndex += lineCount;
-					}
-					{
-						Task<bool> readTask = _dataChannel.Reader.WaitToReadAsync().AsTask();
-						if (await Task.WhenAny(readTask, waitTask) == waitTask)
-						{
-							break;
-						}
-						if (!await readTask)
-						{
-							break;
-						}
-					}
-				}
+				// Read events from the channel until the packet is full or a 2-second timeout elapses
+				nextLineIndex = await ReadEventsIntoBatchAsync(writer, events, nextLineIndex);
 
-				// Upload it to the server
+				// Upload packet to the server
 				if (writer.PacketLength > 0)
 				{
 					(ReadOnlyMemory<byte> packet, int packetLineCount) = writer.CreatePacket();
@@ -204,7 +172,7 @@ namespace EpicGames.Horde.Logs
 					}
 				}
 
-				// Write all the events
+				// Upload log events to the server
 				if (events.Count > 0)
 				{
 					try
@@ -217,20 +185,129 @@ namespace EpicGames.Horde.Logs
 					}
 				}
 
-				// Wait for more data to be available
-				if (writer.PacketLength <= 0 && !await _dataChannel.Reader.WaitToReadAsync())
+				// When no packet was produced, wait for more data (flushing periodically if needed)
+				if (writer.PacketLength <= 0)
 				{
+					if (!await WaitForDataWithPeriodicFlushAsync(packetOffset))
+					{
+						break;
+					}
+				}
+			}
+		}
+
+		/// <summary>
+		/// Reads events from the data channel into a packet builder, with a 2-second timeout.
+		/// Returns the updated next line index.
+		/// </summary>
+		async Task<int> ReadEventsIntoBatchAsync(ServerLogPacketBuilder writer, List<RpcCreateLogEventRequest> events, int nextLineIndex)
+		{
+			using CancellationTokenSource timeoutCts = new(TimeSpan.FromSeconds(2.0));
+			while (writer.PacketLength < writer.MaxPacketLength)
+			{
+				if (_dataChannel.Reader.TryRead(out JsonLogEvent jsonLogEvent))
+				{
+					int lineCount = writer.SanitizeAndWriteEvent(jsonLogEvent);
+					ClassifyAndAddEvent(jsonLogEvent, nextLineIndex, lineCount, events);
+					nextLineIndex += lineCount;
+				}
+				else
+				{
+					// Wait asynchronously only when TryRead fails; avoids unnecessary allocations during bursts
+					bool hasMore;
 					try
 					{
-						await WriteOutputAsync(ReadOnlyMemory<byte>.Empty, true, CancellationToken.None);
+						hasMore = await _dataChannel.Reader.WaitToReadAsync(timeoutCts.Token);
 					}
-					catch (Exception ex)
+					catch (OperationCanceledException)
 					{
-						_internalLogger.LogWarning(ex, "Unable to flush data to server (log {LogId}, offset {Offset})", _logId, packetOffset);
+						break;
+					}
+					if (!hasMore)
+					{
+						break;
+					}
+				}
+			}
+			return nextLineIndex;
+		}
+
+		/// <summary>
+		/// Classifies a log event and adds it to the event list if it is a warning or error
+		/// (up to the configured maximum for each severity).
+		/// </summary>
+		void ClassifyAndAddEvent(JsonLogEvent jsonLogEvent, int lineIndex, int lineCount, List<RpcCreateLogEventRequest> events)
+		{
+			if (jsonLogEvent.LineIndex != 0)
+			{
+				return;
+			}
+
+			int effectiveLineCount = Math.Max(lineCount, jsonLogEvent.LineCount);
+			if (jsonLogEvent.Level == LogLevel.Warning && ++_numWarnings <= MaxWarnings)
+			{
+				AddEvent(jsonLogEvent.Data.Span, lineIndex, effectiveLineCount, LogEventSeverity.Warning, events);
+			}
+			else if (jsonLogEvent.Level is LogLevel.Error or LogLevel.Critical && ++_numErrors <= MaxErrors)
+			{
+				AddEvent(jsonLogEvent.Data.Span, lineIndex, effectiveLineCount, LogEventSeverity.Error, events);
+			}
+		}
+
+		/// <summary>
+		/// Waits for new data on the channel, periodically flushing buffered output while idle.
+		/// Returns true if more data is available, or false if the channel has been closed.
+		/// </summary>
+		async Task<bool> WaitForDataWithPeriodicFlushAsync(long packetOffset)
+		{
+			Task<bool> waitForData = _dataChannel.Reader.WaitToReadAsync().AsTask();
+
+			// If there's unflushed data, use a bounded wait loop so we can flush periodically.
+			// This reuses the same waitForData task to avoid violating the ChannelReader contract
+			// (only one WaitToReadAsync may be outstanding at a time).
+			while (_bufferLength > 0)
+			{
+				double flushDelay = Math.Max(0, _flushInterval.TotalSeconds - _timeSinceFlush.Elapsed.TotalSeconds);
+				if (flushDelay > 0)
+				{
+					Task delayTask = Task.Delay(TimeSpan.FromSeconds(flushDelay));
+					await Task.WhenAny(waitForData, delayTask);
+
+					if (waitForData.IsCompleted)
+					{
+						break;
 					}
-					break;
 				}
+
+				// Flush timer expired with no new data; persist what we have
+				try
+				{
+					await WriteOutputAsync(ReadOnlyMemory<byte>.Empty, false, CancellationToken.None);
+				}
+				catch (Exception ex)
+				{
+					_internalLogger.LogWarning(ex, "Unable to flush data to server (log {LogId}, offset {Offset})", _logId, packetOffset);
+				}
+				// On success: _bufferLength is reset to 0 inside WriteOutputAsync, so the while
+				// condition becomes false and we fall through to await waitForData below.
+				// On failure: _bufferLength stays >0, _timeSinceFlush.Restart() fires in the
+				// finally block, so the next iteration waits a full 30 seconds.
+			}
+
+			if (!await waitForData)
+			{
+				// Channel closed by StopAsync; do a final flush
+				try
+				{
+					await WriteOutputAsync(ReadOnlyMemory<byte>.Empty, true, CancellationToken.None);
+				}
+				catch (Exception ex)
+				{
+					_internalLogger.LogWarning(ex, "Unable to flush data to server (log {LogId}, offset {Offset})", _logId, packetOffset);
+				}
+				return false;
 			}
+			return true;
 		}
 
 		void AddEvent(ReadOnlySpan<byte> span, int lineIndex, int lineCount, LogEventSeverity severity, List<RpcCreateLogEventRequest> events)
@@ -330,8 +407,7 @@ namespace EpicGames.Horde.Logs
 			return lines;
 		}
 
-		/// <inheritdoc/>
-		public async Task WriteEventsAsync(List<RpcCreateLogEventRequest> events, CancellationToken cancellationToken)
+		private async Task WriteEventsAsync(List<RpcCreateLogEventRequest> events, CancellationToken cancellationToken)
 		{
 			LogRpc.LogRpcClient logRpc = await _hordeClient.CreateGrpcClientAsync<LogRpc.LogRpcClient>(cancellationToken);
 
@@ -341,17 +417,23 @@ namespace EpicGames.Horde.Logs
 			await logRpc.CreateLogEventsAsync(request, cancellationToken: cancellationToken);
 		}
 
-		/// <inheritdoc/>
-		public async Task WriteOutputAsync(ReadOnlyMemory<byte> data, bool flush, CancellationToken cancellationToken)
+		private async Task WriteOutputAsync(ReadOnlyMemory<byte> data, bool flush, CancellationToken cancellationToken)
 		{
 			_builder.WriteData(data);
 			_bufferLength += data.Length;
 
-			if (flush || _bufferLength > FlushLength)
+			if (flush || _bufferLength > FlushLength || (_bufferLength > 0 && _timeSinceFlush.Elapsed > _flushInterval))
 			{
-				IHashedBlobRef<LogNode> target = await _builder.FlushAsync(_writer, flush, cancellationToken);
-				await UpdateLogAsync(target, _builder.LineCount, flush, cancellationToken);
-				_bufferLength = 0;
+				try
+				{
+					IHashedBlobRef<LogNode> target = await _builder.FlushAsync(_writer, flush, cancellationToken);
+					await UpdateLogAsync(target, _builder.LineCount, flush, cancellationToken);
+					_bufferLength = 0;
+				}
+				finally
+				{
+					_timeSinceFlush.Restart();
+				}
 			}
 
 			_newTailDataEvent.Set();
@@ -402,7 +484,7 @@ namespace EpicGames.Horde.Logs
 				// Close the request stream to indicate that we're finished
 				await call.RequestStream.CompleteAsync();
 
-				// Wait for a response or a new update to come in, then close the request stream
+				// Drain any remaining responses from the server
 				RpcUpdateLogTailResponse? response = null;
 				while (await moveNextAsync)
 				{
