Horde: Add OriginalLevel to structured logging for lines suppressed by IgnorePatterns

diff --git a/Engine/Source/Programs/Shared/EpicGames.Core.Tests/LogEventParserTests.cs b/Engine/Source/Programs/Shared/EpicGames.Core.Tests/LogEventParserTests.cs
--- a/Engine/Source/Programs/Shared/EpicGames.Core.Tests/LogEventParserTests.cs
+++ b/Engine/Source/Programs/Shared/EpicGames.Core.Tests/LogEventParserTests.cs
@@ -567,10 +567,13 @@ public class LogEventParserTests
 		Assert.IsNotNull(capture.Events[0].Properties);
 		KeyValuePair<string, object?>? fileProperty = capture.Events[0].Properties?.FirstOrDefault(p => p.Key == "IgnorePatternFile");
 		KeyValuePair<string, object?>? patternProperty = capture.Events[0].Properties?.FirstOrDefault(p => p.Key == "IgnorePattern");
+		KeyValuePair<string, object?>? levelProperty = capture.Events[0].Properties?.FirstOrDefault(p => p.Key == "OriginalLevel");
 		Assert.IsNotNull(fileProperty);
 		Assert.AreEqual("/workspace/Engine/Build/Horde/IgnorePatterns.txt", fileProperty.Value.Value?.ToString());
 		Assert.IsNotNull(patternProperty);
 		Assert.AreEqual(".*hello.*", patternProperty.Value.Value?.ToString());
+		Assert.IsNotNull(levelProperty);
+		Assert.AreEqual("Error", levelProperty.Value.Value?.ToString());
 	}
 
 	[TestMethod]
@@ -776,10 +779,13 @@ public class LogEventParserTests
 		Assert.IsNotNull(capture.Events[0].Properties);
 		KeyValuePair<string, object?>? fileProperty = capture.Events[0].Properties?.FirstOrDefault(p => p.Key == "IgnorePatternFile");
 		KeyValuePair<string, object?>? patternProperty = capture.Events[0].Properties?.FirstOrDefault(p => p.Key == "IgnorePattern");
+		KeyValuePair<string, object?>? levelProperty = capture.Events[0].Properties?.FirstOrDefault(p => p.Key == "OriginalLevel");
 		Assert.IsNotNull(fileProperty);
 		Assert.AreEqual("/workspace/Engine/Build/Horde/IgnorePatterns.txt", fileProperty.Value.Value?.ToString());
 		Assert.IsNotNull(patternProperty);
 		Assert.AreEqual(".*C1234.*", patternProperty.Value.Value?.ToString());
+		Assert.IsNotNull(levelProperty);
+		Assert.AreEqual("Error", levelProperty.Value.Value?.ToString());
 	}
 
 	#endregion
diff --git a/Engine/Source/Programs/Shared/EpicGames.Core/LogEventParser.cs b/Engine/Source/Programs/Shared/EpicGames.Core/LogEventParser.cs
--- a/Engine/Source/Programs/Shared/EpicGames.Core/LogEventParser.cs
+++ b/Engine/Source/Programs/Shared/EpicGames.Core/LogEventParser.cs
@@ -452,7 +452,10 @@ namespace EpicGames.Core
 		/// Writes a line that was suppressed by an ignore pattern, including structured properties
 		/// identifying which pattern file and regex caused the suppression.
 		/// </summary>
-		void WriteIgnoredLine(string line, IgnorePatternEntry ignoredBy)
+		/// <param name="line">The original log line text</param>
+		/// <param name="ignoredBy">The ignore pattern entry that matched</param>
+		/// <param name="originalLevel">The log level the line would have been emitted at before suppression</param>
+		void WriteIgnoredLine(string line, IgnorePatternEntry ignoredBy, LogLevel originalLevel)
 		{
 			KeyValuePair<string, object?>[] state =
 			[
@@ -460,6 +463,7 @@ namespace EpicGames.Core
 				new("IgnoredLine", line),
 				new("IgnorePatternFile", ignoredBy.SourceFile),
 				new("IgnorePattern", ignoredBy.Pattern.ToString()),
+				new("OriginalLevel", originalLevel.ToString()),
 			];
 
 			_logger.Log(LogLevel.Information, KnownLogEvents.None, state, null, (s, _) => line);
@@ -479,7 +483,7 @@ namespace EpicGames.Core
 				IgnorePatternEntry? ignoredBy = FindMatchingIgnorePattern(IgnorePatterns, message, _logger);
 				if (ignoredBy != null)
 				{
-					WriteIgnoredLine(message, ignoredBy.Value);
+					WriteIgnoredLine(message, ignoredBy.Value, jsonEvent.Level);
 					return true;
 				}
 			}
@@ -952,7 +956,7 @@ namespace EpicGames.Core
 								continue;
 							}
 
-							WriteIgnoredLine(line, ignoredBy.Value);
+							WriteIgnoredLine(line, ignoredBy.Value, match.Events[0].Level);
 							_buffer.MoveNext();
 							continue;
 						}
