Horde: Add structured IgnoredLine markup for suppressed log lines
Emit ignored log lines as LogValue with $type IgnoredLine so the
dashboard can render them with an "IGNORED:" prefix and the original
error/warning level styling. Eyeball callout shows pattern metadata.

diff --git a/Engine/Source/Programs/Horde/HordeDashboard/src/components/logs/LogRender.tsx b/Engine/Source/Programs/Horde/HordeDashboard/src/components/logs/LogRender.tsx
--- a/Engine/Source/Programs/Horde/HordeDashboard/src/components/logs/LogRender.tsx
+++ b/Engine/Source/Programs/Horde/HordeDashboard/src/components/logs/LogRender.tsx
@@ -21,6 +21,7 @@ enum TagType {
    ProjectId,
    StreamId,
    DeviceId,
+   IgnoredLine,
 }
 
 export type LogItem = {
@@ -316,6 +317,13 @@ const renderTags = (navigate: NavigateFunction, line: LogLine, lineNumber: numbe
                }
             }
             break;
+
+         case TagType.IgnoredLine:
+            return (
+               <span key={key}>
+                  <span style={{ fontFamily: "Horde Open Sans SemiBold, sans-serif", fontWeight: 600 }}>IGNORED:</span> {highlight}
+               </span>
+            );
       }
 
       return linkProps ? <a key={key} {...linkProps}>{highlight}</a> : <span key={key}/>
diff --git a/Engine/Source/Programs/Shared/EpicGames.Core/LogEventParser.cs b/Engine/Source/Programs/Shared/EpicGames.Core/LogEventParser.cs
--- a/Engine/Source/Programs/Shared/EpicGames.Core/LogEventParser.cs
+++ b/Engine/Source/Programs/Shared/EpicGames.Core/LogEventParser.cs
@@ -457,13 +457,17 @@ namespace EpicGames.Core
 		/// <param name="originalLevel">The log level the line would have been emitted at before suppression</param>
 		void WriteIgnoredLine(string line, IgnorePatternEntry ignoredBy, LogLevel originalLevel)
 		{
+			LogValue ignoredLineValue = new(LogValueType.IgnoredLine, line, new Dictionary<Utf8String, object?>
+			{
+				[new Utf8String("ignorePatternFile")] = ignoredBy.SourceFile,
+				[new Utf8String("ignorePattern")] = ignoredBy.Pattern.ToString(),
+				[new Utf8String("originalLevel")] = originalLevel.ToString(),
+			});
+
 			KeyValuePair<string, object?>[] state =
 			[
 				new(MessageTemplate.FormatPropertyName, "{IgnoredLine}"),
-				new("IgnoredLine", line),
-				new("IgnorePatternFile", ignoredBy.SourceFile),
-				new("IgnorePattern", ignoredBy.Pattern.ToString()),
-				new("OriginalLevel", originalLevel.ToString()),
+				new("IgnoredLine", ignoredLineValue),
 			];
 
 			_logger.Log(LogLevel.Information, KnownLogEvents.None, state, null, (s, _) => line);
diff --git a/Engine/Source/Programs/Shared/EpicGames.Core/LogValue.cs b/Engine/Source/Programs/Shared/EpicGames.Core/LogValue.cs
--- a/Engine/Source/Programs/Shared/EpicGames.Core/LogValue.cs
+++ b/Engine/Source/Programs/Shared/EpicGames.Core/LogValue.cs
@@ -27,6 +27,7 @@ namespace EpicGames.Core
 		public static readonly Utf8String ScreenshotTest = new("ScreenshotTest");
 		public static readonly Utf8String DepotPath = new("DepotPath");
 		public static readonly Utf8String Link = new("Link");
+		public static readonly Utf8String IgnoredLine = new("IgnoredLine");
 #pragma warning restore CS1591
 	}
 
