Horde: Update TestIgnorePattern_StructuredProperties for LogValue
Test now validates IgnoredLine as a LogValue with nested properties
instead of top-level state entries.

diff --git a/Engine/Source/Programs/Shared/EpicGames.Core.Tests/LogEventParserTests.cs b/Engine/Source/Programs/Shared/EpicGames.Core.Tests/LogEventParserTests.cs
--- a/Engine/Source/Programs/Shared/EpicGames.Core.Tests/LogEventParserTests.cs
+++ b/Engine/Source/Programs/Shared/EpicGames.Core.Tests/LogEventParserTests.cs
@@ -563,17 +563,16 @@ public class LogEventParserTests
 		Assert.AreEqual(1, capture.Events.Count);
 		Assert.AreEqual(LogLevel.Information, capture.Events[0].Level);
 
-		// Verify structured properties are present
+		// Verify IgnoredLine LogValue with structured properties
 		Assert.IsNotNull(capture.Events[0].Properties);
-		KeyValuePair<string, object?>? fileProperty = capture.Events[0].Properties?.FirstOrDefault(p => p.Key == "IgnorePatternFile");
-		KeyValuePair<string, object?>? patternProperty = capture.Events[0].Properties?.FirstOrDefault(p => p.Key == "IgnorePattern");
-		KeyValuePair<string, object?>? levelProperty = capture.Events[0].Properties?.FirstOrDefault(p => p.Key == "OriginalLevel");
-		Assert.IsNotNull(fileProperty);
-		Assert.AreEqual("/workspace/Engine/Build/Horde/IgnorePatterns.txt", fileProperty.Value.Value?.ToString());
-		Assert.IsNotNull(patternProperty);
-		Assert.AreEqual(".*hello.*", patternProperty.Value.Value?.ToString());
-		Assert.IsNotNull(levelProperty);
-		Assert.AreEqual("Error", levelProperty.Value.Value?.ToString());
+		KeyValuePair<string, object?>? ignoredLineProperty = capture.Events[0].Properties?.FirstOrDefault(p => p.Key == "IgnoredLine");
+		Assert.IsNotNull(ignoredLineProperty);
+		LogValue ignoredLineValue = (LogValue)ignoredLineProperty.Value.Value!;
+		Assert.AreEqual(LogValueType.IgnoredLine, ignoredLineValue.Type);
+		Assert.IsNotNull(ignoredLineValue.Properties);
+		Assert.AreEqual("/workspace/Engine/Build/Horde/IgnorePatterns.txt", ignoredLineValue.Properties[new Utf8String("ignorePatternFile")]?.ToString());
+		Assert.AreEqual(".*hello.*", ignoredLineValue.Properties[new Utf8String("ignorePattern")]?.ToString());
+		Assert.AreEqual("Error", ignoredLineValue.Properties[new Utf8String("originalLevel")]?.ToString());
 	}
 
 	[TestMethod]
@@ -775,17 +774,16 @@ public class LogEventParserTests
 		Assert.AreEqual(1, capture.Events.Count);
 		Assert.AreEqual(LogLevel.Information, capture.Events[0].Level);
 
-		// Verify structured properties tracking the ignore source are present
+		// Verify IgnoredLine LogValue with structured properties
 		Assert.IsNotNull(capture.Events[0].Properties);
-		KeyValuePair<string, object?>? fileProperty = capture.Events[0].Properties?.FirstOrDefault(p => p.Key == "IgnorePatternFile");
-		KeyValuePair<string, object?>? patternProperty = capture.Events[0].Properties?.FirstOrDefault(p => p.Key == "IgnorePattern");
-		KeyValuePair<string, object?>? levelProperty = capture.Events[0].Properties?.FirstOrDefault(p => p.Key == "OriginalLevel");
-		Assert.IsNotNull(fileProperty);
-		Assert.AreEqual("/workspace/Engine/Build/Horde/IgnorePatterns.txt", fileProperty.Value.Value?.ToString());
-		Assert.IsNotNull(patternProperty);
-		Assert.AreEqual(".*C1234.*", patternProperty.Value.Value?.ToString());
-		Assert.IsNotNull(levelProperty);
-		Assert.AreEqual("Error", levelProperty.Value.Value?.ToString());
+		KeyValuePair<string, object?>? ignoredLineProperty = capture.Events[0].Properties?.FirstOrDefault(p => p.Key == "IgnoredLine");
+		Assert.IsNotNull(ignoredLineProperty);
+		LogValue ignoredLineValue = (LogValue)ignoredLineProperty.Value.Value!;
+		Assert.AreEqual(LogValueType.IgnoredLine, ignoredLineValue.Type);
+		Assert.IsNotNull(ignoredLineValue.Properties);
+		Assert.AreEqual("/workspace/Engine/Build/Horde/IgnorePatterns.txt", ignoredLineValue.Properties[new Utf8String("ignorePatternFile")]?.ToString());
+		Assert.AreEqual(".*C1234.*", ignoredLineValue.Properties[new Utf8String("ignorePattern")]?.ToString());
+		Assert.AreEqual("Error", ignoredLineValue.Properties[new Utf8String("originalLevel")]?.ToString());
 	}
 
 	#endregion
