Horde: Make log and HTTP blob uploading more robust

- Use a bare SocketsHttpHandler for upload redirect clients instead of the Horde-server-tuned Polly retry/timeout policies
- Prevent ServerLogger.DisposeAsync from propagating exceptions thrown by the tail task or blob writer
- Fix redirect uploads to presigned URLs sending an empty body on retry by creating a fresh StreamContent and resetting stream position per attempt

diff --git a/Engine/Source/Programs/Shared/EpicGames.Horde.Tests/Logs/ServerLoggerTests.cs b/Engine/Source/Programs/Shared/EpicGames.Horde.Tests/Logs/ServerLoggerTests.cs
--- a/Engine/Source/Programs/Shared/EpicGames.Horde.Tests/Logs/ServerLoggerTests.cs
+++ b/Engine/Source/Programs/Shared/EpicGames.Horde.Tests/Logs/ServerLoggerTests.cs
@@ -2,6 +2,7 @@
 
 using System;
 using System.Collections.Generic;
+using System.IO;
 using System.Linq;
 using System.Threading;
 using System.Threading.Channels;
@@ -129,6 +130,7 @@ namespace EpicGames.Horde.Tests
 		{
 			public FakeLogRpcClient LogRpc { get; } = new FakeLogRpcClient();
 			public Dictionary<string, BundleStorageNamespace> StorageNamespaces { get; } = [];
+			public Func<IStorageNamespace, IStorageNamespace>? StorageNamespaceDecorator { get; set; }
 
 			public Uri ServerUrl => throw new NotImplementedException();
 
@@ -158,7 +160,12 @@ namespace EpicGames.Horde.Tests
 					storageNamespace = BundleStorageNamespace.CreateInMemory(NullLogger.Instance);
 					StorageNamespaces.Add(relativePath, storageNamespace);
 				}
-				return storageNamespace;
+				IStorageNamespace result = storageNamespace;
+				if (StorageNamespaceDecorator != null)
+				{
+					result = StorageNamespaceDecorator(result);
+				}
+				return result;
 			}
 
 			public ValueTask DisposeAsync()
@@ -189,6 +196,40 @@ namespace EpicGames.Horde.Tests
 				=> throw new NotImplementedException();
 		}
 
+		class ThrowOnWriterDisposeStorageNamespace(IStorageNamespace inner) : IStorageNamespace
+		{
+			public IBlobRef CreateBlobRef(BlobLocator locator) => inner.CreateBlobRef(locator);
+			public IBlobWriter CreateBlobWriter(string? basePath = null, BlobSerializerOptions? serializerOptions = null, CancellationToken cancellationToken = default)
+				=> new ThrowOnDisposeBlobWriter(inner.CreateBlobWriter(basePath, serializerOptions, cancellationToken));
+			public IStorageWriter CreateWriter(CancellationToken cancellationToken = default) => inner.CreateWriter(cancellationToken);
+			public Task<BlobAlias[]> FindAliasesAsync(string name, int? maxResults = null, CancellationToken cancellationToken = default) => inner.FindAliasesAsync(name, maxResults, cancellationToken);
+			public Task<IHashedBlobRef?> TryReadRefAsync(RefName name, RefCacheTime cacheTime = default, CancellationToken cancellationToken = default) => inner.TryReadRefAsync(name, cacheTime, cancellationToken);
+			public void GetStats(StorageStats stats) => inner.GetStats(stats);
+		}
+
+		class ThrowOnDisposeBlobWriter(IBlobWriter inner) : IBlobWriter
+		{
+			public BlobSerializerOptions Options => inner.Options;
+			public ReadOnlyMemory<byte> WrittenMemory => inner.WrittenMemory;
+			public int Length => inner.Length;
+			public void AddAlias(string name, int rank, ReadOnlyMemory<byte> data) => inner.AddAlias(name, rank, data);
+			public Task FlushAsync(CancellationToken cancellationToken = default) => inner.FlushAsync(cancellationToken);
+			public IBlobWriter Fork() => inner.Fork();
+			public ValueTask<IHashedBlobRef> CompleteAsync(BlobType type, CancellationToken cancellationToken = default) => inner.CompleteAsync(type, cancellationToken);
+			public ValueTask<IHashedBlobRef<T>> CompleteAsync<T>(BlobType type, CancellationToken cancellationToken = default) => inner.CompleteAsync<T>(type, cancellationToken);
+			public void WriteBlobHandleDangerous(IBlobRef handle) => inner.WriteBlobHandleDangerous(handle);
+			public void WriteBlobRef(IHashedBlobRef blobRef) => inner.WriteBlobRef(blobRef);
+			public Span<byte> GetSpan(int sizeHint = 0) => inner.GetSpan(sizeHint);
+			public Memory<byte> GetMemory(int sizeHint = 0) => inner.GetMemory(sizeHint);
+			public void Advance(int length) => inner.Advance(length);
+
+			public ValueTask DisposeAsync()
+			{
+				GC.SuppressFinalize(this);
+				return ValueTask.FromException(new IOException("Simulated S3 upload failure during writer disposal"));
+			}
+		}
+
 		static List<T> DrainChannel<T>(Channel<T> channel)
 		{
 			List<T> items = [];
@@ -540,5 +581,21 @@ namespace EpicGames.Horde.Tests
 
 			await logger.DisposeAsync();
 		}
+
+		[TestMethod]
+		public async Task DisposeAsync_DoesNotThrow_WhenWriterDisposeThrowsAsync()
+		{
+			await using FakeHordeClient hordeClient = new();
+			hordeClient.StorageNamespaceDecorator = ns => new ThrowOnWriterDisposeStorageNamespace(ns);
+			LogId logId = default;
+
+			// The logger works normally during its lifecycle...
+			ServerLogger logger = new(hordeClient, logId, LogLevel.Information, NullLogger.Instance);
+			logger.LogInformation("Line before dispose");
+
+			// ... but DisposeAsync must not propagate the IOException from the writer.
+			// Before the fix, this would throw IOException("Simulated S3 upload failure during writer disposal").
+			await logger.DisposeAsync();
+		}
 	}
 }
diff --git a/Engine/Source/Programs/Shared/EpicGames.Horde.Tests/Storage/HttpStorageBackendTests.cs b/Engine/Source/Programs/Shared/EpicGames.Horde.Tests/Storage/HttpStorageBackendTests.cs
new file mode 100644
--- /dev/null
+++ b/Engine/Source/Programs/Shared/EpicGames.Horde.Tests/Storage/HttpStorageBackendTests.cs
@@ -0,0 +1,82 @@
+// Copyright Epic Games, Inc. All Rights Reserved.
+
+using System.IO;
+using System.Net;
+using System.Net.Http;
+using System.Threading;
+using System.Threading.Tasks;
+using EpicGames.Horde.Storage;
+using EpicGames.Horde.Storage.Backends;
+using Microsoft.Extensions.Logging.Abstractions;
+using Microsoft.VisualStudio.TestTools.UnitTesting;
+
+namespace EpicGames.Horde.Tests.Storage;
+
+[TestClass]
+public class HttpStorageBackendTests
+{
+	/// <summary>
+	/// HttpResponseMessage subclass that tracks whether Dispose was called.
+	/// </summary>
+	sealed class TrackableHttpResponseMessage : HttpResponseMessage
+	{
+		public bool WasDisposed { get; private set; }
+
+		public TrackableHttpResponseMessage(HttpStatusCode statusCode) : base(statusCode) { }
+
+		protected override void Dispose(bool disposing)
+		{
+			WasDisposed = true;
+			base.Dispose(disposing);
+		}
+	}
+
+	/// <summary>
+	/// Mock handler for OpenBlobAsync: returns blob data in a trackable response so we can
+	/// verify the response is disposed when the stream is disposed.
+	/// </summary>
+	sealed class BlobReadHandler : HttpMessageHandler
+	{
+		public TrackableHttpResponseMessage? LastResponse { get; private set; }
+
+		protected override Task<HttpResponseMessage> SendAsync(HttpRequestMessage request, CancellationToken cancellationToken)
+		{
+			TrackableHttpResponseMessage response = new(HttpStatusCode.OK)
+			{
+				Content = new ByteArrayContent([1, 2, 3, 4, 5])
+			};
+			LastResponse = response;
+			return Task.FromResult<HttpResponseMessage>(response);
+		}
+	}
+
+	[TestMethod]
+	public async Task OpenBlobAsync_DisposesResponseWhenStreamDisposedAsync()
+	{
+		// Arrange
+		using BlobReadHandler handler = new();
+
+		HttpClient CreateApiClient()
+		{
+			HttpClient client = new(handler, disposeHandler: false);
+			client.BaseAddress = new("https://horde.example.com/");
+			return client;
+		}
+
+		HttpStorageBackend backend = new("api/v1/storage/ns", CreateApiClient, () => new HttpClient(), NullLogger<HttpStorageBackend>.Instance);
+
+		BlobLocator locator = new("test-blob");
+		Stream stream = await backend.OpenBlobAsync(locator, 0, null);
+
+		Assert.IsNotNull(handler.LastResponse);
+		Assert.IsFalse(handler.LastResponse.WasDisposed, "Response should not be disposed while stream is open");
+
+		// Act: dispose the stream.
+		// Before the fix: the raw content stream was returned, and disposing it did NOT dispose the HttpResponseMessage.
+		// After the fix: ResponseOwningStream disposes the HttpResponseMessage when the stream is disposed.
+		await stream.DisposeAsync();
+
+		// Assert
+		Assert.IsTrue(handler.LastResponse.WasDisposed, "Response should be disposed when the stream is disposed");
+	}
+}
diff --git a/Engine/Source/Programs/Shared/EpicGames.Horde/HordeClient.cs b/Engine/Source/Programs/Shared/EpicGames.Horde/HordeClient.cs
--- a/Engine/Source/Programs/Shared/EpicGames.Horde/HordeClient.cs
+++ b/Engine/Source/Programs/Shared/EpicGames.Horde/HordeClient.cs
@@ -238,7 +238,13 @@ namespace EpicGames.Horde
 				}
 			}
 
-			HttpStorageBackend httpStorageBackend = new(basePath, CreateClient, CreateUnauthenticatedHttpClient, _loggerFactory.CreateLogger<HttpStorageBackend>());
+			// Upload redirect clients talk directly to S3 (or similar) via presigned URLs.
+			// Use a bare SocketsHttpHandler without the Polly retry/timeout policies that are
+			// designed for Horde server communication. The single outer retry policy in
+			// WriteBlobAsync provides adequate retry behavior for transient S3 errors.
+			static HttpClient CreateUploadRedirectClient() => new(HttpMessageHandlers.CreateDefaultHandler());
+
+			HttpStorageBackend httpStorageBackend = new(basePath, CreateClient, CreateUploadRedirectClient, _loggerFactory.CreateLogger<HttpStorageBackend>());
 			return new BundleStorageNamespace(httpStorageBackend, _bundleCache, _hordeOptions.Bundle, _loggerFactory.CreateLogger<BundleStorageNamespace>());
 		}
 
diff --git a/Engine/Source/Programs/Shared/EpicGames.Horde/Logs/ServerLogger.cs b/Engine/Source/Programs/Shared/EpicGames.Horde/Logs/ServerLogger.cs
--- a/Engine/Source/Programs/Shared/EpicGames.Horde/Logs/ServerLogger.cs
+++ b/Engine/Source/Programs/Shared/EpicGames.Horde/Logs/ServerLogger.cs
@@ -126,14 +126,25 @@ namespace EpicGames.Horde.Logs
 				_tailTaskStop.Latch();
 				_newTailDataEvent.Latch();
 
-				await _tailTask;
+				try
+				{
+					await _tailTask;
+				}
+				catch (Exception ex)
+				{
+					_internalLogger.LogError(ex, "Exception while waiting for tail task to complete (log {LogId})", _logId);
+				}
 				_tailTaskStop = null!;
 			}
 
-			if (_writer != null)
+			try
 			{
 				await _writer.DisposeAsync();
 			}
+			catch (Exception ex)
+			{
+				_internalLogger.LogError(ex, "Exception while disposing blob writer (log {LogId})", _logId);
+			}
 		}
 
 		/// <summary>
diff --git a/Engine/Source/Programs/Shared/EpicGames.Horde/Storage/Backends/HttpStorageBackend.cs b/Engine/Source/Programs/Shared/EpicGames.Horde/Storage/Backends/HttpStorageBackend.cs
--- a/Engine/Source/Programs/Shared/EpicGames.Horde/Storage/Backends/HttpStorageBackend.cs
+++ b/Engine/Source/Programs/Shared/EpicGames.Horde/Storage/Backends/HttpStorageBackend.cs
@@ -30,6 +30,58 @@ namespace EpicGames.Horde.Storage.Backends
 			public bool? SupportsRedirects { get; set; }
 		}
 
+		/// <summary>
+		/// Wraps a stream and owns the <see cref="HttpResponseMessage"/>, disposing it when the stream is disposed.
+		/// This ensures the underlying TCP connection is returned to the pool promptly.
+		/// </summary>
+		sealed class ResponseOwningStream(HttpResponseMessage response, Stream inner) : Stream
+		{
+			public override bool CanRead => inner.CanRead;
+			public override bool CanSeek => inner.CanSeek;
+			public override bool CanWrite => inner.CanWrite;
+			public override long Length => inner.Length;
+			public override long Position { get => inner.Position; set => inner.Position = value; }
+
+			public override int Read(byte[] buffer, int offset, int count) => inner.Read(buffer, offset, count);
+			public override Task<int> ReadAsync(byte[] buffer, int offset, int count, CancellationToken cancellationToken) => inner.ReadAsync(buffer, offset, count, cancellationToken);
+			public override ValueTask<int> ReadAsync(Memory<byte> buffer, CancellationToken cancellationToken = default) => inner.ReadAsync(buffer, cancellationToken);
+			public override long Seek(long offset, SeekOrigin origin) => inner.Seek(offset, origin);
+			public override void SetLength(long value) => inner.SetLength(value);
+			public override void Write(byte[] buffer, int offset, int count) => inner.Write(buffer, offset, count);
+			public override void Flush() => inner.Flush();
+			public override Task FlushAsync(CancellationToken cancellationToken) => inner.FlushAsync(cancellationToken);
+			public override Task CopyToAsync(Stream destination, int bufferSize, CancellationToken cancellationToken) => inner.CopyToAsync(destination, bufferSize, cancellationToken);
+
+			protected override void Dispose(bool disposing)
+			{
+				if (disposing)
+				{
+					try
+					{
+						inner.Dispose();
+					}
+					finally
+					{
+						response.Dispose();
+					}
+				}
+				base.Dispose(disposing);
+			}
+
+			public override async ValueTask DisposeAsync()
+			{
+				try
+				{
+					await inner.DisposeAsync();
+				}
+				finally
+				{
+					response.Dispose();
+				}
+				await base.DisposeAsync();
+			}
+		}
+
 		readonly string _basePath;
 		readonly Func<HttpClient> _createClient;
 		readonly Func<HttpClient> _createUploadRedirectClient;
@@ -89,7 +141,8 @@ namespace EpicGames.Horde.Storage.Backends
 					return new MemoryStream([]);
 				}
 
-				using (HttpClient httpClient = _createClient())
+				HttpClient httpClient = _createClient();
+				try
 				{
 					IAsyncPolicy<HttpResponseMessage> retryPolicy = HordeHttpMessageHandler.CreateDefaultTransientErrorPolicy();
 
@@ -120,7 +173,14 @@ namespace EpicGames.Horde.Storage.Backends
 					{
 					}
 
-					return stream;
+					// Wrap the stream so the HttpResponseMessage (and its underlying connection)
+					// is disposed when the caller disposes the stream, rather than waiting for GC.
+					return new ResponseOwningStream(response, stream);
+				}
+				catch
+				{
+					httpClient.Dispose();
+					throw;
 				}
 			}
 			finally
@@ -157,33 +217,50 @@ namespace EpicGames.Horde.Storage.Backends
 		/// <inheritdoc/>
 		public async Task<BlobLocator> WriteBlobAsync(BlobLocator? locator, Stream stream, IReadOnlyCollection<BlobLocator>? imports, string? prefix = null, CancellationToken cancellationToken = default)
 		{
-			using StreamContent streamContent = new(stream);
-
 			if (_supportsUploadRedirects)
 			{
 				WriteBlobResponse redirectResponse = await SendWriteRequestAsync(locator, null, imports, prefix, cancellationToken);
 				if (redirectResponse.UploadUrl != null)
 				{
+					long streamLength = stream.CanSeek ? stream.Length : -1;
+
 					using HttpClient uploadRedirectClient = _createUploadRedirectClient();
-					// The http client timeout default of 100 seconds acts as a global timeout across all retry and wait periods
-					// We're creating a client per upload, so bump the global client timeout to account for retries specified on the on policy
+					// HttpClient.Timeout is a wall-clock deadline for the entire operation including retries.
+					// The default of 100s is too low for the outer retry policy (6 retries, exponential backoff up to 64s).
+					// 240s allows ~5 retry attempts before the deadline fires.
 					uploadRedirectClient.Timeout = TimeSpan.FromSeconds(240);
 					IAsyncPolicy<HttpResponseMessage> retryPolicy = HordeHttpMessageHandler.CreateDefaultTransientErrorPolicy();
 
-					using HttpResponseMessage uploadResponse = await retryPolicy.ExecuteAsync(() => uploadRedirectClient.PutAsync(redirectResponse.UploadUrl, streamContent, cancellationToken));
+					int attemptNum = 0;
+					using HttpResponseMessage uploadResponse = await retryPolicy.ExecuteAsync(() =>
+					{
+						// Reset stream position before each attempt so retries send the full data
+						if (stream.CanSeek)
+						{
+							stream.Position = 0;
+						}
+						else if (attemptNum > 0)
+						{
+							_logger.LogWarning("Retrying upload to {Url} but stream is not seekable; retry will send from current position", redirectResponse.UploadUrl);
+						}
+						attemptNum++;
+						StreamContent streamContent = new(stream);
+						return uploadRedirectClient.PutAsync(redirectResponse.UploadUrl, streamContent, cancellationToken);
+					});
 					if (!uploadResponse.IsSuccessStatusCode)
 					{
 						string body = await uploadResponse.Content.ReadAsStringAsync(cancellationToken);
 						string headersText = String.Join(", ", uploadResponse.Headers.Select(h => $"{h.Key}={String.Join(",", h.Value)}"));
-						throw new StorageException($"Unable to upload data to redirected URL {redirectResponse.UploadUrl}. Status:{uploadResponse.StatusCode} Reason:{uploadResponse.ReasonPhrase} Response:{body} Headers:{headersText}");
+						throw new StorageException($"Unable to upload data ({streamLength} bytes) to redirected URL {redirectResponse.UploadUrl}. Status:{uploadResponse.StatusCode} Reason:{uploadResponse.ReasonPhrase} Response:{body} Headers:{headersText}");
 					}
 
-					_logger.LogDebug("Written {Locator} (using redirect)", redirectResponse.Blob);
+					_logger.LogDebug("Written {Locator} ({Size} bytes, using redirect)", redirectResponse.Blob, streamLength);
 					return new BlobLocator(redirectResponse.Blob);
 				}
 			}
 
-			WriteBlobResponse response = await SendWriteRequestAsync(locator, streamContent, imports, prefix, cancellationToken);
+			using StreamContent directStreamContent = new(stream);
+			WriteBlobResponse response = await SendWriteRequestAsync(locator, directStreamContent, imports, prefix, cancellationToken);
 			_supportsUploadRedirects = response.SupportsRedirects ?? false;
 			_logger.LogDebug("Written {Locator} (direct)", response.Blob);
 			return new BlobLocator(response.Blob);
