[UBT]
* Added support for disabling casdb on agent side

#Synced-CL 51249403

diff --git a/Engine/Source/Programs/Shared/EpicGames.UBA/IServer.cs b/Engine/Source/Programs/Shared/EpicGames.UBA/IServer.cs
--- a/Engine/Source/Programs/Shared/EpicGames.UBA/IServer.cs
+++ b/Engine/Source/Programs/Shared/EpicGames.UBA/IServer.cs
@@ -42,8 +42,8 @@ namespace EpicGames.UBA
 		/// <summary>
 		/// Set config (.toml) that should be used by all clients.
 		/// </summary>
-		/// <param name="path"></param>
-		public abstract void SetClientsConfig(string path);
+		/// <param name="config"></param>
+		public abstract void SetClientsConfig(IConfig config);
 
 		/// <summary>
 		/// Create a IServer object
diff --git a/Engine/Source/Programs/Shared/EpicGames.UBA/Impl/ServerImpl.cs b/Engine/Source/Programs/Shared/EpicGames.UBA/Impl/ServerImpl.cs
--- a/Engine/Source/Programs/Shared/EpicGames.UBA/Impl/ServerImpl.cs
+++ b/Engine/Source/Programs/Shared/EpicGames.UBA/Impl/ServerImpl.cs
@@ -84,9 +84,8 @@ namespace EpicGames.UBA.Impl
 		public bool AddNamedConnection(string name) => NetworkServer_AddNamedConnection(_handle, name) != 0;
 
 		/// <inheritdoc/>
-		public void SetClientsConfig(string path)
+		public void SetClientsConfig(IConfig config)
 		{
-			using IConfig config = IConfig.LoadConfig(path);
 			NetworkServer_SetClientsConfig(_handle, config.GetHandle());
 		}
 
@@ -107,7 +106,7 @@ namespace EpicGames.UBA.Impl
 		/// <inheritdoc/>
 		public bool AddClient(string ip, int port, string crypto = "") => false;
 		/// <inheritdoc/>
-		public void SetClientsConfig(string path) { }
+		public void SetClientsConfig(IConfig config) { }
 		/// <inheritdoc/>
 		public nint GetHandle() => 0;
 		/// <inheritdoc/>
diff --git a/Engine/Source/Programs/UnrealBuildTool/Executors/UnrealBuildAccelerator/UBAExecutor.cs b/Engine/Source/Programs/UnrealBuildTool/Executors/UnrealBuildAccelerator/UBAExecutor.cs
--- a/Engine/Source/Programs/UnrealBuildTool/Executors/UnrealBuildAccelerator/UBAExecutor.cs
+++ b/Engine/Source/Programs/UnrealBuildTool/Executors/UnrealBuildAccelerator/UBAExecutor.cs
@@ -445,7 +445,13 @@ namespace UnrealBuildTool
 
 			if (!UBAConfig.bDisableRemote)
 			{
-				_server.SetClientsConfig(FileReference.Combine(ubaBinariesDir, "UbaAgent.toml").FullName);
+				using IConfig agentConfig = IConfig.LoadConfig(FileReference.Combine(ubaBinariesDir, "UbaAgent.toml").FullName);
+				if (!UBAConfig.UseAgentStore)
+				{
+					agentConfig.AddValue("Session", "UseStorage", "false");
+					agentConfig.AddValue("Storage", "WriteToDisk", "false");
+				}
+				_server.SetClientsConfig(agentConfig);
 
 				_server.StartServer(UBAConfig.Host, UBAConfig.Port, Crypto);
 			}
diff --git a/Engine/Source/Programs/UnrealBuildTool/Executors/UnrealBuildAccelerator/UnrealBuildAcceleratorConfig.cs b/Engine/Source/Programs/UnrealBuildTool/Executors/UnrealBuildAccelerator/UnrealBuildAcceleratorConfig.cs
--- a/Engine/Source/Programs/UnrealBuildTool/Executors/UnrealBuildAccelerator/UnrealBuildAcceleratorConfig.cs
+++ b/Engine/Source/Programs/UnrealBuildTool/Executors/UnrealBuildAccelerator/UnrealBuildAcceleratorConfig.cs
@@ -455,5 +455,10 @@ namespace UnrealBuildTool
 		/// Whether UBA can retry actions natively if they fail when detoured locally
 		/// </summary>
 		public bool CanRetryNatively { get; set; } = true;
+
+		/// <summary>
+		/// Disable this to stop UbaAgent instances to use casdb store
+		/// </summary>
+		public bool UseAgentStore { get; set; } = true;
 	}
 }
\ No newline at end of file
