UnrealHeaderTool: Fix typo in comment.

#Synced-CL 51291646

diff --git a/Engine/Source/Programs/Shared/EpicGames.UHT/Types/UhtEnum.cs b/Engine/Source/Programs/Shared/EpicGames.UHT/Types/UhtEnum.cs
--- a/Engine/Source/Programs/Shared/EpicGames.UHT/Types/UhtEnum.cs
+++ b/Engine/Source/Programs/Shared/EpicGames.UHT/Types/UhtEnum.cs
@@ -224,7 +224,7 @@ namespace EpicGames.UHT.Types
 		public string CppType { get; set; } = String.Empty;
 
 		/// <summary>
-		/// Fully qualified cpp type name. Inc;ude enclosing namespace as well as CppType (which may include enum's own namespace) 
+		/// Fully qualified cpp type name. Include enclosing namespace as well as CppType (which may include enum's own namespace) 
 		/// </summary>
 		public string FullyQualifiedCppType => _fullyQualifiedCppType ??= $"{Namespace.FullSourceName}{CppType}";
 		private string? _fullyQualifiedCppType = null;
