EpicGames.Perforce: Add Perforce edge server to test fixture

Add an edge server (port 1667) alongside the existing commit server (port 1666).

#skipci

diff --git a/Engine/Source/Programs/Shared/EpicGames.Perforce.Fixture/BasePerforceFixtureTest.cs b/Engine/Source/Programs/Shared/EpicGames.Perforce.Fixture/BasePerforceFixtureTest.cs
--- a/Engine/Source/Programs/Shared/EpicGames.Perforce.Fixture/BasePerforceFixtureTest.cs
+++ b/Engine/Source/Programs/Shared/EpicGames.Perforce.Fixture/BasePerforceFixtureTest.cs
@@ -18,8 +18,10 @@ namespace EpicGames.Perforce.Fixture;
 public abstract class BasePerforceFixtureTest : IDisposable
 {
 	internal const string EnvVar = "EPICGAMES_P4_FIXTURE_SERVER_URL";
+	internal const string EdgeEnvVar = "EPICGAMES_P4_FIXTURE_EDGE_SERVER_URL";
 	private const string UseDefaultMarker = "default";
 	private const string DefaultUrl = "p4://test.user@localhost:1666";
+	private const string DefaultEdgeUrl = "p4://test.user@localhost:1667";
 
 	/// <summary>
 	/// Console-based logger factory
diff --git a/Engine/Source/Programs/Shared/EpicGames.Perforce.Fixture/Dockerfile b/Engine/Source/Programs/Shared/EpicGames.Perforce.Fixture/Dockerfile
--- a/Engine/Source/Programs/Shared/EpicGames.Perforce.Fixture/Dockerfile
+++ b/Engine/Source/Programs/Shared/EpicGames.Perforce.Fixture/Dockerfile
@@ -11,14 +11,19 @@ RUN tar zxf helix-core-server.tgz && \
     mv ./p4 /usr/local/bin/p4 && \
     mkdir -p /app/data
 
-# Start the server and generate changelists
+# Start the server and generate changelists + edge server
 WORKDIR /app
 ADD setup.sh /app/setup.sh
 RUN dos2unix setup.sh && bash setup.sh
 
+ADD entrypoint.sh /app/entrypoint.sh
+RUN dos2unix /app/entrypoint.sh && chmod +x /app/entrypoint.sh
+
 # Copy everything to separate image to minimize image size and number of layers
 FROM docker.io/library/ubuntu:22.04
 COPY --from=builder /usr/local/bin/p4 /usr/local/bin/p4d /usr/local/bin/
 COPY --from=builder /app/data /app/data
+COPY --from=builder /app/edge-data /app/edge-data
+COPY --from=builder /app/entrypoint.sh /app/entrypoint.sh
 COPY --from=builder /trigger-log.sh /trigger-log.sh
-ENTRYPOINT ["p4d", "-r", "/app/data", "-p", "1666"]
+ENTRYPOINT ["/app/entrypoint.sh"]
diff --git a/Engine/Source/Programs/Shared/EpicGames.Perforce.Fixture/entrypoint.sh b/Engine/Source/Programs/Shared/EpicGames.Perforce.Fixture/entrypoint.sh
new file mode 100644
--- /dev/null
+++ b/Engine/Source/Programs/Shared/EpicGames.Perforce.Fixture/entrypoint.sh
@@ -0,0 +1,15 @@
+#!/bin/bash
+
+# Forward signals to child processes for clean Docker shutdown
+trap 'kill $(jobs -p) 2>/dev/null; wait' SIGTERM SIGINT
+
+# Start the commit server
+p4d -r /app/data -p 1666 &
+
+# Wait for the commit server to accept connections before starting the edge
+until p4 -p localhost:1666 info > /dev/null 2>&1; do sleep 0.5; done
+
+# Start the edge server (connects to commit server via configured P4TARGET)
+p4d -r /app/edge-data -p 1667 &
+
+wait
diff --git a/Engine/Source/Programs/Shared/EpicGames.Perforce.Fixture/setup.sh b/Engine/Source/Programs/Shared/EpicGames.Perforce.Fixture/setup.sh
--- a/Engine/Source/Programs/Shared/EpicGames.Perforce.Fixture/setup.sh
+++ b/Engine/Source/Programs/Shared/EpicGames.Perforce.Fixture/setup.sh
@@ -209,5 +209,66 @@ Paths:
         share ...
 EOF
 
-# Stop server running in background
-pkill p4d
\ No newline at end of file
+# Configure Perforce edge server
+EDGE_PORT=1667
+
+# Use test.user for edge setup commands (needs super access)
+export P4USER=test.user
+
+echo "Configuring edge server..."
+
+# Set the commit server's identity
+p4 serverid commit
+
+# Create a service user for the edge server
+p4 user -f -i << EOF
+User:	svc-edge
+Type:	service
+Email:	svc-edge@localhost
+FullName:	Edge Service User
+AuthMethod:	perforce
+EOF
+
+# Grant the service user super access needed for replication
+p4 protect -i << EOF
+Protections:
+	write user * * //...
+	super user test.user * //...
+	super user svc-edge * //...
+EOF
+
+# Create the edge server spec on the commit server
+p4 server -i << EOF
+ServerID:	edge
+Type:	server
+Name:	edge
+Address:	localhost:$EDGE_PORT
+Services:	edge-server
+User:	svc-edge
+Description:
+	Edge server for testing
+EOF
+
+# Configure replication settings for the edge (stored on the commit server)
+p4 configure set "edge#P4TARGET=localhost:$PORT"
+p4 configure set "edge#db.replication=readonly"
+p4 configure set "edge#lbr.replication=readonly"
+p4 configure set "edge#serviceUser=svc-edge"
+p4 configure set "edge#startup.1=pull -i 1"
+p4 configure set "edge#startup.2=pull -u -i 1"
+p4 configure set "edge#db.peeking=2"
+p4 configure set "edge#client.readonly.dir=/app/edge-part-db-have"
+
+# Stop the commit server cleanly
+p4 admin stop
+wait || true
+
+# Create a journal dump of the commit server database
+p4d -r /app/data -jd /tmp/commit.ckp
+
+# Seed the edge server database from the commit server checkpoint
+mkdir -p /app/edge-data
+p4d -r /app/edge-data -jr /tmp/commit.ckp
+
+# Write the edge server identity file (overrides the commit server ID from the checkpoint)
+echo edge > /app/edge-data/server.id
\ No newline at end of file
