EpicGames.Perforce: Fix unshelve not invalidating deleted files in workspace tracking

When unshelving a CL containing file deletes, the deleted files were not marked as invalid in workspace tracking. On the next sync, the tracker would see the stale ContentId and skip re-syncing those files, leaving them permanently missing from disk when UseHaveTable=false.

diff --git a/Engine/Source/Programs/Shared/EpicGames.Perforce.Fixture/BasePerforceFixtureTest.cs b/Engine/Source/Programs/Shared/EpicGames.Perforce.Fixture/BasePerforceFixtureTest.cs
--- a/Engine/Source/Programs/Shared/EpicGames.Perforce.Fixture/BasePerforceFixtureTest.cs
+++ b/Engine/Source/Programs/Shared/EpicGames.Perforce.Fixture/BasePerforceFixtureTest.cs
@@ -21,7 +21,9 @@ public abstract class BasePerforceFixtureTest : IDisposable
 	internal const string EdgeEnvVar = "EPICGAMES_P4_FIXTURE_EDGE_SERVER_URL";
 	private const string UseDefaultMarker = "default";
 	private const string DefaultUrl = "p4://test.user@localhost:1666";
+#pragma warning disable CA1823, IDE0051 // Unused field - kept for future use
 	private const string DefaultEdgeUrl = "p4://test.user@localhost:1667";
+#pragma warning restore CA1823, IDE0051
 
 	/// <summary>
 	/// Console-based logger factory
diff --git a/Engine/Source/Programs/Shared/EpicGames.Perforce.Managed.Tests/ManagedWorkspaceTest.cs b/Engine/Source/Programs/Shared/EpicGames.Perforce.Managed.Tests/ManagedWorkspaceTest.cs
--- a/Engine/Source/Programs/Shared/EpicGames.Perforce.Managed.Tests/ManagedWorkspaceTest.cs
+++ b/Engine/Source/Programs/Shared/EpicGames.Perforce.Managed.Tests/ManagedWorkspaceTest.cs
@@ -252,6 +252,27 @@ public class ManagedWorkspaceTest : BasePerforceFixtureTest
 		await Stream.GetChangelist(7).AssertHaveTableAsync(PerforceConnection, useHaveTable);
 	}
 
+	[TestMethod]
+	[DataRow(true, DisplayName = "With have-table")]
+	[DataRow(false, DisplayName = "Without have-table")]
+	public async Task Unshelve_WithDeletedFiles_NextSync_RestoresDeletedFilesAsync(bool useHaveTable)
+	{
+		ManagedWorkspace mw = await CreateManagedWorkspaceAsync(useHaveTable);
+
+		// Sync CL 7 (has main.cpp)
+		await SyncAsync(mw, 7);
+		Stream.GetChangelist(7).AssertDepotFiles(SyncDir);
+
+		// Unshelve CL 8 (deletes main.cpp)
+		await mw.UnshelveAsync(PerforceConnection, 8, CancellationToken.None);
+		Stream.GetChangelist(8).AssertDepotFiles(SyncDir);
+
+		// Sync CL 7 again - main.cpp should be restored
+		await SyncAsync(mw, 7);
+		Stream.GetChangelist(7).AssertDepotFiles(SyncDir);
+		await Stream.GetChangelist(7).AssertHaveTableAsync(PerforceConnection, useHaveTable);
+	}
+
 	[TestMethod]
 	[DataRow(true, DisplayName = "With have-table")]
 	[DataRow(false, DisplayName = "Without have-table")]
diff --git a/Engine/Source/Programs/Shared/EpicGames.Perforce.Managed/ManagedWorkspace.cs b/Engine/Source/Programs/Shared/EpicGames.Perforce.Managed/ManagedWorkspace.cs
--- a/Engine/Source/Programs/Shared/EpicGames.Perforce.Managed/ManagedWorkspace.cs
+++ b/Engine/Source/Programs/Shared/EpicGames.Perforce.Managed/ManagedWorkspace.cs
@@ -1061,8 +1061,8 @@ namespace EpicGames.Perforce.Managed
 		{
 			using TelemetrySpan span = _tracer.StartActiveSpan($"{nameof(ManagedWorkspace)}.{nameof(UnshelveAsync)}");
 
-			// Need to mark those files as dirty - update the workspace with those files 
-			// Delete is fine, but need to flag anything added
+			// Need to mark those files as dirty - update the workspace with those files
+			// Both written and deleted files must be flagged so the next sync re-syncs them
 
 			Stopwatch timer = Stopwatch.StartNew();
 			_logger.LogInformation("Unshelving changelist {Change}...", unshelveChangelist);
@@ -1118,13 +1118,14 @@ namespace EpicGames.Perforce.Managed
 				throw new Exception($"Unknown client {perforce.Settings.ClientName}");
 			}
 
-			// Add all the files to be written to the workspace with invalid metadata. This will ensure they're removed on next clean.
-			if (writeFiles.Count > 0)
+			// Add all modified and deleted files to the workspace with invalid metadata. This will ensure they're re-synced next time.
+			List<WhereRecord> invalidateFiles = [..writeFiles, ..deleteFiles];
+			if (invalidateFiles.Count > 0)
 			{
-				_logger.LogInformation("Removing {NumFiles} files from tracked workspace", writeFiles.Count);
-				foreach (WhereRecord writeFile in writeFiles)
+				_logger.LogInformation("Invalidating {NumFiles} files in tracked workspace", invalidateFiles.Count);
+				foreach (WhereRecord file in invalidateFiles)
 				{
-					string path = Regex.Replace(writeFile.ClientFile, "^//[^/]+/", "");
+					string path = Regex.Replace(file.ClientFile, "^//[^/]+/", "");
 					_workspace.AddFile(new Utf8String(path), 0, 0, false, new FileContentId(Md5Hash.Zero, default));
 				}
 				await SaveAsync(TransactionState.Clean, CancellationToken.None);
