[UBT]
* Removed codegen function GetPrivateStaticClass since it is only called in one place and code can just as well be there
* Made RegisterNativeFunc optional (handle null) and removed all the stub functions not needed in codegen

#Synced-CL 51430576

diff --git a/Engine/Source/Programs/Shared/EpicGames.UHT/Exporters/CodeGen/UhtHeaderCodeGeneratorCppFile.cs b/Engine/Source/Programs/Shared/EpicGames.UHT/Exporters/CodeGen/UhtHeaderCodeGeneratorCppFile.cs
--- a/Engine/Source/Programs/Shared/EpicGames.UHT/Exporters/CodeGen/UhtHeaderCodeGeneratorCppFile.cs
+++ b/Engine/Source/Programs/Shared/EpicGames.UHT/Exporters/CodeGen/UhtHeaderCodeGeneratorCppFile.cs
@@ -2922,44 +2922,6 @@ namespace EpicGames.UHT.Exporters.CodeGen
 
 			bool hasInterfaces = classObj.Bases.Any(x => x is UhtClass baseClass && baseClass.ClassFlags.HasAnyFlags(EClassFlags.Interface)) || classObj.FlattenedVerseInterfaces.Count > 0;
 
-			// This block of code is the contents of IMPLEMENT_CLASS_NO_AUTO_REGISTRATION 
-			if (Session.IsUsingCompiledInObjectFormat(UhtCompiledInObjectFormat.ConstInit))
-			{
-				using UhtConditionalMacroBlock block = new(builder, "UE_WITH_CONSTINIT_UOBJECT", Session.IsUsingMultipleCompiledInObjectFormats);
-				builder.Append($"UClass* {classObj.Namespace.FullSourceName}{classObj.SourceName}::GetPrivateStaticClass()\r\n");
-				builder.Append("{\r\n");
-				builder.Append($"\treturn &{GetSingletonName(classObj, UhtSingletonType.ConstInit)};\r\n");
-				builder.Append("}\r\n");
-			}
-			if (Session.IsUsingCompiledInObjectFormat(UhtCompiledInObjectFormat.Params))
-			{
-				using UhtConditionalMacroBlock block = new(builder, "!UE_WITH_CONSTINIT_UOBJECT", Session.IsUsingMultipleCompiledInObjectFormats);
-				builder.Append("FClassRegistrationInfo ").Append(registrationName).Append(";\r\n");
-				builder.Append("UClass* ").Append(classObj.Namespace.FullSourceName).Append(classObj.SourceName).Append("::GetPrivateStaticClass()\r\n");
-				builder.Append("{\r\n");
-				builder.Append("\tusing TClass = ").Append(classObj.Namespace.FullSourceName).Append(classObj.SourceName).Append(";\r\n");
-				builder.Append("\tif (!").Append(registrationName).Append(".InnerSingleton)\r\n");
-				builder.Append("\t{\r\n");
-				builder.Append("\t\t").Append(classObj.IsVerseField ? "Verse::CodeGen::Private::ConstructUVerseClassNoInit" : "GetPrivateStaticClassBody").Append("(\r\n");
-				builder.Append("\t\t\tTClass::StaticPackage(),\r\n");
-				builder.Append("\t\t\tTEXT(\"").Append(classObj.EngineName).Append("\"),\r\n");
-				builder.Append("\t\t\t").Append(registrationName).Append(".InnerSingleton,\r\n");
-				builder.Append("\t\t\tStaticRegisterNatives").Append(classObj.SourceName).Append(",\r\n");
-				builder.Append("\t\t\tsizeof(TClass),\r\n");
-				builder.Append("\t\t\talignof(TClass),\r\n");
-				builder.Append("\t\t\tTClass::StaticClassFlags,\r\n");
-				builder.Append("\t\t\tTClass::StaticClassCastFlags(),\r\n");
-				builder.Append("\t\t\tTClass::StaticConfigName(),\r\n");
-				builder.Append("\t\t\t(UClass::ClassConstructorType)InternalConstructor<TClass>,\r\n");
-				builder.Append("\t\t\t(UClass::ClassVTableHelperCtorCallerType)InternalVTableHelperCtorCaller<TClass>,\r\n");
-				builder.Append("\t\t\tUOBJECT_CPPCLASS_STATICFUNCTIONS_FORCLASS(TClass),\r\n");
-				builder.Append("\t\t\t&TClass::Super::StaticClass,\r\n");
-				builder.Append("\t\t\t&TClass::WithinClass::StaticClass\r\n");
-				builder.Append("\t\t);\r\n");
-				builder.Append("\t}\r\n");
-				builder.Append("\treturn ").Append(registrationName).Append(".InnerSingleton;\r\n");
-				builder.Append("}\r\n");
-			}
 
 			// Everything from this point on will be part of the definition hash
 			int hashCodeBlockStart = builder.Length;
@@ -3334,38 +3296,54 @@ namespace EpicGames.UHT.Exporters.CodeGen
 				}
 			}
 
-			// Native function registration
 			if (Session.IsUsingCompiledInObjectFormat(UhtCompiledInObjectFormat.Params))
 			{
 				using UhtConditionalMacroBlock block = new(builder, "!UE_WITH_CONSTINIT_UOBJECT", Session.IsUsingMultipleCompiledInObjectFormats);
-				builder.Append($"void {classObj.Namespace.FullSourceName}{classObj.SourceName}::StaticRegisterNatives{classObj.SourceName}()\r\n");
-				builder.Append("{\r\n");
 
+				string registerNativesFuncName = "nullptr";
 				if (hasNativeFunctions)
 				{
-					builder.Append($"\tUClass* Class = {classObj.SourceName}::StaticClass();\r\n");
+					registerNativesFuncName = $"{classObj.SourceName}_StaticRegisterNatives{classObj.SourceName}";
+					builder.Append($"static void {registerNativesFuncName}()\r\n");
+					builder.Append("{\r\n");
+					builder.Append($"\tUClass* Class = {classObj.Namespace.FullSourceName}{classObj.SourceName}::StaticClass();\r\n");
 					builder.Append($"\tFNativeFunctionRegistrar::RegisterFunctions(Class, ").AppendArrayView(nativeFunctions, UhtNames.StaticsFuncsId, 2, ");\r\n");
-
 					if (hasVerseCallableFunctions)
 					{
 						builder.Append("#if WITH_VERSE_VM\r\n");
 						builder.Append("\tVerse::CodeGen::Private::RegisterVerseCallableThunks(Class, UHT_STATICS::VerseFuncs, UE_ARRAY_COUNT(UHT_STATICS::VerseFuncs));\r\n");
 						builder.Append("#endif\r\n");
 					}
+					builder.Append("}\r\n"); // Close StaticRegisterNatives function
 				}
 
-				builder.Append("}\r\n"); // Close StaticRegisterNatives function
-			}
 
-			// Class registration
-			if (Session.IsUsingCompiledInObjectFormat(UhtCompiledInObjectFormat.Params))
-			{
-				using UhtConditionalMacroBlock block = new(builder, "!UE_WITH_CONSTINIT_UOBJECT", Session.IsUsingMultipleCompiledInObjectFormats);
+				builder.Append("FClassRegistrationInfo ").Append(registrationName).Append(";\r\n");
 				builder.Append("UClass* ").Append(singletonName).Append("(ETypeConstructPhase Phase)\r\n");
 				builder.Append("{\r\n");
 				builder.Append("\tif (Phase == ETypeConstructPhase::Inner)\r\n");
 				builder.Append("\t{\r\n");
-				builder.Append("\t\treturn ").Append(classObj.Namespace.FullSourceName).Append(classObj.SourceName).Append("::GetPrivateStaticClass();\r\n");
+				builder.Append("\t\tusing TClass = ").Append(classObj.Namespace.FullSourceName).Append(classObj.SourceName).Append(";\r\n");
+				builder.Append("\t\tif (!").Append(registrationName).Append(".InnerSingleton)\r\n");
+				builder.Append("\t\t{\r\n");
+				builder.Append("\t\t\t").Append(classObj.IsVerseField ? "Verse::CodeGen::Private::ConstructUVerseClassNoInit" : "GetPrivateStaticClassBody").Append("(\r\n");
+				builder.Append("\t\t\t\tTClass::StaticPackage(),\r\n");
+				builder.Append("\t\t\t\tTEXT(\"").Append(classObj.EngineName).Append("\"),\r\n");
+				builder.Append("\t\t\t\t").Append(registrationName).Append(".InnerSingleton,\r\n");
+				builder.Append("\t\t\t\t").Append(registerNativesFuncName).Append(",\r\n");
+				builder.Append("\t\t\t\tsizeof(TClass),\r\n");
+				builder.Append("\t\t\t\talignof(TClass),\r\n");
+				builder.Append("\t\t\t\tTClass::StaticClassFlags,\r\n");
+				builder.Append("\t\t\t\tTClass::StaticClassCastFlags(),\r\n");
+				builder.Append("\t\t\t\tTClass::StaticConfigName(),\r\n");
+				builder.Append("\t\t\t\t(UClass::ClassConstructorType)InternalConstructor<TClass>,\r\n");
+				builder.Append("\t\t\t\t(UClass::ClassVTableHelperCtorCallerType)InternalVTableHelperCtorCaller<TClass>,\r\n");
+				builder.Append("\t\t\t\tUOBJECT_CPPCLASS_STATICFUNCTIONS_FORCLASS(TClass),\r\n");
+				builder.Append("\t\t\t\t&TClass::Super::StaticClass,\r\n");
+				builder.Append("\t\t\t\t&TClass::WithinClass::StaticClass\r\n");
+				builder.Append("\t\t\t);\r\n");
+				builder.Append("\t\t}\r\n");
+				builder.Append("\t\treturn ").Append(registrationName).Append(".InnerSingleton;\r\n");
 				builder.Append("\t}\r\n");
 				builder.Append("\tif (!").Append(registrationName).Append(".OuterSingleton)\r\n");
 				builder.Append("\t{\r\n");
diff --git a/Engine/Source/Programs/Shared/EpicGames.UHT/Exporters/CodeGen/UhtHeaderCodeGeneratorHFile.cs b/Engine/Source/Programs/Shared/EpicGames.UHT/Exporters/CodeGen/UhtHeaderCodeGeneratorHFile.cs
--- a/Engine/Source/Programs/Shared/EpicGames.UHT/Exporters/CodeGen/UhtHeaderCodeGeneratorHFile.cs
+++ b/Engine/Source/Programs/Shared/EpicGames.UHT/Exporters/CodeGen/UhtHeaderCodeGeneratorHFile.cs
@@ -1642,9 +1642,7 @@ namespace EpicGames.UHT.Exporters.CodeGen
 		{
 			UhtCodeGenerator.ObjectInfo info = ObjectInfos[classObj.ObjectTypeIndex];
 
-			// Export the class's native function registration.
 			builder.Append("private: \\\r\n");
-			builder.Append("\tstatic void StaticRegisterNatives").Append(classObj.SourceName).Append("(); \\\r\n");
 			if (!classObj.ClassExportFlags.HasAnyFlags(UhtClassExportFlags.NoExport))
 			{
 				builder.Append($"\tfriend struct ::{info.RegisteredSingletonName}_Statics; \\\r\n");
@@ -1660,7 +1658,6 @@ namespace EpicGames.UHT.Exporters.CodeGen
 				}
 			}
 
-			builder.Append('\t').Append("static UClass* GetPrivateStaticClass(); \\\r\n");
 			builder.Append("\tfriend ").Append(classObj.Module.Api).Append("UClass* ::").Append(info.RegisteredSingletonName).Append("(ETypeConstructPhase); \\\r\n");
 
 			builder.Append("public: \\\r\n");
diff --git a/Engine/Source/Runtime/CoreUObject/Private/UObject/UObjectConstructInternal.h b/Engine/Source/Runtime/CoreUObject/Private/UObject/UObjectConstructInternal.h
--- a/Engine/Source/Runtime/CoreUObject/Private/UObject/UObjectConstructInternal.h
+++ b/Engine/Source/Runtime/CoreUObject/Private/UObject/UObjectConstructInternal.h
@@ -81,7 +81,10 @@ void ConstructUClassNoInitHelper(
 				))
 				{
 					// Register the class's native functions.
-					RegisterNativeFunc();
+					if (RegisterNativeFunc)
+					{
+						RegisterNativeFunc();
+					}
 				}
 				return;
 			}
@@ -127,7 +130,10 @@ void ConstructUClassNoInitHelper(
 	);
 
 	// Register the class's native functions.
-	RegisterNativeFunc();
+	if (RegisterNativeFunc)
+	{
+		RegisterNativeFunc();
+	}
 }
 
 template<typename UEnumClass, typename EnumParams, typename PostNewFn>
