[Backout] - CL51429293
#rnx
Original CL Desc
-----------------------------------------------------------------
[CoreUObject/UHT]
* Added support for UPARTIAL in unreal.

diff --git a/Engine/Source/Programs/LowLevelTests/FoundationTests/Private/UPartialTestClasses.h b/Engine/Source/Programs/LowLevelTests/FoundationTests/Private/UPartialTestClasses.h
deleted file mode 100644
diff --git a/Engine/Source/Programs/LowLevelTests/FoundationTests/Private/UPartialTestOutlinedFunctions.cpp b/Engine/Source/Programs/LowLevelTests/FoundationTests/Private/UPartialTestOutlinedFunctions.cpp
deleted file mode 100644
diff --git a/Engine/Source/Programs/LowLevelTests/FoundationTests/Private/UPartialTestPartials.h b/Engine/Source/Programs/LowLevelTests/FoundationTests/Private/UPartialTestPartials.h
deleted file mode 100644
diff --git a/Engine/Source/Programs/LowLevelTests/FoundationTests/Tests/UPartialTests.cpp b/Engine/Source/Programs/LowLevelTests/FoundationTests/Tests/UPartialTests.cpp
deleted file mode 100644
diff --git a/Engine/Source/Programs/Shared/EpicGames.Core/UnrealEngineTypes.cs b/Engine/Source/Programs/Shared/EpicGames.Core/UnrealEngineTypes.cs
--- a/Engine/Source/Programs/Shared/EpicGames.Core/UnrealEngineTypes.cs
+++ b/Engine/Source/Programs/Shared/EpicGames.Core/UnrealEngineTypes.cs
@@ -158,11 +158,6 @@ namespace EpicGames.Core
 		/// </summary>
 		Native = 0x00000080u,
 
-		/// <summary>
-		/// Class is a partial class
-		/// </summary>
-		Partial = 0x00000100u,
-
 		/// <summary>
 		/// Do not allow users to create in the editor.
 		/// </summary>
@@ -281,7 +276,7 @@ namespace EpicGames.Core
 		/// <summary>
 		/// Flags to inherit from base class
 		/// </summary>
-		Inherit = Transient | Optional | DefaultConfig | Config | Partial | PerObjectConfig | ConfigDoNotCheckDefaults | NotPlaceable | Const | HasInstancedReference |
+		Inherit = Transient | Optional | DefaultConfig | Config | PerObjectConfig | ConfigDoNotCheckDefaults | NotPlaceable | Const | HasInstancedReference |
 			Deprecated | DefaultToInstanced | GlobalUserConfig | ProjectUserConfig | NeedsDeferredDependencyLoading,
 
 		/// <summary>
@@ -302,7 +297,7 @@ namespace EpicGames.Core
 		/// <summary>
 		/// This is used as a mask for the flags put into generated code for "compiled in" classes.
 		/// </summary>
-		SaveInCompiledInClasses = Abstract | DefaultConfig | GlobalUserConfig | ProjectUserConfig | PerPlatformConfig | Config | Transient | Optional | Native | Partial | NotPlaceable | PerObjectConfig |
+		SaveInCompiledInClasses = Abstract | DefaultConfig | GlobalUserConfig | ProjectUserConfig | PerPlatformConfig | Config | Transient | Optional | Native | NotPlaceable | PerObjectConfig |
 			ConfigDoNotCheckDefaults | EditInlineNew | CollapseCategories | Interface | DefaultToInstanced | HasInstancedReference | Hidden | Deprecated |
 			HideDropDown | Intrinsic | Const | MinimalAPI | RequiredAPI | MatchedSerializers | NeedsDeferredDependencyLoading,
 	};
diff --git a/Engine/Source/Programs/Shared/EpicGames.UHT/Exporters/CodeGen/UhtCodeGenerator.cs b/Engine/Source/Programs/Shared/EpicGames.UHT/Exporters/CodeGen/UhtCodeGenerator.cs
--- a/Engine/Source/Programs/Shared/EpicGames.UHT/Exporters/CodeGen/UhtCodeGenerator.cs
+++ b/Engine/Source/Programs/Shared/EpicGames.UHT/Exporters/CodeGen/UhtCodeGenerator.cs
@@ -510,16 +510,8 @@ namespace EpicGames.UHT.Exporters.CodeGen
 			objectInfo.RegisteredSingletonName = builder.ToString();
 			objectInfo.StaticsName = $"{objectInfo.RegisteredSingletonName}_Statics";
 			objectInfo.RegisteredExternalDecl = $"\t{api}U{obj.EngineLinkClassName}* {objectInfo.RegisteredSingletonName}(ETypeConstructPhase);\r\n";
-			if (obj is UhtPartial)
-			{
-				objectInfo.ConstInitSingletonName = $"{objectInfo.RegisteredSingletonName}_PartialParams";
-				objectInfo.ConstInitExternalDecl = $"\textern UE_CONSTINIT_UOBJECT_DECL UE::CoreUObject::Private::FPartialClass {objectInfo.ConstInitSingletonName};\r\n";
-			}
-			else
-			{
-				objectInfo.ConstInitSingletonName = objectInfo.RegisteredSingletonName.Replace("Z_Construct_", "Z_ConstInit_", StringComparison.InvariantCultureIgnoreCase);
-				objectInfo.ConstInitExternalDecl = $"\textern UE_CONSTINIT_UOBJECT_DECL TNoDestroy<U{obj.EngineClassName}> {objectInfo.ConstInitSingletonName};\r\n";
-			}
+			objectInfo.ConstInitSingletonName = objectInfo.RegisteredSingletonName.Replace("Z_Construct_", "Z_ConstInit_", StringComparison.InvariantCultureIgnoreCase);
+			objectInfo.ConstInitExternalDecl = $"\textern UE_CONSTINIT_UOBJECT_DECL TNoDestroy<U{obj.EngineClassName}> {objectInfo.ConstInitSingletonName};\r\n";
 
 			// Init the children
 			foreach (UhtType child in obj.Children)
diff --git a/Engine/Source/Programs/Shared/EpicGames.UHT/Exporters/CodeGen/UhtHeaderCodeGenerator.cs b/Engine/Source/Programs/Shared/EpicGames.UHT/Exporters/CodeGen/UhtHeaderCodeGenerator.cs
--- a/Engine/Source/Programs/Shared/EpicGames.UHT/Exporters/CodeGen/UhtHeaderCodeGenerator.cs
+++ b/Engine/Source/Programs/Shared/EpicGames.UHT/Exporters/CodeGen/UhtHeaderCodeGenerator.cs
@@ -129,11 +129,6 @@ namespace EpicGames.UHT.Exporters.CodeGen
 			{
 				outerName = classObj.EngineName;
 			}
-			else if (outer is UhtPartial partial)
-			{
-				// For functions in partials, use the extended class name for parameter struct naming
-				outerName = partial.OwnerClass != null ? partial.OwnerClass.EngineName : partial.EngineName;
-			}
 			else if (outer is UhtPackage packageObj)
 			{
 				string packageName = packageObj.EngineName;
@@ -674,7 +669,6 @@ namespace EpicGames.UHT.Exporters.CodeGen
 			{
 				UhtClass classObj => classObj.GeneratedBodyLineNumber,
 				UhtScriptStruct scriptStruct => scriptStruct.MacroDeclaredLineNumber,
-				UhtPartial partial => partial.MacroDeclaredLineNumber,
 				UhtFunction function => function.MacroLineNumber,
 				_ => throw new UhtException(type, "Can not use given type to create a macro"),
 			};
diff --git a/Engine/Source/Programs/Shared/EpicGames.UHT/Exporters/CodeGen/UhtHeaderCodeGeneratorCppFile.cs b/Engine/Source/Programs/Shared/EpicGames.UHT/Exporters/CodeGen/UhtHeaderCodeGeneratorCppFile.cs
--- a/Engine/Source/Programs/Shared/EpicGames.UHT/Exporters/CodeGen/UhtHeaderCodeGeneratorCppFile.cs
+++ b/Engine/Source/Programs/Shared/EpicGames.UHT/Exporters/CodeGen/UhtHeaderCodeGeneratorCppFile.cs
@@ -21,7 +21,6 @@ namespace EpicGames.UHT.Exporters.CodeGen
 		public UhtUsedDefineScopes<UhtScriptStruct> ScriptStructs { get; } = new();
 		public UhtUsedDefineScopes<UhtScriptStruct> RigVMStructs { get; } = new();
 		public UhtUsedDefineScopes<UhtClass> Classes { get; } = new();
-		public UhtUsedDefineScopes<UhtPartial> Partials { get; } = new();
 	}
 
 	internal sealed class UhtHeaderCodeGeneratorCppFile : UhtHeaderCodeGenerator
@@ -95,20 +94,8 @@ namespace EpicGames.UHT.Exporters.CodeGen
 					includesToAdd.Add("UObject/VerseClassProperty.h");
 				}
 
-				// Collect unique partial class names from all partials in this file
-				Dictionary<string, UhtClass> partialClassNames = [];
-
 				foreach (UhtType type in HeaderFile.Children)
 				{
-					if (type is UhtPartial partial && partial.OwnerClass != null)
-					{
-						if (partialClassNames.Count == 0)
-						{
-							includesToAdd.Add("UObject/UObjectPartials.h");
-						}
-						partialClassNames.TryAdd(partial.OwnerClassName, partial.OwnerClass);
-					}
-
 					if (type is UhtStruct structObj)
 					{
 						// Functions
@@ -288,10 +275,6 @@ namespace EpicGames.UHT.Exporters.CodeGen
 						UhtSingletonTypeFlag filter = UhtSingletonTypeFlag.Registered | UhtSingletonTypeFlag.Unregistered | UhtSingletonTypeFlag.ConstInit;
 						foreach (UhtObject obj in HeaderFile.References.InternalObjects.SortReferences(filter))
 						{
-							if (obj is UhtPartial)
-							{
-								continue;
-							}
 							builder.Append(GetExternalDecl(obj, UhtSingletonType.ConstInit).AsSpan().TrimStart());
 						}
 					}
@@ -368,13 +351,6 @@ namespace EpicGames.UHT.Exporters.CodeGen
 							}
 						}
 					}
-					else if (field is UhtPartial partial)
-					{
-						using UhtCodeBlockComment blockComment = new(builder, field);
-						using UhtMacroBlockEmitter macroBlockEmitter = new(builder, partial.DefineScope);
-						AppendPartial(builder, partial);
-						GetRegistrations(packageRegistrations, field).Partials.Add(partial);
-					}
 					else if (field is UhtFunction function)
 					{
 						using UhtCodeBlockComment blockComment = new(builder, field);
@@ -486,21 +462,6 @@ namespace EpicGames.UHT.Exporters.CodeGen
 								combinedHash = IoHash.Combine(combinedHash, hash);
 							},
 							builder => builder.Append("\t};\r\n"));
-
-						builder.AppendInstances(registrations.Partials,
-							builder => builder.Append("\tstatic constexpr FPartialRegisterCompiledInInfo PartialInfo[] = {\r\n"),
-							(builder, partial) =>
-							{
-								IoHash hash = ObjectInfos[partial.ObjectTypeIndex].Hash;
-								builder.Append("\t\t{ ")
-									.Append($"[](ETypeConstructPhase) -> UE::CoreUObject::Private::FPartialClass* {{ return &{GetSingletonName(partial, UhtSingletonType.Registered)}_PartialParams; }},")
-									.Append($"TEXT(\"{partial.EngineName}\"), ")
-									.Append($"&Z_Registration_Info_Partial_{partial.SourceName}, ")
-									.Append($"CONSTRUCT_RELOAD_VERSION_INFO(FPartialReloadVersionInfo, sizeof({partial.FullyQualifiedSourceName}), {(uint)hash}U)")
-									.Append(" },\r\n");
-								combinedHash = IoHash.Combine(combinedHash, hash);
-							},
-							builder => builder.Append("\t};\r\n"));
 					}
 
 					builder.Append($"}}; // UHT_STATICS \r\n");
@@ -527,7 +488,6 @@ namespace EpicGames.UHT.Exporters.CodeGen
 							.AppendArrayPtrAndCountLine(registrations.Classes, UhtNames.StaticsClassInfoId, 1, ",\r\n")
 							.AppendArrayPtrAndCountLine(registrations.ScriptStructs, UhtNames.StaticsScriptStructInfoId, 1, ",\r\n")
 							.AppendArrayPtrAndCountLine(registrations.Enumerations, UhtNames.StaticsEnumInfoId, 1, ",\r\n")
-							.AppendArrayPtrAndCountLine(registrations.Partials, UhtNames.StaticsPartialInfoId, 1, ",\r\n")
 							.Append("};\r\n");
 					}
 				}
@@ -1342,172 +1302,13 @@ namespace EpicGames.UHT.Exporters.CodeGen
 			return builder;
 		}
 
-		private StringBuilder AppendPartial(StringBuilder builder, UhtPartial partial)
-		{
-			// Generate functions for this partial first
-			// Functions must be generated before the function pointer array
-			if (partial.OwnerClass != null)
-			{
-				foreach (UhtFunction function in partial.Functions)
-				{
-					AppendClassFunction(builder, partial.OwnerClass, function);
-				}
-			}
-
-			string registrationName = $"Z_Registration_Info_Partial_{partial.SourceName}";
-
-			if (Session.IsUsingCompiledInObjectFormat(UhtCompiledInObjectFormat.Params))
-			{
-				builder.Append("static FPartialRegistrationInfo ").Append(registrationName).Append(";\r\n");
-			}
-
-			// Collect the properties
-			UhtPropertyMemberContextImpl context = new(this, partial);
-			UhtUsedDefineScopes<UhtProperty> properties = new(partial.Properties);
-
-			using UhtStaticsEmitter statics = new(builder, this, partial);// staticsName);
-
-			builder.Append($"struct UHT_STATICS\r\n");
-			builder.Append("{\r\n");
-
-			// Meta data
-			// No need adding this since the partial "disappears" in runtime so there is way to access the information here
-			builder.AppendConditionalMetaData(partial, UhtNames.TypeId, context, properties, 1);
-
-			// Declare properties
-			// ROBM: NOTE: This should already create property linked lists in the constinit case, but the head of the list is usually a field in a UStruct
-			AppendPropertiesDecl(builder, context, properties, 1);
-
-			if (Session.IsUsingCompiledInObjectFormat(UhtCompiledInObjectFormat.ConstInit) && partial.Functions.Any())
-			{
-				using UhtConditionalMacroBlock block = new(builder, "UE_WITH_CONSTINIT_UOBJECT", Session.IsUsingMultipleCompiledInObjectFormats);
-				builder.Append("\tstatic consteval UFunction* GetFirstFunction()\r\n");
-				builder.Append("\t{\r\n");
-				builder.AppendScopeLink(partial.FirstFunction, context, 2, "return ", ";\r\n");
-				builder.Append("\t}\r\n");
-			}
-
-			builder.Append("};\r\n\r\n");
-
-			// Define properties
-			AppendPropertiesDefs(builder, context, properties, 0);
-			builder.Append("\r\n");
-
-
-			// Wrappers
-			builder.Append($"static void Construct_{partial.SourceName}(void* Memory) {{ new(Memory) {partial.Namespace.FullSourceName}{partial.SourceName}(); }}\r\n");
-			builder.Append($"static void Destruct_{partial.SourceName}(void* Memory) {{ (({partial.Namespace.FullSourceName}{partial.SourceName}*)Memory)->~{partial.SourceName}(); }}\r\n");
-			if (partial.HasBeginDestroy)
-			{
-				// BeginDestroy wrapper (only if partial defines BeginDestroy method)
-				// For now, we'll always generate it and have it check for the method
-				builder.Append($"static void BeginDestroy_{partial.SourceName}(void* Memory) {{ (({partial.Namespace.FullSourceName}{partial.SourceName}*)Memory)->BeginDestroy(); }}\r\n");
-			}
-
-			// Define static PartialOffset member
-			builder.Append($"int32 {partial.Namespace.FullSourceName}{partial.SourceName}::PartialOffset;\r\n\r\n");
-
-			// Forward declare the partial class for partial function compilation
-			if (partial.OwnerClass != null && partial.Functions.Any())
-			{
-				builder.Append($"// Forward declaration for partial function compilation\r\n");
-				builder.Append($"class {partial.OwnerClass.Namespace.FullSourceName}{partial.OwnerClass.SourceName};\r\n\r\n");
-			}
-
-			// Generate function arrays for partial functions (after functions are generated)
-			if (!Session.IsUsingCompiledInObjectFormat(UhtCompiledInObjectFormat.ConstInit) && partial.Functions.Any())
-			{
-				// Function constructor pointers (for creating UFunctions)
-				builder.Append($"// Function constructor pointers for {partial.SourceName}\r\n");
-				builder.Append($"static UFunction* (*PartialFunctionConstructors_{partial.SourceName}[])(ETypeConstructPhase) = {{\r\n");
-				foreach (UhtFunction function in partial.Functions)
-				{
-					builder.Append($"\t&{GetSingletonName(function, UhtSingletonType.Registered)},\r\n");
-				}
-				builder.Append("};\r\n\r\n");
-
-				// Native function bindings (for registering exec functions)
-				builder.Append($"// Native function bindings for {partial.SourceName}\r\n");
-				builder.Append($"static constexpr UE::CodeGen::FClassNativeFunction PartialNativeFunctions_{partial.SourceName}[] = {{\r\n");
-				foreach (UhtFunction function in partial.Functions)
-				{
-					builder.Append($"\t{{ .NameUTF8 = UTF8TEXT(\"{function.EngineName}\"), .Pointer = &{partial.SourceName}::{function.UnMarshalAndCallName} }},\r\n");
-				}
-				builder.Append("};\r\n\r\n");
-			}
-
-			// ROBM: There should be an if/else constinit case here, and the constinit case should just store the head of the linked list of properties etc
-			// See UClass/UStruct versions 
-			// There should be a global var like `Z_ConstInit_UPartial_{BaseName}_{PartialName}` here that can be linked in the package code generator 
-			builder.Append($"UE_CONSTINIT_UOBJECT_DECL UE::CoreUObject::Private::FPartialClass {GetSingletonName(partial, UhtSingletonType.Registered)}_PartialParams\r\n");
-			builder.Append("{\r\n");
-			builder.Append($"\t.Name = TEXT(\"{partial.SourceName}\"),\r\n");
-			//builder.Append($"\t.ClassName = TEXT(\"{partial.OwnerClassName}\"),\r\n");
-			builder.Append($"\t.Size = sizeof({partial.Namespace.FullSourceName}{partial.SourceName}),\r\n");
-			builder.Append($"\t.Alignment = alignof({partial.Namespace.FullSourceName}{partial.SourceName}),\r\n");
-
-			if (Session.IsUsingCompiledInObjectFormat(UhtCompiledInObjectFormat.ConstInit))
-			{
-				using UhtConditionalMacroBlock block = new(builder, "UE_WITH_CONSTINIT_UOBJECT", Session.IsUsingMultipleCompiledInObjectFormats);
-				builder.Append($"\t.Class = &{GetSingletonName(partial.OwnerClass, UhtSingletonType.ConstInit)},\r\n");
-
-				if (partial.Properties.Any())
-				{
-					builder.Append($"\t.FirstProperty = UHT_STATICS::GetFirstProperty(),\r\n");
-				}
-				if (partial.Functions.Any())
-				{
-					builder.Append($"\t.FirstChild = UHT_STATICS::GetFirstFunction(),\r\n");
-				}
-			}
-
-			if (Session.IsUsingCompiledInObjectFormat(UhtCompiledInObjectFormat.Params))
-			{
-				using UhtConditionalMacroBlock block = new(builder, "!UE_WITH_CONSTINIT_UOBJECT", Session.IsUsingMultipleCompiledInObjectFormats);
-				builder.Append($"\t.Class = {GetSingletonName(partial.OwnerClass, UhtSingletonType.Registered)},\r\n");
-				if (partial.Properties.Any())
-				{
-					builder.Append($"\t.Properties = UHT_STATICS::PropPointers,\r\n");
-					builder.Append($"\t.NumProperties = UE_ARRAY_COUNT(UHT_STATICS::PropPointers),\r\n");
-				}
-				if (partial.Functions.Any())
-				{
-					builder.Append($"\t.FunctionConstructors = PartialFunctionConstructors_{partial.SourceName},\r\n");
-					builder.Append($"\t.NativeFunctions = PartialNativeFunctions_{partial.SourceName},\r\n");
-					builder.Append($"\t.NumFunctions = {partial.Functions.Count()},\r\n");
-				}
-			}
-
-			builder.Append($"\t.PartialOffsetPtr = &{partial.Namespace.FullSourceName}{partial.SourceName}::PartialOffset,\r\n");
-			builder.Append($"\t.Constructor = &Construct_{partial.SourceName},\r\n");
-			builder.Append($"\t.BeginDestroy = {(partial.HasBeginDestroy ? $"BeginDestroy_{partial.SourceName}" : "nullptr")},\r\n");
-			builder.Append($"\t.Destructor = &Destruct_{partial.SourceName},\r\n");
-			builder.Append($"\t.GetLifetimeReplicatedProps = nullptr, // TODO: Detect if method exists\r\n");
-			builder.Append("};\r\n");
-
-			return builder;
-		}
-
 		private StringBuilder AppendFunction(StringBuilder builder, UhtFunction function, bool isNoExport, bool isVerseNativeCallable)
 		{
 			string singletonName = GetSingletonName(function, UhtSingletonType.Registered);
 			string constinitName = GetSingletonName(function, UhtSingletonType.ConstInit);
 			bool paramsInStatic = isNoExport || !IsCallbackFunction(function);
 			bool isNet = function.FunctionFlags.HasAnyFlags(EFunctionFlags.NetRequest | EFunctionFlags.NetResponse);
-
-			string fullSourceName = "";
-			string callName = "";
-
-			// Get the outer class - either directly or from the partial's partial class
-			if (function.Outer is UhtClass outerClass)
-			{
-				fullSourceName = outerClass.Namespace.FullSourceName;
-				callName = outerClass.NativeFunctionCallName;
-			}
-			else if (function.Outer is UhtPartial partial)
-			{
-				fullSourceName = partial.SourceName;
-			}
+			UhtClass? outerClass = function.Outer as UhtClass;
 
 			using UhtStaticsEmitter statics = new(builder, this, function);
 
@@ -1607,7 +1408,7 @@ namespace EpicGames.UHT.Exporters.CodeGen
 								=> $"\t\t.NativeFunction = UHT_STATICS::GetCustomThunk_{function.SourceName}(),\r\n",
 							// Native without custom thunk
 							(EFunctionFlags.Native, EFunctionFlags.None, _)
-								=> $"\t\t.NativeFunction = &{fullSourceName}{callName}::exec{function.SourceName},\r\n",
+								=> $"\t\t.NativeFunction = &{outerClass!.Namespace.FullSourceName}{outerClass!.NativeFunctionCallName}::exec{function.SourceName},\r\n",
 							// Non-native function 
 							// Can't use cross-dll assignment in modular builds here 
 							(EFunctionFlags.None, _, _) => ".NativeFunction = UE_IF(IS_MONOLITHIC, &UObject::ProcessInternal, nullptr),\r\n",
@@ -1717,15 +1518,6 @@ namespace EpicGames.UHT.Exporters.CodeGen
 			{
 				return GetSingletonName(package, UhtSingletonType.Registered);
 			}
-			else if (function.Outer is UhtPartial partial)
-			{
-				// Functions in partials should use the partial class as their outer
-				if (partial.OwnerClass != null)
-				{
-					return GetSingletonName(partial.OwnerClass, UhtSingletonType.Registered);
-				}
-				return "nullptr";
-			}
 			else
 			{
 				return GetSingletonName((UhtObject)function.Outer, UhtSingletonType.Registered);
@@ -1966,21 +1758,10 @@ namespace EpicGames.UHT.Exporters.CodeGen
 				}
 				if (isRpc)
 				{
-					if (function.Outer is UhtPartial partial)
-					{
-						builder.Append($"DEFINE_FUNCTION({partial.SourceName}::{function.UnMarshalAndCallName})\r\n");
-						builder.Append("{\r\n");
-						builder.Append($"\ttypedef {classObj.Namespace.FullSourceName}{classObj.SourceName} ThisClass;\r\n");
-						AppendFunctionThunk(builder, classObj, function);
-						builder.Append("}\r\n");
-					}
-					else
-					{
-						builder.Append($"DEFINE_FUNCTION({classObj.Namespace.FullSourceName}{classObj.NativeFunctionCallName}::{function.UnMarshalAndCallName})\r\n");
-						builder.Append("{\r\n");
-						AppendFunctionThunk(builder, classObj, function);
-						builder.Append("}\r\n");
-					}
+					builder.Append($"DEFINE_FUNCTION({classObj.Namespace.FullSourceName}{classObj.NativeFunctionCallName}::{function.UnMarshalAndCallName})\r\n");
+					builder.Append("{\r\n");
+					AppendFunctionThunk(builder, classObj, function);
+					builder.Append("}\r\n");
 				}
 			}
 			return builder;
@@ -3552,11 +3333,6 @@ namespace EpicGames.UHT.Exporters.CodeGen
 					builder.Append($"{classObj.SourceName}::{function.CppImplName}(").AppendFunctionThunkParameterNames(function).Append(");\r\n");
 				}
 			}
-			else if (function.Outer is UhtPartial partial)
-			{
-				// Function is in an partial - call through GetPartial<>()
-				builder.Append("P_THIS->GetPartial<").Append(partial.SourceName).Append(">().").Append(function.CppImplName).Append('(').AppendFunctionThunkParameterNames(function).Append(");\r\n");
-			}
 			else
 			{
 				builder.Append("P_THIS->").Append(function.CppImplName).Append('(').AppendFunctionThunkParameterNames(function).Append(");\r\n");
@@ -3654,7 +3430,7 @@ namespace EpicGames.UHT.Exporters.CodeGen
 			{
 				_codeGenerator = codeGenerator;
 				_outerStruct = outerStruct;
-				string outerStructNamespaceName = _outerStruct is UhtScriptStruct scriptStruct ? scriptStruct.NamespaceExportName : (_outerStruct is UhtPartial partial ? partial.NamespaceExportName : _outerStruct.Namespace.FullSourceName);
+				string outerStructNamespaceName = _outerStruct is UhtScriptStruct scriptStruct ? scriptStruct.NamespaceExportName : _outerStruct.Namespace.FullSourceName;
 				string outerStructSourceName = _outerStruct.SourceName;
 				_outerIdentifier = new(outerStructNamespaceName, null, outerStructSourceName, null);
 			}
@@ -3716,13 +3492,8 @@ namespace EpicGames.UHT.Exporters.CodeGen
 			{
 				builder.AppendLine($"{Tabs(tabs + 1)}.Class = {SingletonRef(linker, uhtObject.EngineClass, UhtSingletonType.ConstInit)}, ");
 			}
-			UhtObject? outer = uhtObject.Outer as UhtObject;
-			if (outer is UhtPartial partial)
-			{
-				outer = partial.OwnerClass;
-			}
 			builder.AppendTabs(tabs + 1).Append($".NameUTF8 = UTF8TEXT(").AppendUTF8LiteralString(uhtObject.EngineName).Append("),\r\n");
-			builder.AppendTabs(tabs + 1).AppendLine($".Outer = {ConstInitSingletonRef(linker, outer)},");
+			builder.AppendTabs(tabs + 1).AppendLine($".Outer = {ConstInitSingletonRef(linker, uhtObject.Outer as UhtObject)},");
 			builder.AppendTabs(tabs).Append("},\r\n");
 			return builder;
 		}
diff --git a/Engine/Source/Programs/Shared/EpicGames.UHT/Exporters/CodeGen/UhtHeaderCodeGeneratorHFile.cs b/Engine/Source/Programs/Shared/EpicGames.UHT/Exporters/CodeGen/UhtHeaderCodeGeneratorHFile.cs
--- a/Engine/Source/Programs/Shared/EpicGames.UHT/Exporters/CodeGen/UhtHeaderCodeGeneratorHFile.cs
+++ b/Engine/Source/Programs/Shared/EpicGames.UHT/Exporters/CodeGen/UhtHeaderCodeGeneratorHFile.cs
@@ -181,12 +181,6 @@ namespace EpicGames.UHT.Exporters.CodeGen
 							AppendScriptStruct(builder, scriptStruct);
 						}
 					}
-					else if (field is UhtPartial partial)
-					{
-						using UhtCodeBlockComment blockComment = new(builder, field);
-						using UhtMacroBlockEmitter macroBlockEmitter = new(builder, partial.DefineScope);
-						Appendpartial(builder, partial);
-					}
 					else if (field is UhtFunction function)
 					{
 					}
@@ -736,31 +730,6 @@ namespace EpicGames.UHT.Exporters.CodeGen
 			return builder;
 		}
 
-		private StringBuilder Appendpartial(StringBuilder builder, UhtPartial partial)
-		{
-			using (UhtMacroCreator macro = new(builder, this, partial, GeneratedBodyMacroSuffix))
-			{
-				builder.Append('\t').Append(Module.Api).Append("static int32 PartialOffset; \\\r\n");
-				builder.Append('\t').Append(partial.OwnerClassName).Append("& GetOwner() {");
-				builder.Append(" return *reinterpret_cast<").Append(partial.OwnerClassName).Append("*>((uint8*)this - PartialOffset); } \\\r\n");
-				builder.Append("\tconst ").Append(partial.OwnerClassName).Append("& GetOwner() const {");
-				builder.Append(" return *reinterpret_cast<const ").Append(partial.OwnerClassName).Append("*>((const uint8*)this - PartialOffset); } \\\r\n");
-				if (partial.Functions.Any())
-				{
-					builder.Append("\t \\\r\n");
-					foreach (UhtFunction function in partial.Functions)
-					{
-						if (function.FunctionFlags.HasAnyFlags(EFunctionFlags.Native))
-						{
-							builder.Append("\tDECLARE_FUNCTION(").Append(function.UnMarshalAndCallName).Append("); \\\r\n");
-						}
-					}
-				}
-			}
-
-			return builder;
-		}
-
 		private StringBuilder AppendEnum(StringBuilder builder, UhtEnum enumObj)
 		{
 			// Export FOREACH macro
@@ -1710,14 +1679,6 @@ namespace EpicGames.UHT.Exporters.CodeGen
 					builder.Append("\tDECLARE_WITHIN(").Append(classObj.ClassWithin.SourceName).Append(") \\\r\n");
 				}
 			}
-
-			// Only generate GetPartial<>() if parent is not partial (i.e., this is the first partial class in the hierarchy)
-			if (classObj.ClassFlags.HasFlag(EClassFlags.Partial) && !classObj.SuperClass!.ClassFlags.HasFlag(EClassFlags.Partial))
-			{
-				builder.Append("\ttemplate<typename T> T& GetPartial() { return *reinterpret_cast<T*>((uint8*)this + T::PartialOffset); } \\\r\n");
-				builder.Append("\ttemplate<typename T> const T& GetPartial() const { return *reinterpret_cast<const T*>((const uint8*)this + T::PartialOffset); } \\\r\n");
-			}
-
 			return builder;
 		}
 
@@ -1738,10 +1699,6 @@ namespace EpicGames.UHT.Exporters.CodeGen
 			{
 				builder.Append(" | CLASS_Optional");
 			}
-			if (classObj.ClassFlags.HasAnyFlags(EClassFlags.Partial))
-			{
-				builder.Append(" | CLASS_Partial");
-			}
 			if (classObj.ClassFlags.HasAnyFlags(EClassFlags.DefaultConfig))
 			{
 				builder.Append(" | CLASS_DefaultConfig");
diff --git a/Engine/Source/Programs/Shared/EpicGames.UHT/Exporters/CodeGen/UhtPackageCodeGeneratorCppFile.cs b/Engine/Source/Programs/Shared/EpicGames.UHT/Exporters/CodeGen/UhtPackageCodeGeneratorCppFile.cs
--- a/Engine/Source/Programs/Shared/EpicGames.UHT/Exporters/CodeGen/UhtPackageCodeGeneratorCppFile.cs
+++ b/Engine/Source/Programs/Shared/EpicGames.UHT/Exporters/CodeGen/UhtPackageCodeGeneratorCppFile.cs
@@ -144,9 +144,6 @@ namespace EpicGames.UHT.Exporters.CodeGen
 					.Select((int order, UhtObject obj) (p) => (p.Value, p.Key))
 					.OrderBy(p => p.order)
 					.Select(p => p.obj)];
-
-				List<UhtPartial> partials = Module.Packages.SelectMany(p => p.Children.OfType<UhtPartial>().Where(o => o.AlternateObject is null)).ToList();
-
 				if (Session.IsUsingCompiledInObjectFormat(UhtCompiledInObjectFormat.ConstInit))
 				{
 					using UhtConditionalMacroBlock block = new(builder, "UE_WITH_CONSTINIT_UOBJECT", Session.IsUsingMultipleCompiledInObjectFormats);
@@ -159,11 +156,6 @@ namespace EpicGames.UHT.Exporters.CodeGen
 						builder.Append(GetExternalDecl(obj, UhtSingletonType.ConstInit));
 					}
 
-					foreach (UhtPartial partial in partials)
-					{
-						builder.Append(GetExternalDecl(partial, UhtSingletonType.ConstInit));
-					}
-
 					builder.AppendLine($$"""
 						extern {{Module.Api}}UObject* const {{ModuleLinkableObjects(Module)}}[{{linkableObjectsOrdered.Count}}] = 
 						{
@@ -183,28 +175,6 @@ namespace EpicGames.UHT.Exporters.CodeGen
 
 						""");
 
-					if (partials.Count != 0)
-					{
-						builder.AppendLine($$"""
-							extern UE::CoreUObject::Private::FPartialClass* const Z_Partials_{{Module.ShortName}}[{{partials.Count}}] = 
-							{
-							""");
-						foreach (UhtPartial partial in partials)
-						{
-							using UhtMacroBlockEmitter macroBlockEmitter = new(builder, partial.DefineScope);
-							builder.AppendLine($"\t{SingletonRef(this, partial, UhtSingletonType.ConstInit)}, ");
-							if (partial.DefineScope != UhtDefineScope.None)
-							{
-								macroBlockEmitter.SetElse(UhtDefineScope.None);
-								builder.AppendLine("\tnullptr, ");
-							}
-						}
-						builder.AppendLine($$"""
-							};	
-
-							""");
-					}
-
 					List<UhtModule> orderedLinkedModules = [.. Module.Linking!.ReferencedModules
 						.Select((int Index, UhtModule M) (p) => (p.Value, p.Key))
 						.OrderBy(p => p.Index)
@@ -327,12 +297,8 @@ namespace EpicGames.UHT.Exporters.CodeGen
 							.Structs = MakeConstArrayView({{ModuleLinkableObjects(Module)}}).Mid({{structsStart}}, {{structsEnd - structsStart}}),
 							.Enums = MakeConstArrayView({{ModuleLinkableObjects(Module)}}).Mid({{enumsStart}}, {{enumsEnd - enumsStart}}),
 							.Others = MakeConstArrayView({{ModuleLinkableObjects(Module)}}).RightChop({{othersStart}}),
+						} };
 						""");
-					if (partials.Count > 0)
-					{
-						builder.Append($"\t.Partials = MakeConstArrayView(Z_Partials_{Module.ShortName}),\r\n");
-					}
-					builder.AppendLine("} };");
 				}
 
 				// Verse registration
diff --git a/Engine/Source/Programs/Shared/EpicGames.UHT/Parsers/UhtFunctionParser.cs b/Engine/Source/Programs/Shared/EpicGames.UHT/Parsers/UhtFunctionParser.cs
--- a/Engine/Source/Programs/Shared/EpicGames.UHT/Parsers/UhtFunctionParser.cs
+++ b/Engine/Source/Programs/Shared/EpicGames.UHT/Parsers/UhtFunctionParser.cs
@@ -109,7 +109,6 @@ namespace EpicGames.UHT.Parsers
 		[UhtKeyword(Extends = UhtTableNames.Interface)]
 		[UhtKeyword(Extends = UhtTableNames.NativeInterface)]
 		[UhtKeyword(Extends = UhtTableNames.VModule)]
-		[UhtKeyword(Extends = UhtTableNames.Partial)]
 		[SuppressMessage("CodeQuality", "IDE0051:Remove unused private members", Justification = "Attribute accessed method")]
 		[SuppressMessage("Style", "IDE0060:Remove unused parameter", Justification = "Attribute accessed method")]
 		private static UhtParseResult UFUNCTIONKeyword(UhtParsingScope topScope, UhtParsingScope actionScope, ref UhtToken token)
@@ -314,22 +313,7 @@ namespace EpicGames.UHT.Parsers
 			{
 				using UhtParsingScope topScope = new(parentScope, function, parentScope.Session.GetKeywordTable(UhtTableNames.Function), UhtAccessSpecifier.Public);
 				UhtParsingScope outerClassScope = topScope.CurrentClassScope;
-
-				// Get the outer class - either directly or from partial's owner class
-				UhtClass? outerClass = outerClassScope.ScopeType as UhtClass;
-				UhtPartial? outerPartial = null;
-				if (outerClass == null && outerClassScope.ScopeType is UhtPartial partial)
-				{
-					outerPartial = partial;
-					outerClass = partial.OwnerClass;
-				}
-
-				// If we're not in a class or partial, that's an error
-				if (outerClass == null && outerPartial == null)
-				{
-					throw new UhtException(function, "UFUNCTIONs must be declared in a class or partial");
-				}
-
+				UhtClass outerClass = (UhtClass)outerClassScope.ScopeType;
 				string scopeName = "function";
 
 				{
@@ -340,8 +324,7 @@ namespace EpicGames.UHT.Parsers
 					UhtSpecifierParser specifierParser = UhtSpecifierParser.GetThreadInstance(specifierContext, scopeName, parentScope.Session.GetSpecifierTable(UhtTableNames.Function));
 					specifierParser.ParseSpecifiers();
 
-					// Only check native class flag if we have a resolved outer class (not an partial during parsing)
-					if (outerClass != null && !outerClass.ClassFlags.HasAnyFlags(EClassFlags.Native))
+					if (!outerClass.ClassFlags.HasAnyFlags(EClassFlags.Native))
 					{
 						throw new UhtException(function, "Should only be here for native classes!");
 					}
diff --git a/Engine/Source/Programs/Shared/EpicGames.UHT/Parsers/UhtHeaderFileParser.cs b/Engine/Source/Programs/Shared/EpicGames.UHT/Parsers/UhtHeaderFileParser.cs
--- a/Engine/Source/Programs/Shared/EpicGames.UHT/Parsers/UhtHeaderFileParser.cs
+++ b/Engine/Source/Programs/Shared/EpicGames.UHT/Parsers/UhtHeaderFileParser.cs
@@ -694,22 +694,6 @@ namespace EpicGames.UHT.Parsers
 						}
 					}
 				}
-				else if (topScope.ScopeType is UhtPartial partial)
-				{
-					using UhtTokenRecorder recorder = new(topScope, ref token);
-					topScope.SkipDeclaration(token);
-					if (recorder.Stop())
-					{
-						if (partial.Declarations != null)
-						{
-							UhtDeclaration declaration = partial.Declarations[partial.Declarations.Count - 1];
-							if (topScope.HeaderParser.CheckForBeginDestroy(partial, declaration))
-							{
-							}
-							// Add AddReferencedObjects, PostLoad, GetLifetimeReplicatedProps and friends later for feature complete
-						}
-					}
-				}
 				else
 				{
 					topScope.SkipDeclaration(token);
@@ -1498,24 +1482,6 @@ namespace EpicGames.UHT.Parsers
 
 			return false;
 		}
-
-		private bool CheckForBeginDestroy(UhtPartial partial, UhtDeclaration declaration)
-		{
-			using UhtTokenReplayReaderBorrower borrowedReader = new(HeaderFile, HeaderFile.Data.Memory, declaration.Tokens, UhtTokenType.EndOfDeclaration);
-			IUhtTokenReader replayReader = borrowedReader.Reader;
-
-			if (!replayReader.TryOptional("void") ||
-				!replayReader.TryOptional("BeginDestroy") ||
-				!replayReader.TryOptional('('))
-			{
-				return false;
-			}
-
-			partial.HasBeginDestroy = true;
-
-			return true;
-		}
-
 		#endregion
 	}
 }
diff --git a/Engine/Source/Programs/Shared/EpicGames.UHT/Parsers/UhtParsingScope.cs b/Engine/Source/Programs/Shared/EpicGames.UHT/Parsers/UhtParsingScope.cs
--- a/Engine/Source/Programs/Shared/EpicGames.UHT/Parsers/UhtParsingScope.cs
+++ b/Engine/Source/Programs/Shared/EpicGames.UHT/Parsers/UhtParsingScope.cs
@@ -73,7 +73,7 @@ namespace EpicGames.UHT.Parsers
 				UhtParsingScope? currentScope = this;
 				while (currentScope != null)
 				{
-					if (currentScope.ScopeType is UhtClass or UhtPartial)
+					if (currentScope.ScopeType is UhtClass)
 					{
 						return currentScope;
 					}
@@ -623,7 +623,7 @@ namespace EpicGames.UHT.Parsers
 			_compilerDirective = _scope.HeaderParser.GetCurrentCompositeCompilerDirective();
 			_function = null;
 			_flushed = false;
-			if (_scope.ScopeType is UhtClass or UhtPartial)
+			if (_scope.ScopeType is UhtClass)
 			{
 				_scope.TokenReader.EnableRecording();
 				_scope.TokenReader.RecordToken(ref initialToken);
@@ -641,7 +641,7 @@ namespace EpicGames.UHT.Parsers
 			_compilerDirective = _scope.HeaderParser.GetCurrentCompositeCompilerDirective();
 			_function = function;
 			_flushed = false;
-			if (_scope.ScopeType is UhtClass or UhtPartial)
+			if (_scope.ScopeType is UhtClass)
 			{
 				_scope.TokenReader.EnableRecording();
 			}
@@ -658,7 +658,7 @@ namespace EpicGames.UHT.Parsers
 		/// <summary>
 		/// Stop the recording
 		/// </summary>
-		/// <returns>True if the recorded content was added to a class or partial</returns>
+		/// <returns>True if the recorded content was added to a class</returns>
 		public bool Stop()
 		{
 			if (!_flushed)
@@ -670,12 +670,6 @@ namespace EpicGames.UHT.Parsers
 					_scope.TokenReader.DisableRecording();
 					return true;
 				}
-				else if (_scope.ScopeType is UhtPartial partialObj)
-				{
-					partialObj.AddDeclaration(_compilerDirective, _scope.TokenReader.RecordedTokens, _function);
-					_scope.TokenReader.DisableRecording();
-					return true;
-				}
 			}
 			return false;
 		}
diff --git a/Engine/Source/Programs/Shared/EpicGames.UHT/Parsers/UhtPartialParser.cs b/Engine/Source/Programs/Shared/EpicGames.UHT/Parsers/UhtPartialParser.cs
deleted file mode 100644
diff --git a/Engine/Source/Programs/Shared/EpicGames.UHT/Specifiers/UhtClassSpecifiers.cs b/Engine/Source/Programs/Shared/EpicGames.UHT/Specifiers/UhtClassSpecifiers.cs
--- a/Engine/Source/Programs/Shared/EpicGames.UHT/Specifiers/UhtClassSpecifiers.cs
+++ b/Engine/Source/Programs/Shared/EpicGames.UHT/Specifiers/UhtClassSpecifiers.cs
@@ -113,13 +113,6 @@ namespace EpicGames.UHT.Parsers
 			classObj.AddClassFlags(EClassFlags.Const);
 		}
 
-		[UhtSpecifier(Extends = UhtTableNames.Class, ValueType = UhtSpecifierValueType.None)]
-		private static void PartialSpecifier(UhtSpecifierContext specifierContext)
-		{
-			UhtClass classObj = (UhtClass)specifierContext.Type;
-			classObj.AddClassFlags(EClassFlags.Partial);
-		}
-
 		[UhtSpecifier(Extends = UhtTableNames.Class, ValueType = UhtSpecifierValueType.None)]
 		private static void PerObjectConfigSpecifier(UhtSpecifierContext specifierContext)
 		{
diff --git a/Engine/Source/Programs/Shared/EpicGames.UHT/Types/UhtFunction.cs b/Engine/Source/Programs/Shared/EpicGames.UHT/Types/UhtFunction.cs
--- a/Engine/Source/Programs/Shared/EpicGames.UHT/Types/UhtFunction.cs
+++ b/Engine/Source/Programs/Shared/EpicGames.UHT/Types/UhtFunction.cs
@@ -760,13 +760,7 @@ namespace EpicGames.UHT.Types
 			switch (FunctionType)
 			{
 				case UhtFunctionType.Function:
-
-					// Get the outer class - either directly or from partial's Owner class
-					UhtClass? outerClass = Outer as UhtClass;
-					if (outerClass == null && Outer is UhtPartial partial)
-					{
-						outerClass = partial.OwnerClass;
-					}
+					UhtClass? outerClass = (UhtClass?)Outer;
 
 					if (FunctionFlags.HasAnyFlags(EFunctionFlags.BlueprintPure) && outerClass != null &&
 						outerClass.ClassFlags.HasAnyFlags(EClassFlags.Interface) &&
@@ -1200,12 +1194,7 @@ namespace EpicGames.UHT.Types
 				return;
 			}
 
-			// Get the outer class - either directly or from partial's Owner class
-			UhtClass? outerClass = Outer as UhtClass;
-			if (outerClass == null && Outer is UhtPartial partial)
-			{
-				outerClass = partial.OwnerClass;
-			}
+			UhtClass? outerClass = (UhtClass?)Outer;
 
 			if (policy.FunctionToolTipsRequired)
 			{
diff --git a/Engine/Source/Programs/Shared/EpicGames.UHT/Types/UhtPartial.cs b/Engine/Source/Programs/Shared/EpicGames.UHT/Types/UhtPartial.cs
deleted file mode 100644
diff --git a/Engine/Source/Programs/Shared/EpicGames.UHT/Types/UhtProperty.cs b/Engine/Source/Programs/Shared/EpicGames.UHT/Types/UhtProperty.cs
--- a/Engine/Source/Programs/Shared/EpicGames.UHT/Types/UhtProperty.cs
+++ b/Engine/Source/Programs/Shared/EpicGames.UHT/Types/UhtProperty.cs
@@ -1890,10 +1890,6 @@ namespace EpicGames.UHT.Types
 				outerFunc(builder);
 				builder.Append(", ");
 			}
-			else if (property.Outer is UhtPartial partial)
-			{
-				builder.Append($"&{context.GetSingletonName(partial.OwnerClass, UhtSingletonType.ConstInit)}, ");
-			}
 			else if (property.Outer is UhtObject obj)
 			{
 				builder.Append($"&{context.GetSingletonName(obj, UhtSingletonType.ConstInit)}, ");
diff --git a/Engine/Source/Programs/Shared/EpicGames.UHT/Utils/UhtNames.cs b/Engine/Source/Programs/Shared/EpicGames.UHT/Utils/UhtNames.cs
--- a/Engine/Source/Programs/Shared/EpicGames.UHT/Utils/UhtNames.cs
+++ b/Engine/Source/Programs/Shared/EpicGames.UHT/Utils/UhtNames.cs
@@ -109,8 +109,6 @@ namespace EpicGames.UHT.Utils
 		public static readonly UhtCppIdentifier StaticsRigVMStructsId = RigVMStructsId.MakeStatics();
 		public static readonly UhtCppIdentifier VerseFuncsId = new("VerseFuncs");
 		public static readonly UhtCppIdentifier StaticsVerseFuncsId = VerseFuncsId.MakeStatics();
-		public static readonly UhtCppIdentifier PartialInfoId = new("PartialInfo");
-		public static readonly UhtCppIdentifier StaticsPartialInfoId = PartialInfoId.MakeStatics();
 
 		// Helper methods to generate cpp identifiers
 		public static UhtCppIdentifier GetPropertyIdentifier(UhtProperty? property)
diff --git a/Engine/Source/Programs/Shared/EpicGames.UHT/Utils/UhtSession.cs b/Engine/Source/Programs/Shared/EpicGames.UHT/Utils/UhtSession.cs
--- a/Engine/Source/Programs/Shared/EpicGames.UHT/Utils/UhtSession.cs
+++ b/Engine/Source/Programs/Shared/EpicGames.UHT/Utils/UhtSession.cs
@@ -2289,7 +2289,7 @@ namespace EpicGames.UHT.Utils
 				}
 
 				// Top-level objects, excluding those with alternate objects (native inteface objects)
-				Dictionary<UhtObject, int> linkableObjects = module.Packages.SelectMany(p => p.Children.OfType<UhtObject>().Where(o => o.AlternateObject is null && o is not UhtPartial))
+				Dictionary<UhtObject, int> linkableObjects = module.Packages.SelectMany(p => p.Children.OfType<UhtObject>().Where(o => o.AlternateObject is null))
 					// Objects whose outer is not a package but which need to be registered
 					.Concat(module.Headers.SelectMany(h => h.References.Singletons).Where(o => o.Package != o.Outer))
 					.OrderBy(o => (o.LinkOrder, o.SourceName, o.Outer!.SourceName))
diff --git a/Engine/Source/Programs/Shared/EpicGames.UHT/Utils/UhtTables.cs b/Engine/Source/Programs/Shared/EpicGames.UHT/Utils/UhtTables.cs
--- a/Engine/Source/Programs/Shared/EpicGames.UHT/Utils/UhtTables.cs
+++ b/Engine/Source/Programs/Shared/EpicGames.UHT/Utils/UhtTables.cs
@@ -94,11 +94,6 @@ namespace EpicGames.UHT.Utils
 		/// Table for VModules
 		/// </summary>
 		public const string VModule = "VModule";
-
-		/// <summary>
-		/// Table for UPARTIAL
-		/// </summary>
-		public const string Partial = "Partial";
 	}
 
 	/// <summary>
@@ -394,7 +389,6 @@ namespace EpicGames.UHT.Utils
 			CreateTables(tables, UhtTableNames.Function, "Function", EUhtCreateTablesFlags.Specifiers | EUhtCreateTablesFlags.Keyword, UhtTableNames.Struct, false);
 			CreateTables(tables, UhtTableNames.ScriptStruct, "Struct", EUhtCreateTablesFlags.Specifiers | EUhtCreateTablesFlags.Keyword, UhtTableNames.Struct, false);
 			CreateTables(tables, UhtTableNames.VModule, "VModule", EUhtCreateTablesFlags.Specifiers | EUhtCreateTablesFlags.Keyword, UhtTableNames.Field, false);
-			CreateTables(tables, UhtTableNames.Partial, "Partial", EUhtCreateTablesFlags.Specifiers | EUhtCreateTablesFlags.Keyword, UhtTableNames.ScriptStruct, false);
 		}
 
 		/// <summary>
diff --git a/Engine/Source/Programs/UnrealBuildTool/System/SourceFileMetadataCache.cs b/Engine/Source/Programs/UnrealBuildTool/System/SourceFileMetadataCache.cs
--- a/Engine/Source/Programs/UnrealBuildTool/System/SourceFileMetadataCache.cs
+++ b/Engine/Source/Programs/UnrealBuildTool/System/SourceFileMetadataCache.cs
@@ -134,7 +134,7 @@ namespace UnrealBuildTool
 		/// <summary>
 		/// Regex that matches C++ code with UObject declarations which we will need to generated code for.
 		/// </summary>
-		static readonly Regex ReflectionMarkupRegex = new Regex("^\\s*(UCLASS|USTRUCT|UENUM|UINTERFACE|UDELEGATE|UPARTIAL|VMODULE)\\b", RegexOptions.Compiled);
+		static readonly Regex ReflectionMarkupRegex = new Regex("^\\s*(UCLASS|USTRUCT|UENUM|UINTERFACE|UDELEGATE|VMODULE)\\b", RegexOptions.Compiled);
 
 		/// <summary>
 		/// Regex that matches #include UE_INLINE_GENERATED_CPP_BY_NAME(****) statements.
diff --git a/Engine/Source/Runtime/CoreUObject/Private/Serialization/UnversionedPropertySerialization.cpp b/Engine/Source/Runtime/CoreUObject/Private/Serialization/UnversionedPropertySerialization.cpp
--- a/Engine/Source/Runtime/CoreUObject/Private/Serialization/UnversionedPropertySerialization.cpp
+++ b/Engine/Source/Runtime/CoreUObject/Private/Serialization/UnversionedPropertySerialization.cpp
@@ -37,27 +37,9 @@ struct FDefaultStruct
 {
 	const uint8* Data;
 #if CACHE_UNVERSIONED_PROPERTY_SCHEMA
-	int32 StructStart;
-	int32 StructEnd;
+	uint32 StructSize;
 
-	FDefaultStruct(const uint8* InData, const UStruct* InStruct)
-		: Data(InData)
-	{
-		StructStart = 0;
-		StructEnd = 0;
-		if (InData)
-		{
-			if (const UClass* StructAsClass = Cast<const UClass>(InStruct))
-			{
-				StructStart = StructAsClass->GetPropertiesStartOffset();
-				StructEnd = InStruct->GetPropertiesSize() + StructStart;
-			}
-			else
-			{
-				StructEnd = InStruct->GetPropertiesSize();
-			}
-		}
-	}
+	FDefaultStruct(const uint8* InData, const UStruct* InStruct) : Data(InData), StructSize(InData ? InStruct->GetPropertiesSize() : 0)  {}
 #else
 	const UStruct* Struct;
 
@@ -115,7 +97,7 @@ public:
 	const void* GetDefaultValue(FDefaultStruct Defaults) const
 	{
 #if CACHE_UNVERSIONED_PROPERTY_SCHEMA
-		return Offset >= Defaults.StructStart && Offset < Defaults.StructEnd ? Defaults.Data + Offset : nullptr;
+		return Offset < Defaults.StructSize ? Defaults.Data + Offset : nullptr;
 #else
 		return Property->ContainerPtrToValuePtrForDefaults<uint8>(Defaults.Struct, Defaults.Data);
 #endif
@@ -341,7 +323,7 @@ private:
 
 	FProperty* Property;
 #if CACHE_UNVERSIONED_PROPERTY_SCHEMA
-	int32 Offset;
+	uint32 Offset;
 	bool bSerializeAsInteger;
 	bool bIsOptional : 1;
 #if WITH_EDITORONLY_DATA
diff --git a/Engine/Source/Runtime/CoreUObject/Private/UObject/Class.cpp b/Engine/Source/Runtime/CoreUObject/Private/UObject/Class.cpp
--- a/Engine/Source/Runtime/CoreUObject/Private/UObject/Class.cpp
+++ b/Engine/Source/Runtime/CoreUObject/Private/UObject/Class.cpp
@@ -38,7 +38,6 @@
 #include "UObject/LinkerSave.h"
 #include "UObject/InstanceDataObjectUtils.h"
 #include "UObject/Interface.h"
-#include "UObject/UObjectPartials.h"
 #include "UObject/LinkerPlaceholderClass.h"
 #include "UObject/LinkerPlaceholderFunction.h"
 #include "UObject/PropertyOptional.h"
@@ -6409,7 +6408,7 @@ PRAGMA_ENABLE_DEPRECATION_WARNINGS
 	PropertyLink = nullptr;
 	DestructorLink = nullptr;
 	CppClassStaticFunctions.Reset();
-	Partials = nullptr;
+
 	ScriptAndPropertyObjectReferences.Empty();
 	DeleteUnresolvedScriptProperties();
 
diff --git a/Engine/Source/Runtime/CoreUObject/Private/UObject/CompiledInUObjectInit.cpp b/Engine/Source/Runtime/CoreUObject/Private/UObject/CompiledInUObjectInit.cpp
--- a/Engine/Source/Runtime/CoreUObject/Private/UObject/CompiledInUObjectInit.cpp
+++ b/Engine/Source/Runtime/CoreUObject/Private/UObject/CompiledInUObjectInit.cpp
@@ -5,7 +5,6 @@
 #include "UObject/PropertyOptional.h"
 #include "UObject/UObjectBase.h"
 #include "UObject/UObjectGlobals.h"
-#include "UObject/UObjectPartials.h"
 
 #if UE_WITH_CONSTINIT_UOBJECT
 
@@ -354,29 +353,6 @@ void FCompiledInObjectRegistry::AddAndHashObjects()
 	SCOPED_BOOT_TIMING("FDeferredRegistry::AddAndHashPendingRegistrants");
 	FIterationMarker Stop = AdvanceConstruction(EConstructionStep::AddObjects);
 
-	{
-		TArray<UClass*> PartialClasses;
-		Iterate(Stop, 
-			[&](const FRegisterCompiledInObjects& Generated)
-			{
-				for (UE::CoreUObject::Private::FPartialClass* Partial : Generated.Objects.Partials)
-				{
-					UClass* PartialClass = Partial->Class;
-					if (!PartialClass->Partials)
-					{
-						PartialClasses.Add(PartialClass);
-					}
-					Partial->LinkToClass(*PartialClass);
-				}
-			}
-		);
-		Algo::Sort(PartialClasses, [](UClass* A, UClass* B) { return B->IsChildOf(A); });
-		for (UClass* PartialClass : PartialClasses)
-		{
-			UE::CoreUObject::Private::ConstructUClassPartials(*PartialClass);
-		}
-	}
-
 	// Fixup object pointers first - in order to hash a class we will need e.g. the super pointer
 	// This is only necessary for generated objects and not intrinsic ones
 #if !IS_MONOLITHIC
@@ -392,7 +368,6 @@ void FCompiledInObjectRegistry::AddAndHashObjects()
 			{
 				if (Class)
 				{ 
-					UE::CoreUObject::Private::ConstructUClassPartials(*Class);
 					Class->LinkCompiledInPointerFields(Linker);
 					for (UField* Field = Class->Children; Field; Field = Field->Next)
 					{
@@ -445,15 +420,10 @@ void FCompiledInObjectRegistry::AddAndHashObjects()
 				}
 			}
 		}
-	});
 
-	IterateObjects(Stop, [](UObject* Obj)
-	{
 		// TODO: General purpose interface function for initializing compiled in FNames
 		if (UClass* Class = Cast<UClass>(Obj))
 		{
-			UE::CoreUObject::Private::ConstructUClassPartials(*Class); // We need to call this again because there might be subclasses to partial classes that has no partials
-
 			Class->ClassConfigName = FName(Class->CompiledInClassConfigName);
 
 			// Create TArray of native functions from compiled-in data to maintain API
@@ -462,16 +432,6 @@ void FCompiledInObjectRegistry::AddAndHashObjects()
 			{
 				TConstArrayView<UE::CodeGen::FClassNativeFunction> CompiledInNativeFunctions = Class->CompiledInNativeFunctions;
 				new (&Class->NativeFunctionLookupTable) TArray<FNativeFunctionLookup>(CompiledInNativeFunctions);
-
-				// Partials have functions too that needs to be added to the NativeFunctionLookupTable
-				for (UE::CoreUObject::Private::FPartialClass* It = Class->Partials; It ; It = It->NextPartial)
-				{
-					for (UField* Child = It->FirstChild; Child != It->EndChild; Child = Child->Next)
-					{
-						UFunction* Func = (UFunction*)Child;
-						Class->AddNativeFunction(*Func->GetName(), Func->GetNativeFunc());
-					}
-				}
 			}
 		}
 		else if (USparseDelegateFunction* SparseDelegate = Cast<USparseDelegateFunction>(Obj))
diff --git a/Engine/Source/Runtime/CoreUObject/Private/UObject/DeferredRegistry.h b/Engine/Source/Runtime/CoreUObject/Private/UObject/DeferredRegistry.h
--- a/Engine/Source/Runtime/CoreUObject/Private/UObject/DeferredRegistry.h
+++ b/Engine/Source/Runtime/CoreUObject/Private/UObject/DeferredRegistry.h
@@ -32,11 +32,6 @@ inline bool operator != (const FPackageReloadVersionInfo& lhs, const FPackageRel
 	return lhs.BodyHash != rhs.BodyHash || lhs.DeclarationsHash != rhs.DeclarationsHash;
 }
 
-inline bool operator != (const FPartialReloadVersionInfo& lhs, const FPartialReloadVersionInfo& rhs)
-{
-	return lhs.Size != rhs.Size || lhs.Hash != rhs.Hash;
-}
-
 void ReloadProcessObject(UScriptStruct* Object, const TCHAR* RenamePrefix);
 void ReloadProcessObject(UEnum* Object, const TCHAR* RenamePrefix);
 void ReloadProcessObject(UClass* Object, const TCHAR* RenamePrefix);
@@ -369,6 +364,5 @@ using FClassDeferredRegistry = TDeferredRegistry<FClassRegistrationInfo>;
 using FEnumDeferredRegistry = TDeferredRegistry<FEnumRegistrationInfo>;
 using FStructDeferredRegistry = TDeferredRegistry<FStructRegistrationInfo>;
 using FPackageDeferredRegistry = TDeferredRegistry<FPackageRegistrationInfo>;
-using FPartialDeferredRegistry = TDeferredRegistry<FPartialRegistrationInfo>;
 
 #endif // !UE_WITH_CONSTINIT_UOBJECT
diff --git a/Engine/Source/Runtime/CoreUObject/Private/UObject/GarbageCollection.cpp b/Engine/Source/Runtime/CoreUObject/Private/UObject/GarbageCollection.cpp
--- a/Engine/Source/Runtime/CoreUObject/Private/UObject/GarbageCollection.cpp
+++ b/Engine/Source/Runtime/CoreUObject/Private/UObject/GarbageCollection.cpp
@@ -26,7 +26,6 @@
 #include "UObject/ObjectPtr.h"
 #include "UObject/Class.h"
 #include "UObject/UObjectIterator.h"
-#include "UObject/UObjectPartials.h"
 #include "UObject/UnrealType.h"
 #include "UObject/LinkerLoad.h"
 #include "UObject/GCObject.h"
@@ -951,17 +950,8 @@ public:
 				if (Object)
 #endif // UE_WITH_CONSTINIT_UOBJECT
 				{
-					int32 OffsetToObject = 0;
-
-					if (Object->HasAnyFlags(RF_HasPartials))
-					{
-						UClass* Class = Object->GetClass();
-						UE::CoreUObject::Private::DestructPartials(*Class, Object);
-						OffsetToObject = -Class->GetPropertiesStartOffset();
-					}
-
 					Object->~UObject();
-					GUObjectAllocator.FreeUObject(Object, OffsetToObject);
+					GUObjectAllocator.FreeUObject(Object);
 					UnreachableObject.Object = nullptr;
 				}
 
@@ -3957,7 +3947,7 @@ void TFastReferenceCollector<ProcessorType, CollectorType>::ProcessStructs(Dispa
 				typename DispatcherType::SchemaStackScopeType SchemaStack(Dispatcher.Context, AoS.Schema);
 				do
 				{
-					Private::VisitMembers(Dispatcher, AoS.Schema, StructIt, 0); 
+					Private::VisitMembers(Dispatcher, AoS.Schema, StructIt); 
 					StructIt += AoS.Stride;
 				}
 				while (StructIt != StructEnd);
@@ -7009,16 +6999,13 @@ void UClass::AssembleReferenceTokenStreamInternal(bool bForce)
 			ClassFlags &= ~CLASS_TokenStreamAssembled;
 		}
 
-		// We need to make sure all offsets to properties are positive because of schema format.
-		int32 StartOffset = -GetPropertiesStartOffset();
-
 		FSchemaBuilder Schema(/* don't store sizeof(class) to enable super class schema reuse */ 0);
 		FSchemaView SuperSchema;
 		if (UClass* SuperClass = GetSuperClass())
 		{
 			SuperClass->AssembleReferenceTokenStreamInternal();	
 			SuperSchema = SuperClass->ReferenceSchema.Get();
-			Schema.Append(SuperSchema, StartOffset - (-SuperClass->GetPropertiesStartOffset()));
+			Schema.Append(SuperSchema);
 		}
 		const int32 NumSuperMembers = Schema.NumMembers();
 
@@ -7031,13 +7018,13 @@ void UClass::AssembleReferenceTokenStreamInternal(bool bForce)
 			{
 				FProperty* Property = *It;
 				FPropertyStackScope PropertyScope(DebugPath, Property);
-				Property->EmitReferenceInfo(Schema, StartOffset, EncounteredStructProps, DebugPath);
+				Property->EmitReferenceInfo(Schema, 0, EncounteredStructProps, DebugPath);
 			}
 		}
 		
 		if (ClassFlags & CLASS_Intrinsic)
 		{
-			Schema.Append(UE::GC::GetIntrinsicSchema(this), 0);
+			Schema.Append(UE::GC::GetIntrinsicSchema(this));
 		}
 		
 		// Make sure all Blueprint properties are marked as non-native
diff --git a/Engine/Source/Runtime/CoreUObject/Private/UObject/GarbageCollectionSchema.cpp b/Engine/Source/Runtime/CoreUObject/Private/UObject/GarbageCollectionSchema.cpp
--- a/Engine/Source/Runtime/CoreUObject/Private/UObject/GarbageCollectionSchema.cpp
+++ b/Engine/Source/Runtime/CoreUObject/Private/UObject/GarbageCollectionSchema.cpp
@@ -328,7 +328,7 @@ static void CopyDebugNames(const FName* In, TArrayView<FMemberDeclaration> Out)
 	}
 }
 
-void FSchemaBuilder::Append(FSchemaView SuperSchema, int32 OffsetForProperties)
+void FSchemaBuilder::Append(FSchemaView SuperSchema)
 {
 	if (SuperSchema.IsEmpty())
 	{
@@ -356,7 +356,7 @@ void FSchemaBuilder::Append(FSchemaView SuperSchema, int32 OffsetForProperties)
 			}
 			else
 			{
-				Members.Add({ FName(), OffsetForProperties + (JumpedWords + Member.WordOffset) * 8, Member.Type, ExtraWord});
+				Members.Add({ FName(), (JumpedWords + Member.WordOffset) * 8, Member.Type, ExtraWord});
 				TryAddSchemaReference(Members.Last());
 			}
 		}
diff --git a/Engine/Source/Runtime/CoreUObject/Private/UObject/Obj.cpp b/Engine/Source/Runtime/CoreUObject/Private/UObject/Obj.cpp
--- a/Engine/Source/Runtime/CoreUObject/Private/UObject/Obj.cpp
+++ b/Engine/Source/Runtime/CoreUObject/Private/UObject/Obj.cpp
@@ -28,7 +28,6 @@
 #include "Serialization/ArchiveUObject.h"
 #include "UObject/GarbageCollectionHistory.h"
 #include "UObject/Class.h"
-#include "UObject/UObjectPartials.h"
 #include "UObject/EnumProperty.h"
 #include "UObject/UObjectIterator.h"
 #include "UObject/Package.h"
@@ -1110,8 +1109,6 @@ void UObject::BeginDestroy()
 			);
 	}
 
-	UE::CoreUObject::Private::BeginDestroyPartials(*GetClass(), this);
-
 	// Remove from linker's export table.
 	SetLinker( NULL, INDEX_NONE );
 
diff --git a/Engine/Source/Runtime/CoreUObject/Private/UObject/UObjectAllocator.cpp b/Engine/Source/Runtime/CoreUObject/Private/UObject/UObjectAllocator.cpp
--- a/Engine/Source/Runtime/CoreUObject/Private/UObject/UObjectAllocator.cpp
+++ b/Engine/Source/Runtime/CoreUObject/Private/UObject/UObjectAllocator.cpp
@@ -31,7 +31,7 @@ void FUObjectAllocator::DisablePersistentAllocator()
  * @return newly allocated UObjectBase (not really a UObjectBase yet, no constructor-like thing has been called).
  */
 UE_AUTORTFM_ALWAYS_OPEN
-UObjectBase* FUObjectAllocator::AllocateUObject(int32 Size, int32 Alignment, bool bAllowPermanent, int32 OffsetToObject)
+UObjectBase* FUObjectAllocator::AllocateUObject(int32 Size, int32 Alignment, bool bAllowPermanent)
 {
 	void* Result = nullptr;
 
@@ -46,7 +46,7 @@ UObjectBase* FUObjectAllocator::AllocateUObject(int32 Size, int32 Alignment, boo
 		Result = FMemory::Malloc(Size, Alignment);
 	}
 
-	return (UObjectBase*)(((uint8*)Result) + OffsetToObject);
+	return (UObjectBase*)Result;
 }
 
 /**
@@ -54,13 +54,13 @@ UObjectBase* FUObjectAllocator::AllocateUObject(int32 Size, int32 Alignment, boo
  *
  * @param Object object to free
  */
-void FUObjectAllocator::FreeUObject(UObjectBase *Object, int32 OffsetToObject) const
+void FUObjectAllocator::FreeUObject(UObjectBase *Object) const
 {
 	check(Object);
 	// Only free memory if it was allocated directly from allocator and not from permanent object pool.
 	if (FPermanentObjectPoolExtents().Contains(Object) == false)
 	{
-		FMemory::Free(((uint8*)Object) - OffsetToObject);
+		FMemory::Free(Object);
 	}
 	// We only destroy objects residing in permanent object pool during the exit purge.
 	else
diff --git a/Engine/Source/Runtime/CoreUObject/Private/UObject/UObjectBase.cpp b/Engine/Source/Runtime/CoreUObject/Private/UObject/UObjectBase.cpp
--- a/Engine/Source/Runtime/CoreUObject/Private/UObject/UObjectBase.cpp
+++ b/Engine/Source/Runtime/CoreUObject/Private/UObject/UObjectBase.cpp
@@ -20,7 +20,6 @@
 #include "UObject/DeferredRegistry.h"
 #include "UObject/ObjectHandlePrivate.h"
 #include "UObject/UObjectIterator.h"
-#include "UObject/UObjectPartials.h"
 #include "UObject/UObjectStats.h"
 #include "UObject/UObjectPrivate.h"
 #include "UObject/Package.h"
@@ -784,7 +783,7 @@ void RegisterCompiledInInfo(UPackage* (*InOuterRegister)(ETypeConstructPhase), c
 }
 
 // Multiple registrations
-void RegisterCompiledInInfo(const TCHAR* PackageName, const FClassRegisterCompiledInInfo* ClassInfo, size_t NumClassInfo, const FStructRegisterCompiledInInfo* StructInfo, size_t NumStructInfo, const FEnumRegisterCompiledInInfo* EnumInfo, size_t NumEnumInfo, const FPartialRegisterCompiledInInfo *PartialsInfo, size_t NumPartialsInfo)
+void RegisterCompiledInInfo(const TCHAR* PackageName, const FClassRegisterCompiledInInfo* ClassInfo, size_t NumClassInfo, const FStructRegisterCompiledInInfo* StructInfo, size_t NumStructInfo, const FEnumRegisterCompiledInInfo* EnumInfo, size_t NumEnumInfo)
 {
 	LLM_SCOPE(ELLMTag::UObject);
 	LLM_SCOPE_BYTAG(UObject_UObjectBase);
@@ -810,12 +809,6 @@ void RegisterCompiledInInfo(const TCHAR* PackageName, const FClassRegisterCompil
 		const FEnumRegisterCompiledInInfo& Info = EnumInfo[Index];
 		RegisterCompiledInInfo(Info.OuterRegister, PackageName, Info.Name, *Info.Info, Info.VersionInfo);
 	}
-
-	for (size_t Index = 0; Index < NumPartialsInfo; ++Index)
-	{
-		const FPartialRegisterCompiledInInfo& Info = PartialsInfo[Index];
-		FPartialDeferredRegistry::Get().AddRegistration(Info.OuterRegister, PackageName, Info.Name, *Info.Info, Info.VersionInfo);
-	}
 }
 
 /** Register all loaded classes */
@@ -1097,7 +1090,6 @@ void ProcessNewlyLoadedUObjects(FName InModuleName, bool bCanProcessNewlyLoadedO
 	FClassDeferredRegistry& ClassRegistry = FClassDeferredRegistry::Get();
 	FStructDeferredRegistry& StructRegistry = FStructDeferredRegistry::Get();
 	FEnumDeferredRegistry& EnumRegistry = FEnumDeferredRegistry::Get();
-	FPartialDeferredRegistry& PartialRegistry = FPartialDeferredRegistry::Get();
 
 	PackageRegistry.ProcessChangedObjects(true);
 	StructRegistry.ProcessChangedObjects();
@@ -1110,8 +1102,7 @@ void ProcessNewlyLoadedUObjects(FName InModuleName, bool bCanProcessNewlyLoadedO
 	while (GFirstPendingRegistrant ||
 		ClassRegistry.HasPendingRegistrations() ||
 		StructRegistry.HasPendingRegistrations() ||
-		EnumRegistry.HasPendingRegistrations() ||
-		PartialRegistry.HasPendingRegistrations())
+		EnumRegistry.HasPendingRegistrations())
 	{
 		if (bNewUObjects)
 		{
@@ -1123,16 +1114,6 @@ void ProcessNewlyLoadedUObjects(FName InModuleName, bool bCanProcessNewlyLoadedO
 		}
 
 		bNewUObjects = true;
-
-		// Initializing partials first since they need to be linked to class before outer construction
-		for (const FPartialDeferredRegistry::FRegistrant& Registrant : PartialRegistry.GetRegistrations())
-		{
-			UE::CoreUObject::Private::FPartialClass* Params = Registrant.RegisterFn(ETypeConstructPhase::Outer);
-			Params->LinkToClass(*Params->Class(ETypeConstructPhase::Inner));
-		}
-		PartialRegistry.EmptyRegistrations();
-
-
 		UObjectProcessRegistrants();
 		UObjectLoadAllCompiledInStructs();
 
diff --git a/Engine/Source/Runtime/CoreUObject/Private/UObject/UObjectConstructInternal.h b/Engine/Source/Runtime/CoreUObject/Private/UObject/UObjectConstructInternal.h
--- a/Engine/Source/Runtime/CoreUObject/Private/UObject/UObjectConstructInternal.h
+++ b/Engine/Source/Runtime/CoreUObject/Private/UObject/UObjectConstructInternal.h
@@ -252,8 +252,6 @@ void ConstructUClassHelper(UClass*& OutClass, const ClassParams& Params, PostNew
 
 	ConstructFProperties(NewClass, Params.PropertyArray, Params.NumProperties);
 
-	UE::CoreUObject::Private::ConstructUClassPartials(*NewClass);
-
 	if (Params.ClassConfigNameUTF8)
 	{
 		NewClass->ClassConfigName = FName(UTF8_TO_TCHAR(Params.ClassConfigNameUTF8));
diff --git a/Engine/Source/Runtime/CoreUObject/Private/UObject/UObjectGlobals.cpp b/Engine/Source/Runtime/CoreUObject/Private/UObject/UObjectGlobals.cpp
--- a/Engine/Source/Runtime/CoreUObject/Private/UObject/UObjectGlobals.cpp
+++ b/Engine/Source/Runtime/CoreUObject/Private/UObject/UObjectGlobals.cpp
@@ -33,7 +33,6 @@
 #include "UObject/GarbageCollection.h"
 #include "UObject/GCScopeLock.h"
 #include "UObject/Class.h"
-#include "UObject/UObjectPartials.h"
 #include "UObject/CoreRedirects.h"
 #include "UObject/FastReferenceCollector.h"
 #include "UObject/ObjectInstancingGraph.h"
@@ -3467,13 +3466,6 @@ UObject* StaticAllocateObject
 	int32 TotalSize = InClass->GetPropertiesSize();
 	checkSlow(TotalSize);
 
-	int32 OffsetToObject = -InClass->GetPropertiesStartOffset(); // Start offset for properties can be negative so we want to move object forward in allocation
-
-	if (InClass->HasAnyClassFlags(CLASS_Partial))
-	{
-		InFlags |= RF_HasPartials;
-	}
-
 	int32 OldIndex = -1;
 	int32 OldSerialNumber = 0;
 #if UE_WITH_REMOTE_OBJECT_HANDLE
@@ -3485,7 +3477,7 @@ UObject* StaticAllocateObject
 	if( Obj == nullptr )
 	{	
 		int32 Alignment	= FMath::Max( 4, InClass->GetMinAlignment() );
-		Obj = (UObject *)GUObjectAllocator.AllocateUObject(TotalSize, Alignment, GIsInitialLoad, OffsetToObject);
+		Obj = (UObject *)GUObjectAllocator.AllocateUObject(TotalSize,Alignment,GIsInitialLoad);
 	}
 	else
 	{
@@ -3676,7 +3668,7 @@ UObject* StaticAllocateObject
 
 		UE_AUTORTFM_OPEN_NO_VALIDATION
 		{
-			FMemory::Memzero((uint8 *)Obj - OffsetToObject, TotalSize);
+			FMemory::Memzero((void *)Obj, TotalSize);
 			new ((void *)Obj) UObjectBase(const_cast<UClass*>(InClass), InFlags|RF_NeedInitialization, InternalSetFlags, InOuter, InName, OldIndex, OldSerialNumber == 0 ? SerialNumber : OldSerialNumber, 
 #if UE_WITH_REMOTE_OBJECT_HANDLE
 				OldRemoteId.IsValid() ? OldRemoteId :
@@ -4141,8 +4133,6 @@ void FObjectInitializer::PostConstructInit()
 	}
 #endif
 
-	UE::CoreUObject::Private::ConstructPartials(*Class, Obj);
-
 	if (bShouldInitializePropsFromArchetype)
 	{
 		UClass* BaseClass = (bIsCDO && !GIsDuplicatingClassForReinstancing) ? SuperClass : Class;
@@ -6097,7 +6087,7 @@ namespace UECodeGen_Private
 		return NewProp;
 	}
 
-	FProperty* ConstructFProperty(FFieldVariant Outer, const FPropertyParamsBase* const*& PropertyArray, int32& NumProperties)
+	void ConstructFProperty(FFieldVariant Outer, const FPropertyParamsBase* const*& PropertyArray, int32& NumProperties)
 	{
 		const FPropertyParamsBase* PropBase = *--PropertyArray;
 
@@ -6392,7 +6382,6 @@ namespace UECodeGen_Private
 		{
 			ConstructFProperty(NewProp, PropertyArray, NumProperties);
 		}
-		return NewProp;
 	}
 
 	void ConstructFProperties(UObject* Outer, const FPropertyParamsBase* const* PropertyArray, int32 NumProperties)
diff --git a/Engine/Source/Runtime/CoreUObject/Private/UObject/UObjectPartials.cpp b/Engine/Source/Runtime/CoreUObject/Private/UObject/UObjectPartials.cpp
deleted file mode 100644
diff --git a/Engine/Source/Runtime/CoreUObject/Public/UObject/Class.h b/Engine/Source/Runtime/CoreUObject/Public/UObject/Class.h
--- a/Engine/Source/Runtime/CoreUObject/Public/UObject/Class.h
+++ b/Engine/Source/Runtime/CoreUObject/Public/UObject/Class.h
@@ -132,8 +132,6 @@ struct FUObjectSerializeContext;
 template <typename FuncType> class TFunctionRef;
 enum class EPropertyObjectReferenceType : uint32;
 
-namespace UE::CoreUObject::Private { void ConstructUClassPartials(UClass&); }
-
 COREUOBJECT_API DECLARE_LOG_CATEGORY_EXTERN(LogClass, Log, All);
 COREUOBJECT_API DECLARE_LOG_CATEGORY_EXTERN(LogScriptSerialization, Log, All);
 
@@ -3891,12 +3889,7 @@ protected:
 
 	/** TRUE if instances of this class contain non-serialized references that require construction on load via instancing. */
 	bool bNeedsPostLoadSubobjectInstancing = false;
-
-private:
-	/** When using partials they are allocated in front of UObject, and then this value will be negative. Note that PropertiesSize includes the negative part too */
-	int32 PropertiesStartOffset = 0;
-
-	friend void UE::CoreUObject::Private::ConstructUClassPartials(UClass&);
+	
 public:
 
 	/** Class flags; See EClassFlags for more information */
@@ -4120,9 +4113,6 @@ public:
 #endif // UE_WITH_CONSTINIT_UOBJECT
 	};
 
-	/** Linked list of partials extending this class */
-	UE::CoreUObject::Private::FPartialClass* Partials = nullptr;
-
 public:
 	// Constructors
 	COREUOBJECT_API UClass(const FObjectInitializer& ObjectInitializer = FObjectInitializer::Get());
@@ -4314,12 +4304,6 @@ public:
 	UE_INTERNAL COREUOBJECT_API virtual void SerializePropertyValueFlags(EPropertyValueFlags Flags, void* Data, FStructuredArchiveRecord Record, FArchiveFieldName Name) const override;
 #endif // WITH_EDITORONLY_DATA
 
-	/** Offset from UObject pointer to where properties start. This method will return a negative value if partials are present */
-	UE_FORCEINLINE_HINT int32 GetPropertiesStartOffset() const
-	{
-		return PropertiesStartOffset;
-	}
-
 	/**
 	 * Returns class path name as a package + class FName pair
 	 */
diff --git a/Engine/Source/Runtime/CoreUObject/Public/UObject/FastReferenceCollector.h b/Engine/Source/Runtime/CoreUObject/Public/UObject/FastReferenceCollector.h
--- a/Engine/Source/Runtime/CoreUObject/Public/UObject/FastReferenceCollector.h
+++ b/Engine/Source/Runtime/CoreUObject/Public/UObject/FastReferenceCollector.h
@@ -684,13 +684,12 @@ FORCEINLINE EOrigin GetSchemaOrigin(FSchemaView Schema, UObject* Instance)
 #endif // UE_WITH_REMOTE_OBJECT_HANDLE
 
 template<class DispatcherType, typename ObjectType>
-AUTORTFM_INFER FORCEINLINE_DEBUGGABLE void VisitMembers(DispatcherType& Dispatcher, FSchemaView Schema, ObjectType* Instance, int32 OffsetToPropertiesStart)
+AUTORTFM_INFER FORCEINLINE_DEBUGGABLE void VisitMembers(DispatcherType& Dispatcher, FSchemaView Schema, ObjectType* Instance)
 {
 	check(!Schema.IsEmpty());
 
 	const EOrigin Origin = GetSchemaOrigin(Schema, Instance);
-	uint8* InstanceBegin = (uint8*)Instance + OffsetToPropertiesStart;
-	uint64* InstanceCursor = (uint64*)InstanceBegin;	// Advanced via Jump to reach far members
+	uint64* InstanceCursor = (uint64*)Instance;	// Advanced via Jump to reach far members
 	uint32 DebugIdx = 0;
 	for (const FMemberWord* WordIt = Schema.GetWords(); true; ++WordIt)
 	{
@@ -754,7 +753,7 @@ template<class DispatcherType>
 void VisitNestedStructMembers(DispatcherType& Dispatcher, FSchemaView Schema, uint8* Instance)
 {
 	static_assert(!DispatcherType::bBatching);
-	VisitMembers(Dispatcher, Schema, Instance, 0);
+	VisitMembers(Dispatcher, Schema, Instance);
 }
 
 } // namespace Private
@@ -1027,10 +1026,9 @@ private:
 
 			if (!Schema.IsEmpty())
 			{
-				int32 OffsetToProperty = Class->GetPropertiesStartOffset();
 				typename DispatcherType::SchemaStackScopeType SchemaStack(Dispatcher.Context, Schema);
 				Processor.BeginTimingObject(CurrentObject);
-				Private::VisitMembers(Dispatcher, Schema, CurrentObject, OffsetToProperty);
+				Private::VisitMembers(Dispatcher, Schema, CurrentObject);
 				Processor.UpdateDetailedStats(CurrentObject);
 			}
 		}
@@ -1047,7 +1045,7 @@ private:
 			FSchemaView Schema = VerseStruct->ReferenceSchema.Get();
 			check(!Schema.IsEmpty()); // Native structs with empty schemas should not be in the list to begin with
 			typename DispatcherType::SchemaStackScopeType SchemaStack(Dispatcher.Context, Schema);
-			Private::VisitMembers(Dispatcher, Schema, (uint8*)NativeStruct->GetStruct(), 0);
+			Private::VisitMembers(Dispatcher, Schema, (uint8*)NativeStruct->GetStruct());
 		}
 	}
 #endif
diff --git a/Engine/Source/Runtime/CoreUObject/Public/UObject/GarbageCollectionSchema.h b/Engine/Source/Runtime/CoreUObject/Public/UObject/GarbageCollectionSchema.h
--- a/Engine/Source/Runtime/CoreUObject/Public/UObject/GarbageCollectionSchema.h
+++ b/Engine/Source/Runtime/CoreUObject/Public/UObject/GarbageCollectionSchema.h
@@ -270,7 +270,7 @@ public:
 	COREUOBJECT_API void Add(FMemberDeclaration Member);
 
 	// Append all but the last Stop/ARO member
-	COREUOBJECT_API void Append(FSchemaView SuperSchema, int32 OffsetForProperties = 0);
+	COREUOBJECT_API void Append(FSchemaView SuperSchema);
 
 	// Emit Stop or ARO and build a view that builder holds reference to. Returns same view if called multiple times.
 	COREUOBJECT_API FSchemaView Build(ObjectAROFn ARO = nullptr);
diff --git a/Engine/Source/Runtime/CoreUObject/Public/UObject/ObjectMacros.h b/Engine/Source/Runtime/CoreUObject/Public/UObject/ObjectMacros.h
--- a/Engine/Source/Runtime/CoreUObject/Public/UObject/ObjectMacros.h
+++ b/Engine/Source/Runtime/CoreUObject/Public/UObject/ObjectMacros.h
@@ -221,8 +221,8 @@ enum EClassFlags
 	CLASS_ProjectUserConfig	  = 0x00000040u,
 	/** Class is a native class - native interfaces will have CLASS_Native set, but not RF_MarkAsNative */
 	CLASS_Native			  = 0x00000080u,
-	/** Class can contain partials */
-	CLASS_Partial			  = 0x00000100u,
+
+	// CLASS_Unused			  = 0x00000100u,
 
 	/** Do not allow users to create in the editor. */
 	CLASS_NotPlaceable        = 0x00000200u,
@@ -601,7 +601,7 @@ enum EObjectFlags : uint32
 	// The group of flags tracks the stages of the lifetime of a uobject
 	RF_NeedInitialization		=0x00000200,	///< This object has not completed its initialization process. Cleared when ~FObjectInitializer completes
 	RF_NeedLoad					=0x00000400,	///< During load, indicates object needs loading.
-	RF_HasPartials				=0x00000800,	///< Set if object has partials (this means memory allocation is offsetted since partials are in front of object ptr)
+	RF_KeepForCooker UE_DEPRECATED(5.6, "This flag is now obsolete as the functionality is unnecessary.")			=0x00000800,	///< Keep this object during garbage collection because it's still being used by the cooker
 	RF_GeneratePublicHash		=0x00000800,	///< During cooking, generate a public hash for this object even is not marked RF_Public.
 	RF_NeedPostLoad				=0x00001000,	///< Object needs to be postloaded.
 	RF_NeedPostLoadSubobjects	=0x00002000,	///< During load, indicates that the object still needs to instance subobjects and fixup serialized component references
@@ -786,7 +786,6 @@ struct FReferencerInformationList
 #define UPROPERTY(...)
 #define UFUNCTION(...)
 #define USTRUCT(...)
-#define UPARTIAL(...)
 #define UMETA(...)
 #define UPARAM(...)
 #define UENUM(...)
diff --git a/Engine/Source/Runtime/CoreUObject/Public/UObject/RegisterCompiledInObjects.h b/Engine/Source/Runtime/CoreUObject/Public/UObject/RegisterCompiledInObjects.h
--- a/Engine/Source/Runtime/CoreUObject/Public/UObject/RegisterCompiledInObjects.h
+++ b/Engine/Source/Runtime/CoreUObject/Public/UObject/RegisterCompiledInObjects.h
@@ -15,7 +15,6 @@ class UClass;
 class UDelegateFunction;
 class UObject;
 class UPackage;
-namespace UE::CoreUObject::Private { struct FPartialClass; }
 
 // Utility struct for passing many arrays with designated initializer for clarity 
 struct FCompiledInObjects
@@ -43,8 +42,6 @@ struct FCompiledInObjects
 	// Other types of objects (e.g. UDelegateFunction) which can be initialized by casting after
 	// classes are propery initialized.
 	TConstArrayView<UObject*> Others;
-	
-	TConstArrayView<UE::CoreUObject::Private::FPartialClass*> Partials;
 
 	inline TConstArrayView<UClass*> GetClasses() const
 	{
diff --git a/Engine/Source/Runtime/CoreUObject/Public/UObject/UObjectAllocator.h b/Engine/Source/Runtime/CoreUObject/Public/UObject/UObjectAllocator.h
--- a/Engine/Source/Runtime/CoreUObject/Public/UObject/UObjectAllocator.h
+++ b/Engine/Source/Runtime/CoreUObject/Public/UObject/UObjectAllocator.h
@@ -41,18 +41,16 @@ public:
 	 * @param Size size of uobject to allocate
 	 * @param Alignment alignment of uobject to allocate
 	 * @param bAllowPermanent if true, allow allocation in the permanent object pool, if it fits
-	 * @param OffsetToObject offset of where the UObjectBase* starts. Should be a positive value to create memory in front of returned pointer
 	 * @return newly allocated UObjectBase (not really a UObjectBase yet, no constructor like thing has been called).
 	 */
-	COREUOBJECT_API UObjectBase* AllocateUObject(int32 Size, int32 Alignment, bool bAllowPermanent, int32 OffsetToObject = 0);
+	COREUOBJECT_API UObjectBase* AllocateUObject(int32 Size, int32 Alignment, bool bAllowPermanent);
 
 	/**
 	 * Returns a UObjectBase to the free store, unless it is in the permanent object pool
 	 *
 	 * @param Object object to free
-	 * @param OffsetToObject offset to Object in allocation. Should match the value sent into AllocateUObject
 	 */
-	COREUOBJECT_API void FreeUObject(UObjectBase *Object, int32 OffsetToObject = 0) const;
+	COREUOBJECT_API void FreeUObject(UObjectBase *Object) const;
 };
 
 /** Global UObjectBase allocator							*/
diff --git a/Engine/Source/Runtime/CoreUObject/Public/UObject/UObjectBase.h b/Engine/Source/Runtime/CoreUObject/Public/UObject/UObjectBase.h
--- a/Engine/Source/Runtime/CoreUObject/Public/UObject/UObjectBase.h
+++ b/Engine/Source/Runtime/CoreUObject/Public/UObject/UObjectBase.h
@@ -32,7 +32,7 @@ class UScriptStruct;
 namespace UE::CoreUObject::Private
 {
 	struct FObjectHandleUtils;
-	struct FPartialClass;
+
 	class FCompiledInObjectLinker;
 }
 
@@ -582,25 +582,6 @@ struct FStructRegisterCompiledInInfo
 	FStructReloadVersionInfo VersionInfo;
 };
 
-
-struct FPartialReloadVersionInfo
-{
-#if WITH_RELOAD
-	SIZE_T Size = 0;
-	uint32 Hash = 0;
-#endif
-};
-
-using FPartialRegistrationInfo = TRegistrationInfo<UE::CoreUObject::Private::FPartialClass, FPartialReloadVersionInfo>;
-
-struct FPartialRegisterCompiledInInfo
-{
-	UE::CoreUObject::Private::FPartialClass* (*OuterRegister)(ETypeConstructPhase);
-	const TCHAR* Name;
-	FPartialRegistrationInfo* Info;
-	FPartialReloadVersionInfo VersionInfo;
-};
-
 /**
  * Adds a struct registration and version information. The InInfo parameter must be static.
  */
@@ -671,7 +652,7 @@ COREUOBJECT_API void RegisterCompiledInInfo(UPackage* (*InOuterRegister)(ETypeCo
 /**
  * Register compiled in information for multiple classes, structures, and enumerations
  */
-COREUOBJECT_API void RegisterCompiledInInfo(const TCHAR *PackageName, const FClassRegisterCompiledInInfo *ClassInfo, size_t NumClassInfo, const FStructRegisterCompiledInInfo *StructInfo, size_t NumStructInfo, const FEnumRegisterCompiledInInfo *EnumInfo, size_t NumEnumInfo, const FPartialRegisterCompiledInInfo *PartialsInfo, size_t NumPartialsInfo);
+COREUOBJECT_API void RegisterCompiledInInfo(const TCHAR* PackageName, const FClassRegisterCompiledInInfo* ClassInfo, size_t NumClassInfo, const FStructRegisterCompiledInInfo* StructInfo, size_t NumStructInfo, const FEnumRegisterCompiledInInfo* EnumInfo, size_t NumEnumInfo);
 
 #endif // !UE_WITH_CONSTINIT_UOBJECT
 
diff --git a/Engine/Source/Runtime/CoreUObject/Public/UObject/UObjectPartials.h b/Engine/Source/Runtime/CoreUObject/Public/UObject/UObjectPartials.h
deleted file mode 100644
