EpicGames.LogMatchers: Fix a static analysis warning by specifying ordinal string comparison when checking if a log line contains the callstack indicator.

#Synced-CL 51454738

diff --git a/Engine/Source/Programs/Shared/EpicGames.LogMatchers/LogChannelEventMatcher.cs b/Engine/Source/Programs/Shared/EpicGames.LogMatchers/LogChannelEventMatcher.cs
--- a/Engine/Source/Programs/Shared/EpicGames.LogMatchers/LogChannelEventMatcher.cs
+++ b/Engine/Source/Programs/Shared/EpicGames.LogMatchers/LogChannelEventMatcher.cs
@@ -64,7 +64,7 @@ namespace EpicGames.LogMatchers
 						{
 							break;
 						}
-						if (line.Contains("[Callstack]"))
+						if (line.Contains("[Callstack]", System.StringComparison.Ordinal))
 						{
 							callstackOffset = i;
 							break;
