Horde: Set GrpcServices="None" to fix error when using IntermediateOutputPath during a dotnet build

This fixes CS2001 errors about LogRpcMessagesGrpc.cs files not existing. The issue only comes up when setting custom IntermediateOutputPath values on builds.

The log_rpc_messages.proto spec does not define any services, only messages. The value of GrpcServices defaults to Both which causes the build to expect a LogRpcMessagesGrpc.cs file to be generated. However because the spec doesn't define any services, that file is never generated.

In a normal build the Google.Protobuf code handles this silently. It checks if files are missing under `BaseIntermediateOutputPath` (which defaults to `obj`), and if any files are missing it silently creates empty stubs of those files. When we override the IntermediateOutputPath the silent creation of missing files continues to use the original location, not the new location.

By setting GrpcServices="None" we fix this whole issue as the compiler no longer expects the missing file to exist at all.

#rnx
#Synced-CL 51454371

diff --git a/Engine/Source/Programs/Shared/EpicGames.Horde/EpicGames.Horde.csproj b/Engine/Source/Programs/Shared/EpicGames.Horde/EpicGames.Horde.csproj
--- a/Engine/Source/Programs/Shared/EpicGames.Horde/EpicGames.Horde.csproj
+++ b/Engine/Source/Programs/Shared/EpicGames.Horde/EpicGames.Horde.csproj
@@ -36,6 +36,6 @@
 
   <ItemGroup>
     <Protobuf Include="Protos\horde\log_rpc.proto" ProtoRoot="Protos" GrpcServices="Both" Access="Public" />
-    <Protobuf Include="Protos\horde\log_rpc_messages.proto" ProtoRoot="Protos" Access="Public" />
+    <Protobuf Include="Protos\horde\log_rpc_messages.proto" ProtoRoot="Protos" GrpcServices="None" Access="Public" />
   </ItemGroup>
 </Project>
