Fixing bug where the suspend and resume log parsing messages would not be handled correctly when output as json.

#Synced-CL 51292189

diff --git a/Engine/Source/Programs/Shared/EpicGames.Core/LogEventParser.cs b/Engine/Source/Programs/Shared/EpicGames.Core/LogEventParser.cs
--- a/Engine/Source/Programs/Shared/EpicGames.Core/LogEventParser.cs
+++ b/Engine/Source/Programs/Shared/EpicGames.Core/LogEventParser.cs
@@ -663,6 +663,16 @@ namespace EpicGames.Core
 					if (JsonLogEvent.TryParse(data, out jsonEvent))
 					{
 						ProcessData(true);
+
+						if (Regex.IsMatch(line, @"\\u003c-- Suspend Log Parsing --\\u003e", RegexOptions.IgnoreCase))
+						{
+							_matchingEnabled--;
+						}
+						else if (Regex.IsMatch(line, @"\\u003c-- Resume Log Parsing --\\u003e", RegexOptions.IgnoreCase))
+						{
+							_matchingEnabled++;
+						}
+
 						if (!TryIgnoreStructuredEvent(jsonEvent))
 						{
 							_logger.Log(jsonEvent.Level, jsonEvent.EventId, jsonEvent, null, JsonLogEvent.Format);
