[Horde] Add GetBatchLogEventsAsync to HordeHttpClient for efficient warning/error count retrieval
- Adds a batch log events API method that accepts multiple log IDs in a single
- POST to /api/v1/logs/events. Supports severity filtering (Warning/Error) and Count=0 mode to retrieve only per-log totals without event payloads.

#rnx

#Synced-CL 51437754

diff --git a/Engine/Source/Programs/Shared/EpicGames.Horde/HordeHttpClient.cs b/Engine/Source/Programs/Shared/EpicGames.Horde/HordeHttpClient.cs
--- a/Engine/Source/Programs/Shared/EpicGames.Horde/HordeHttpClient.cs
+++ b/Engine/Source/Programs/Shared/EpicGames.Horde/HordeHttpClient.cs
@@ -685,6 +685,18 @@ namespace EpicGames.Horde
 			return GetAsync<LogLinesResponse>(_httpClient, $"/api/v1/logs/{logId}/lines?index={startIndex}&count={count}", cancellationToken);
 		}
 
+		/// <summary>
+		/// Get events for multiple log files in a single request.
+		/// Set <see cref="GetBatchLogEventsRequest.Count"/> to 0 to retrieve only the per-log totals
+		/// without downloading any event payloads.
+		/// </summary>
+		/// <param name="request">Log IDs, optional severity filter, and optional per-log event count limit</param>
+		/// <param name="cancellationToken">Cancellation token for the operation</param>
+		public Task<GetBatchLogEventsResponse> GetBatchLogEventsAsync(GetBatchLogEventsRequest request, CancellationToken cancellationToken = default)
+		{
+			return PostAsync<GetBatchLogEventsResponse, GetBatchLogEventsRequest>(_httpClient, "/api/v1/logs/events", request, cancellationToken);
+		}
+
 		#endregion
 
 		#region Graph
