Horde: Reword audit log message template for consistency

#rnx

diff --git a/Engine/Source/Programs/Shared/EpicGames.Horde/Issues/Handlers/BuildGraphIssueHandler.cs b/Engine/Source/Programs/Shared/EpicGames.Horde/Issues/Handlers/BuildGraphIssueHandler.cs
--- a/Engine/Source/Programs/Shared/EpicGames.Horde/Issues/Handlers/BuildGraphIssueHandler.cs
+++ b/Engine/Source/Programs/Shared/EpicGames.Horde/Issues/Handlers/BuildGraphIssueHandler.cs
@@ -41,7 +41,7 @@ namespace EpicGames.Horde.Issues.Handlers
 					issue.Events.Add(issueEvent);
 					issue.Keys.AddSourceFiles(issueEvent);
 					_issues.Add(issue);
-					issueEvent.AuditLogger?.LogDebug("BuildGraph issue added by event: '{Event}'", issueEvent.Render());
+					issueEvent.AuditLogger?.LogDebug("BuildGraph issue added for event: '{Event}'", issueEvent.Render());
 
 					return true;
 				}
