Remove redundant TargetFramework from csproj

#rnx
#Synced-CL 51506754

diff --git a/Engine/Source/Programs/AutomationTool.Tests/AutomationTool.Tests.csproj b/Engine/Source/Programs/AutomationTool.Tests/AutomationTool.Tests.csproj
--- a/Engine/Source/Programs/AutomationTool.Tests/AutomationTool.Tests.csproj
+++ b/Engine/Source/Programs/AutomationTool.Tests/AutomationTool.Tests.csproj
@@ -2,7 +2,6 @@
   <Import Project="..\Shared\UnrealEngine.csproj.props" />
 
   <PropertyGroup>
-    <TargetFramework>net8.0</TargetFramework>
     <ImplicitUsings>enable</ImplicitUsings>
     <Nullable>enable</Nullable>
 
diff --git a/Engine/Source/Programs/Shared/EpicGames.BuildGraph/EpicGames.BuildGraph.csproj b/Engine/Source/Programs/Shared/EpicGames.BuildGraph/EpicGames.BuildGraph.csproj
--- a/Engine/Source/Programs/Shared/EpicGames.BuildGraph/EpicGames.BuildGraph.csproj
+++ b/Engine/Source/Programs/Shared/EpicGames.BuildGraph/EpicGames.BuildGraph.csproj
@@ -1,8 +1,5 @@
 ﻿<Project Sdk="Microsoft.NET.Sdk">
   <PropertyGroup>
-    <TargetFramework>net8.0</TargetFramework>
-    <DebugType>pdbonly</DebugType>
-    <DebugType Condition="'$([System.Runtime.InteropServices.RuntimeInformation]::IsOSPlatform($([System.Runtime.InteropServices.OSPlatform]::Windows)))' == 'true' And '$([System.Runtime.InteropServices.RuntimeInformation]::OSArchitecture)' == 'Arm64'">portable</DebugType>
     <GenerateDocumentationFile>True</GenerateDocumentationFile>
   </PropertyGroup>
 
diff --git a/Engine/Source/Programs/Shared/EpicGames.LogMatchers/EpicGames.LogMatchers.csproj b/Engine/Source/Programs/Shared/EpicGames.LogMatchers/EpicGames.LogMatchers.csproj
--- a/Engine/Source/Programs/Shared/EpicGames.LogMatchers/EpicGames.LogMatchers.csproj
+++ b/Engine/Source/Programs/Shared/EpicGames.LogMatchers/EpicGames.LogMatchers.csproj
@@ -1,8 +1,4 @@
 ﻿<Project Sdk="Microsoft.NET.Sdk">
-  <PropertyGroup>
-    <TargetFramework>net8.0</TargetFramework>
-  </PropertyGroup>
-
   <ItemGroup>
     <PackageReference Include="Microsoft.Extensions.Logging" Version="8.0.1" />
   </ItemGroup>
diff --git a/Engine/Source/Programs/UnrealGameSync/UnrealGameSyncCmd/UnrealGameSyncCmd.csproj b/Engine/Source/Programs/UnrealGameSync/UnrealGameSyncCmd/UnrealGameSyncCmd.csproj
--- a/Engine/Source/Programs/UnrealGameSync/UnrealGameSyncCmd/UnrealGameSyncCmd.csproj
+++ b/Engine/Source/Programs/UnrealGameSync/UnrealGameSyncCmd/UnrealGameSyncCmd.csproj
@@ -3,7 +3,6 @@
 
   <PropertyGroup>
     <OutputType>Exe</OutputType>
-    <TargetFramework>net8.0</TargetFramework>
     <Nullable>enable</Nullable>
     <AssemblyName>ugs</AssemblyName>
   </PropertyGroup>
diff --git a/Engine/Source/Programs/UnrealGameSync/UnrealGameSyncShared/UnrealGameSyncShared.csproj b/Engine/Source/Programs/UnrealGameSync/UnrealGameSyncShared/UnrealGameSyncShared.csproj
--- a/Engine/Source/Programs/UnrealGameSync/UnrealGameSyncShared/UnrealGameSyncShared.csproj
+++ b/Engine/Source/Programs/UnrealGameSync/UnrealGameSyncShared/UnrealGameSyncShared.csproj
@@ -2,7 +2,6 @@
   <Import Project="..\..\Shared\UnrealEngine.csproj.props" />
 
   <PropertyGroup>
-    <TargetFramework>net8.0</TargetFramework>
     <Nullable>enable</Nullable>
     <RootNamespace>UnrealGameSync</RootNamespace>
   </PropertyGroup>
