EpicGames.OIDC: Fix concurrent refresh token handling

Fix a race condition where multiple concurrent processes (UBT, UGS, HordeCmd) could race on OIDC refresh token rotation, causing invalid_grant errors when one process uses a token already rotated by another.

* Always read the refresh token fresh from the backing store rather than caching it in memory
* Add a machine-wide file lock around the refresh operation so only one process refreshes at a time
* Retry with a re-read token on invalid_grant as a fallback for older binaries without the lock
* Add tests for refresh retry logic and WindowsTokenStore DPAPI round-trip

diff --git a/Engine/Source/Programs/OidcToken/OidcToken.Tests/OidcTokenTests.cs b/Engine/Source/Programs/OidcToken/OidcToken.Tests/OidcTokenTests.cs
--- a/Engine/Source/Programs/OidcToken/OidcToken.Tests/OidcTokenTests.cs
+++ b/Engine/Source/Programs/OidcToken/OidcToken.Tests/OidcTokenTests.cs
@@ -65,7 +65,6 @@ namespace OidcToken.Tests
 				CallBase = true
 			};
 			tokenStoreMock.Setup(store => store.AddRefreshToken(It.Is(providerName, StringComparer.InvariantCultureIgnoreCase), It.Is(mockRefreshToken, StringComparer.InvariantCultureIgnoreCase))).CallBase().Verifiable(Times.Once);
-			tokenStoreMock.Setup(store => store.Save()).CallBase().Verifiable(Times.Once);
 
 			string tempOutputFileName = Path.GetTempFileName();
 			try
@@ -134,7 +133,6 @@ namespace OidcToken.Tests
 				CallBase = true
 			};
 			tokenStoreMock.Setup(store => store.AddRefreshToken(It.Is(providerName, StringComparer.InvariantCultureIgnoreCase), It.IsAny<string>())).CallBase().Verifiable(Times.Never);
-			tokenStoreMock.Setup(store => store.Save()).CallBase().Verifiable(Times.Never);
 
 			const string loginResultErrorMessage = "failed to login";
 
@@ -223,7 +221,6 @@ namespace OidcToken.Tests
 				CallBase = true
 			};
 			tokenStoreMock.Setup(store => store.AddRefreshToken(It.Is(providerName, StringComparer.InvariantCultureIgnoreCase), It.IsAny<string>())).CallBase().Verifiable(Times.Once);
-			tokenStoreMock.Setup(store => store.Save()).CallBase().Verifiable(Times.Once);
 
 			string mockAccessToken = "access-token";
 			string mockRefreshToken = "refresh-token";
diff --git a/Engine/Source/Programs/Shared/EpicGames.Horde.Tests/EpicGames.Horde.Tests.csproj b/Engine/Source/Programs/Shared/EpicGames.Horde.Tests/EpicGames.Horde.Tests.csproj
--- a/Engine/Source/Programs/Shared/EpicGames.Horde.Tests/EpicGames.Horde.Tests.csproj
+++ b/Engine/Source/Programs/Shared/EpicGames.Horde.Tests/EpicGames.Horde.Tests.csproj
@@ -6,6 +6,7 @@
   <ItemGroup>
     <ProjectReference Include="..\EpicGames.Core\EpicGames.Core.csproj" />
     <ProjectReference Include="..\EpicGames.Horde\EpicGames.Horde.csproj" />
+    <ProjectReference Include="..\EpicGames.OIDC\EpicGames.OIDC.csproj" />
   </ItemGroup>
 
   <ItemGroup>
diff --git a/Engine/Source/Programs/Shared/EpicGames.Horde.Tests/OidcTokenRefreshTests.cs b/Engine/Source/Programs/Shared/EpicGames.Horde.Tests/OidcTokenRefreshTests.cs
new file mode 100644
--- /dev/null
+++ b/Engine/Source/Programs/Shared/EpicGames.Horde.Tests/OidcTokenRefreshTests.cs
@@ -0,0 +1,205 @@
+// Copyright Epic Games, Inc. All Rights Reserved.
+
+using System;
+using System.Threading;
+using System.Threading.Tasks;
+using EpicGames.OIDC;
+using Microsoft.VisualStudio.TestTools.UnitTesting;
+
+namespace EpicGames.Horde.Tests;
+
+/// <summary>
+/// Testable subclass of OidcTokenClient that overrides DoRefreshAsync (no real OIDC backend)
+/// and AcquireRefreshLock (in-process SemaphoreSlim instead of file lock).
+/// </summary>
+sealed class TestableOidcTokenClient(
+	string name,
+	ITokenStore tokenStore,
+	Func<string, TokenRefreshResult> refreshHandler)
+	: OidcTokenClient(name, CreateDummyProviderInfo(), TimeSpan.FromMinutes(20), tokenStore)
+{
+	private static readonly SemaphoreSlim s_testLock = new(1, 1);
+	public int RefreshCallCount { get; private set; }
+
+	protected override Task<TokenRefreshResult> DoRefreshAsync(
+		string refreshToken, CancellationToken cancellationToken)
+	{
+		RefreshCallCount++;
+		return Task.FromResult(refreshHandler(refreshToken));
+	}
+
+	protected override async Task<IDisposable> AcquireRefreshLockAsync(CancellationToken cancellationToken)
+	{
+		await s_testLock.WaitAsync(cancellationToken);
+		return new SemaphoreReleaser(s_testLock);
+	}
+
+	private sealed class SemaphoreReleaser(SemaphoreSlim s) : IDisposable
+	{
+		public void Dispose() => s.Release();
+	}
+
+	private static ProviderInfo CreateDummyProviderInfo() => new()
+	{
+		ServerUri = new Uri("https://fake-oidc.test"),
+		ClientId = "test-client",
+		DisplayName = "Test",
+		RedirectUri = new Uri("http://localhost:12345/callback"),
+		UseDiscoveryDocument = false
+	};
+}
+
+[TestClass]
+public sealed class OidcTokenRefreshTests
+{
+	private const string Provider = "test";
+
+	private static TokenRefreshResult SuccessResult(string accessToken = "newAccess", string refreshToken = "newRefresh") =>
+		new(IsError: false, Error: null, ErrorDescription: null,
+			AccessToken: accessToken, RefreshToken: refreshToken,
+			AccessTokenExpiration: DateTimeOffset.UtcNow.AddHours(1));
+
+	private static TokenRefreshResult InvalidGrantResult(string description = "Token expired") =>
+		new(IsError: true, Error: "invalid_grant", ErrorDescription: description,
+			AccessToken: null, RefreshToken: null, AccessTokenExpiration: default);
+
+	private static (InMemoryTokenStore Store, TestableOidcTokenClient Client) CreateClient(
+		string initialToken, Func<string, TokenRefreshResult> refreshHandler)
+	{
+		InMemoryTokenStore store = new();
+		store.AddRefreshToken(Provider, initialToken);
+		return (store, new TestableOidcTokenClient(Provider, store, refreshHandler));
+	}
+
+	/// <summary>
+	/// When the first refresh returns invalid_grant and a different token exists in the store,
+	/// the client should re-read the store and retry with the fresh token.
+	/// </summary>
+	[TestMethod]
+	public async Task RefreshToken_InvalidGrant_ReadsStoreAndRetriesAsync()
+	{
+		using InMemoryTokenStore store = new();
+		store.AddRefreshToken(Provider, "tokenA");
+		TestableOidcTokenClient client = new(Provider, store, refreshToken =>
+		{
+			if (refreshToken == "tokenA")
+			{
+				store.AddRefreshToken(Provider, "tokenB"); // simulate another process rotating the token
+				return InvalidGrantResult("Token has been rotated");
+			}
+			return SuccessResult("newAccess", "tokenC");
+		});
+		{
+			IOidcTokenInfo? result = await client.TryGetAccessTokenAsync(CancellationToken.None);
+
+			Assert.IsNotNull(result);
+			Assert.AreEqual("newAccess", result!.AccessToken);
+			Assert.IsTrue(store.TryGetRefreshToken(Provider, out string? storedToken));
+			Assert.AreEqual("tokenC", storedToken);
+			Assert.AreEqual(2, client.RefreshCallCount);
+		}
+	}
+
+	/// <summary>
+	/// When invalid_grant is returned but the store still has the same token,
+	/// the client should return null (no fresh token available).
+	/// </summary>
+	[TestMethod]
+	public async Task RefreshToken_InvalidGrant_NoFreshToken_ReturnsNullAsync()
+	{
+		(InMemoryTokenStore store, TestableOidcTokenClient client) = CreateClient("tokenA", _ => InvalidGrantResult());
+		using (store)
+		{
+			IOidcTokenInfo? result = await client.TryGetAccessTokenAsync(CancellationToken.None);
+
+			Assert.IsNull(result);
+			Assert.AreEqual(1, client.RefreshCallCount);
+		}
+	}
+
+	/// <summary>
+	/// A normal successful refresh should work and store the new token.
+	/// </summary>
+	[TestMethod]
+	public async Task RefreshToken_ValidRefresh_SucceedsAsync()
+	{
+		(InMemoryTokenStore store, TestableOidcTokenClient client) = CreateClient("tokenA", _ => SuccessResult("access1", "tokenB"));
+		using (store)
+		{
+			IOidcTokenInfo? result = await client.TryGetAccessTokenAsync(CancellationToken.None);
+
+			Assert.IsNotNull(result);
+			Assert.AreEqual("access1", result!.AccessToken);
+			Assert.IsTrue(store.TryGetRefreshToken(Provider, out string? storedToken));
+			Assert.AreEqual("tokenB", storedToken);
+			Assert.AreEqual(1, client.RefreshCallCount);
+		}
+	}
+
+	/// <summary>
+	/// The client should always read the store fresh at refresh time, not use a stale
+	/// cached copy from construction.
+	/// </summary>
+	[TestMethod]
+	public async Task RefreshToken_AlwaysReadsFromStoreFreshAsync()
+	{
+		(InMemoryTokenStore store, TestableOidcTokenClient client) = CreateClient("tokenA", refreshToken =>
+		{
+			if (refreshToken == "tokenA")
+			{
+				return InvalidGrantResult("Stale");
+			}
+			return SuccessResult("freshAccess", "tokenC");
+		});
+		using (store)
+		{
+			// Externally rotate the token (simulating another process)
+			store.AddRefreshToken(Provider, "tokenB");
+
+			IOidcTokenInfo? result = await client.TryGetAccessTokenAsync(CancellationToken.None);
+
+			Assert.IsNotNull(result);
+			Assert.AreEqual("freshAccess", result!.AccessToken);
+		}
+	}
+
+	/// <summary>
+	/// Concurrent calls to TryGetAccessTokenAsync should be serialized by the lock.
+	/// </summary>
+	[TestMethod]
+	public async Task RefreshToken_ConcurrentCalls_SerializedAsync()
+	{
+		int concurrentCount = 0;
+		int maxConcurrent = 0;
+
+		using InMemoryTokenStore store = new();
+		store.AddRefreshToken(Provider, "tokenA");
+		TestableOidcTokenClient client = new(Provider, store, refreshToken =>
+		{
+			int current = Interlocked.Increment(ref concurrentCount);
+			int snapshot;
+			do
+			{
+				snapshot = Volatile.Read(ref maxConcurrent);
+			}
+			while (current > snapshot && Interlocked.CompareExchange(ref maxConcurrent, current, snapshot) != snapshot);
+
+			Thread.Sleep(50); // small delay to increase chance of overlap detection
+			Interlocked.Decrement(ref concurrentCount);
+
+			string newToken = $"token_{Guid.NewGuid():N}";
+			store.AddRefreshToken(Provider, newToken);
+			return SuccessResult($"access_{refreshToken}", newToken);
+		});
+		{
+			Task<IOidcTokenInfo?> task1 = client.TryGetAccessTokenAsync(CancellationToken.None);
+			Task<IOidcTokenInfo?> task2 = client.TryGetAccessTokenAsync(CancellationToken.None);
+
+			IOidcTokenInfo?[] results = await Task.WhenAll(task1, task2);
+
+			Assert.IsNotNull(results[0]);
+			Assert.IsNotNull(results[1]);
+			Assert.AreEqual(1, maxConcurrent, "Refresh calls should have been serialized by the lock");
+		}
+	}
+}
diff --git a/Engine/Source/Programs/Shared/EpicGames.Horde.Tests/WindowsTokenStoreTests.cs b/Engine/Source/Programs/Shared/EpicGames.Horde.Tests/WindowsTokenStoreTests.cs
new file mode 100644
--- /dev/null
+++ b/Engine/Source/Programs/Shared/EpicGames.Horde.Tests/WindowsTokenStoreTests.cs
@@ -0,0 +1,380 @@
+// Copyright Epic Games, Inc. All Rights Reserved.
+
+using System;
+using System.Collections.Generic;
+using System.IO;
+using System.Runtime.InteropServices;
+using System.Runtime.Versioning;
+using System.Text.Json;
+using System.Threading;
+using System.Threading.Tasks;
+using EpicGames.OIDC;
+using Microsoft.Extensions.Logging;
+using Microsoft.Extensions.Logging.Abstractions;
+using Microsoft.VisualStudio.TestTools.UnitTesting;
+
+namespace EpicGames.Horde.Tests;
+
+/// <summary>
+/// Logger that fails the test if <see cref="LogLevel.Error"/> or higher is logged.
+/// This surfaces errors that <c>AddRefreshToken</c>'s catch-all would otherwise swallow.
+/// </summary>
+sealed class FailOnErrorLogger<T>(Action<string> failAction) : ILogger<T>
+{
+	public IDisposable? BeginScope<TState>(TState state) where TState : notnull => null;
+	public bool IsEnabled(LogLevel logLevel) => true;
+
+	public void Log<TState>(LogLevel logLevel, EventId eventId, TState state, Exception? exception, Func<TState, Exception?, string> formatter)
+	{
+		if (logLevel >= LogLevel.Error)
+		{
+			failAction($"Unexpected {logLevel}: {formatter(state, exception)}{(exception != null ? $"\n{exception}" : "")}");
+		}
+	}
+}
+
+/// <summary>
+/// End-to-end tests for <see cref="WindowsTokenStore"/> including DPAPI encrypt/decrypt
+/// and file I/O. These tests are Windows-only because they depend on CryptProtectData.
+/// </summary>
+[TestClass]
+[SupportedOSPlatform("windows")]
+public sealed class WindowsTokenStoreTests : IDisposable
+{
+	private readonly string _tempDir;
+	private readonly string _storePath;
+	private readonly ILogger<WindowsTokenStore> _logger;
+
+	public WindowsTokenStoreTests()
+	{
+		_tempDir = Path.Combine(Path.GetTempPath(), $"WindowsTokenStoreTests_{Guid.NewGuid():N}");
+		Directory.CreateDirectory(_tempDir);
+		_storePath = Path.Combine(_tempDir, "oidcTokenStore.dat");
+		_logger = new FailOnErrorLogger<WindowsTokenStore>(msg => Assert.Fail(msg));
+	}
+
+	public void Dispose()
+	{
+		try
+		{
+			Directory.Delete(_tempDir, true);
+		}
+		catch
+		{
+			// Best effort cleanup
+		}
+	}
+
+	private WindowsTokenStore CreateStore() => new(_logger, _storePath);
+
+	private static bool IsWindows => RuntimeInformation.IsOSPlatform(OSPlatform.Windows);
+
+	/// <summary>
+	/// Round-trip: AddRefreshToken persists to disk, TryGetRefreshToken reads it back
+	/// through DPAPI decrypt.
+	/// </summary>
+	[TestMethod]
+	public void AddAndGet_RoundTrip()
+	{
+		if (!IsWindows)
+		{
+			Assert.Inconclusive("Windows-only test (requires DPAPI).");
+		}
+
+		using WindowsTokenStore store = CreateStore();
+		store.AddRefreshToken("myProvider", "secret-refresh-token-123");
+
+		Assert.IsTrue(File.Exists(_storePath), "Store file should exist after AddRefreshToken.");
+
+		// Read back with a fresh instance (proves it's persisted, not just in memory)
+		using WindowsTokenStore store2 = CreateStore();
+		Assert.IsTrue(store2.TryGetRefreshToken("myProvider", out string? token));
+		Assert.AreEqual("secret-refresh-token-123", token);
+	}
+
+	/// <summary>
+	/// Multiple providers can coexist in the same store file.
+	/// </summary>
+	[TestMethod]
+	public void MultipleProviders_Coexist()
+	{
+		if (!IsWindows)
+		{
+			Assert.Inconclusive("Windows-only test (requires DPAPI).");
+		}
+
+		using WindowsTokenStore store = CreateStore();
+		store.AddRefreshToken("providerA", "tokenA");
+		store.AddRefreshToken("providerB", "tokenB");
+
+		using WindowsTokenStore store2 = CreateStore();
+		Assert.IsTrue(store2.TryGetRefreshToken("providerA", out string? a));
+		Assert.AreEqual("tokenA", a);
+		Assert.IsTrue(store2.TryGetRefreshToken("providerB", out string? b));
+		Assert.AreEqual("tokenB", b);
+	}
+
+	/// <summary>
+	/// Overwriting a provider's token replaces the old value.
+	/// </summary>
+	[TestMethod]
+	public void AddRefreshToken_Overwrite_ReplacesOldValue()
+	{
+		if (!IsWindows)
+		{
+			Assert.Inconclusive("Windows-only test (requires DPAPI).");
+		}
+
+		using WindowsTokenStore store = CreateStore();
+		store.AddRefreshToken("prov", "old-token");
+		store.AddRefreshToken("prov", "new-token");
+
+		using WindowsTokenStore store2 = CreateStore();
+		Assert.IsTrue(store2.TryGetRefreshToken("prov", out string? token));
+		Assert.AreEqual("new-token", token);
+	}
+
+	/// <summary>
+	/// TryGetRefreshToken returns false for a provider that was never stored.
+	/// </summary>
+	[TestMethod]
+	public void TryGetRefreshToken_MissingProvider_ReturnsFalse()
+	{
+		if (!IsWindows)
+		{
+			Assert.Inconclusive("Windows-only test (requires DPAPI).");
+		}
+
+		using WindowsTokenStore store = CreateStore();
+		Assert.IsFalse(store.TryGetRefreshToken("nonexistent", out _));
+	}
+
+	/// <summary>
+	/// TryGetRefreshToken returns false when the store file doesn't exist.
+	/// </summary>
+	[TestMethod]
+	public void TryGetRefreshToken_NoFile_ReturnsFalse()
+	{
+		if (!IsWindows)
+		{
+			Assert.Inconclusive("Windows-only test (requires DPAPI).");
+		}
+
+		Assert.IsFalse(File.Exists(_storePath));
+		using WindowsTokenStore store = CreateStore();
+		Assert.IsFalse(store.TryGetRefreshToken("any", out _));
+	}
+
+	/// <summary>
+	/// Corrupt JSON in the store file is handled gracefully — TryGetRefreshToken
+	/// returns false instead of throwing.
+	/// </summary>
+	[TestMethod]
+	public void TryGetRefreshToken_CorruptJson_ReturnsFalse()
+	{
+		if (!IsWindows)
+		{
+			Assert.Inconclusive("Windows-only test (requires DPAPI).");
+		}
+
+		Directory.CreateDirectory(Path.GetDirectoryName(_storePath)!);
+		File.WriteAllText(_storePath, "{{not valid json!!");
+
+		using WindowsTokenStore store = CreateStore();
+		Assert.IsFalse(store.TryGetRefreshToken("anyProvider", out _));
+	}
+
+	/// <summary>
+	/// A store file with one valid and one corrupt Base64 entry still returns
+	/// the valid provider's token (the corrupt entry is skipped).
+	/// </summary>
+	[TestMethod]
+	public void TryGetRefreshToken_CorruptBase64Entry_SkipsCorruptProvider()
+	{
+		if (!IsWindows)
+		{
+			Assert.Inconclusive("Windows-only test (requires DPAPI).");
+		}
+
+		// First, write a valid token so we get a real DPAPI-encrypted entry
+		using WindowsTokenStore store = CreateStore();
+		store.AddRefreshToken("good", "valid-token");
+
+		// Now manually corrupt the store file by adding a bad Base64 entry
+		string json = File.ReadAllText(_storePath);
+		WindowsTokenStoreState? state = JsonSerializer.Deserialize(json, WindowsTokenStoreStateContext.Default.WindowsTokenStoreState);
+		Assert.IsNotNull(state);
+		state!.Providers["corrupt"] = "not-valid-base64!!!";
+		string corrupted = JsonSerializer.Serialize(state, WindowsTokenStoreStateContext.Default.WindowsTokenStoreState);
+		File.WriteAllText(_storePath, corrupted);
+
+		// The good provider should still work
+		using WindowsTokenStore store2 = CreateStore();
+		Assert.IsTrue(store2.TryGetRefreshToken("good", out string? token));
+		Assert.AreEqual("valid-token", token);
+
+		// The corrupt provider should return false (not throw)
+		Assert.IsFalse(store2.TryGetRefreshToken("corrupt", out _));
+	}
+
+	/// <summary>
+	/// AddRefreshToken does not throw even when the store file is read-only
+	/// (simulating a permissions issue). The token is silently lost.
+	/// </summary>
+	[TestMethod]
+	public void AddRefreshToken_ReadOnlyFile_DoesNotThrow()
+	{
+		if (!IsWindows)
+		{
+			Assert.Inconclusive("Windows-only test (requires DPAPI).");
+		}
+
+		// Create a valid store first
+		using WindowsTokenStore store = CreateStore();
+		store.AddRefreshToken("prov", "original");
+
+		// Make the file read-only
+		File.SetAttributes(_storePath, FileAttributes.ReadOnly);
+		try
+		{
+			// This should not throw — AddRefreshToken catches and logs the error.
+			// Use a null logger here since we EXPECT the error to be swallowed.
+			using WindowsTokenStore store2 = new(NullLogger<WindowsTokenStore>.Instance, _storePath);
+			store2.AddRefreshToken("prov", "updated");
+
+			// The original token should still be there (write failed silently)
+			using WindowsTokenStore store3 = CreateStore();
+			Assert.IsTrue(store3.TryGetRefreshToken("prov", out string? token));
+			Assert.AreEqual("original", token);
+		}
+		finally
+		{
+			File.SetAttributes(_storePath, FileAttributes.Normal);
+		}
+	}
+}
+
+/// <summary>
+/// Tests for <see cref="WindowsTokenStoreState"/> JSON serialization (cross-platform, no DPAPI).
+/// </summary>
+[TestClass]
+public sealed class WindowsTokenStoreStateTests
+{
+	[TestMethod]
+	public void RoundTrip_SerializeDeserialize()
+	{
+		byte[] bytesA = [1, 2, 3, 4];
+		byte[] bytesB = [10, 20, 30];
+		WindowsTokenStoreState state = new(new Dictionary<string, string> { ["providerA"] = Convert.ToBase64String(bytesA), ["providerB"] = Convert.ToBase64String(bytesB) });
+
+		string json = JsonSerializer.Serialize(state, WindowsTokenStoreStateContext.Default.WindowsTokenStoreState);
+		WindowsTokenStoreState? deserialized = JsonSerializer.Deserialize(json, WindowsTokenStoreStateContext.Default.WindowsTokenStoreState);
+
+		Assert.IsNotNull(deserialized);
+		CollectionAssert.AreEqual(bytesA, Convert.FromBase64String(deserialized!.Providers["providerA"]));
+		CollectionAssert.AreEqual(bytesB, Convert.FromBase64String(deserialized.Providers["providerB"]));
+	}
+
+	[TestMethod]
+	public void Deserialize_InvalidJson_ReturnsNull()
+	{
+		WindowsTokenStoreState? result = null;
+		try
+		{
+			result = JsonSerializer.Deserialize("not valid json{{{", WindowsTokenStoreStateContext.Default.WindowsTokenStoreState);
+		}
+		catch (JsonException)
+		{
+			// Expected — callers catch this and treat as empty store
+		}
+
+		Assert.IsNull(result);
+	}
+}
+
+/// <summary>
+/// Tests that <see cref="OidcTokenClient"/> callers are resilient to silent store failures.
+/// </summary>
+[TestClass]
+public sealed class TokenStoreFailureTests
+{
+	private const string Provider = "test";
+
+	private static TokenRefreshResult SuccessResult(string accessToken = "newAccess", string refreshToken = "newRefresh") =>
+		new(IsError: false, Error: null, ErrorDescription: null,
+			AccessToken: accessToken, RefreshToken: refreshToken,
+			AccessTokenExpiration: DateTimeOffset.UtcNow.AddHours(1));
+
+	/// <summary>
+	/// When the store silently fails to persist (simulating WindowsTokenStore's
+	/// catch-and-log on I/O error), the client still returns the access token.
+	/// </summary>
+	[TestMethod]
+	public async Task RefreshToken_StoreWriteFails_StillReturnsTokenAsync()
+	{
+		using FailingTokenStore store = new();
+		store.AddRefreshToken(Provider, "initialToken");
+
+		int callCount = 0;
+		TestableOidcTokenClient client = new(Provider, store, _ =>
+		{
+			callCount++;
+			store.SetDropWrites(true);
+			return SuccessResult("access1", "rotatedToken");
+		});
+
+		IOidcTokenInfo? result = await client.TryGetAccessTokenAsync(CancellationToken.None);
+
+		Assert.IsNotNull(result);
+		Assert.AreEqual("access1", result!.AccessToken);
+		Assert.AreEqual(1, callCount);
+		Assert.IsTrue(store.TryGetRefreshToken(Provider, out string? storedToken));
+		Assert.AreEqual("initialToken", storedToken, "Write was silently dropped; old token should remain.");
+	}
+
+	[TestMethod]
+	public async Task TryGetAccessToken_EmptyStore_ThrowsNotLoggedInAsync()
+	{
+		using FailingTokenStore store = new();
+		TestableOidcTokenClient client = new(Provider, store, _ => SuccessResult());
+
+		await Assert.ThrowsExactlyAsync<NotLoggedInException>(
+			() => client.TryGetAccessTokenAsync(CancellationToken.None));
+	}
+}
+
+/// <summary>
+/// Token store that silently drops writes when configured, simulating WindowsTokenStore's
+/// catch-and-log behavior on persist failure.
+/// </summary>
+sealed class FailingTokenStore : ITokenStore
+{
+	private readonly Dictionary<string, string> _tokens = [];
+	private bool _dropWrites;
+
+	public void SetDropWrites(bool drop) => _dropWrites = drop;
+
+	public bool TryGetRefreshToken(string oidcProvider, out string refreshToken)
+	{
+		if (!_tokens.TryGetValue(oidcProvider, out string? result))
+		{
+			refreshToken = "";
+			return false;
+		}
+
+		refreshToken = result;
+		return true;
+	}
+
+	public void AddRefreshToken(string oidcProvider, string refreshToken)
+	{
+		if (!_dropWrites)
+		{
+			_tokens[oidcProvider] = refreshToken;
+		}
+	}
+
+	public void Dispose()
+	{
+	}
+}
diff --git a/Engine/Source/Programs/Shared/EpicGames.Horde/Auth/HordeOidcTokenManager.cs b/Engine/Source/Programs/Shared/EpicGames.Horde/Auth/HordeOidcTokenManager.cs
--- a/Engine/Source/Programs/Shared/EpicGames.Horde/Auth/HordeOidcTokenManager.cs
+++ b/Engine/Source/Programs/Shared/EpicGames.Horde/Auth/HordeOidcTokenManager.cs
@@ -8,12 +8,15 @@ using System.Threading.Tasks;
 using EpicGames.Horde.Server;
 using EpicGames.OIDC;
 using Microsoft.Extensions.Configuration;
+using Microsoft.Extensions.Logging;
 
 namespace EpicGames.Horde.Auth;
 
 internal sealed class HordeOidcTokenManager(
 	HordeAuthConfigProvider hordeAuthConfigProvider,
-	ITokenStore tokenStore
+	ITokenStore tokenStore,
+	ILogger<HordeOidcTokenManager>? logger = null,
+	ILoggerFactory? loggerFactory = null
 ) : IOidcTokenManager
 {
 	public async Task<IOidcTokenInfo> GetAccessToken(string providerIdentifier, CancellationToken cancellationToken)
@@ -51,12 +54,14 @@ internal sealed class HordeOidcTokenManager(
 			return new AnonymousOidcTokenInfo();
 		}
 
+		_logger?.LogDebug("Interactive login requested for OIDC provider '{Provider}'. A browser window will open.", providerIdentifier);
 		OidcTokenManager oidcTokenManager = CreateOidcTokenManager(providerIdentifier, authConfig);
 		return await oidcTokenManager.LoginAsync(providerIdentifier, cancellationToken);
 	}
 
 	public async Task<IOidcTokenInfo?> TryGetAccessToken(string providerIdentifier, CancellationToken cancellationToken)
 	{
+		_logger?.LogDebug("Requesting access token for provider '{Provider}'.", providerIdentifier);
 		GetAuthConfigResponse authConfig = await _hordeAuthConfigProvider.GetAuthConfigResponseAsync(cancellationToken);
 
 		if (authConfig.Method == AuthMethod.Anonymous)
@@ -70,6 +75,8 @@ internal sealed class HordeOidcTokenManager(
 
 	private readonly HordeAuthConfigProvider _hordeAuthConfigProvider = hordeAuthConfigProvider;
 	private readonly ITokenStore _tokenStore = tokenStore;
+	private readonly ILogger<HordeOidcTokenManager>? _logger = logger;
+	private readonly ILoggerFactory? _loggerFactory = loggerFactory;
 
 	private OidcTokenManager CreateOidcTokenManager(string providerIdentifier, GetAuthConfigResponse authConfig)
 	{
@@ -85,6 +92,8 @@ internal sealed class HordeOidcTokenManager(
 			throw new Exception("Provider identifier mismatch");
 		}
 
+		_logger?.LogDebug("Creating new OidcTokenManager for provider '{Provider}' (note: in-memory token cache is not preserved across calls).", oidcProvider);
+
 		Dictionary<string, string?> values = [];
 		values[$"Providers:{oidcProvider}:DisplayName"] = oidcProvider;
 		values[$"Providers:{oidcProvider}:ServerUri"] = authConfig.ServerUrl;
@@ -99,6 +108,6 @@ internal sealed class HordeOidcTokenManager(
 		builder.AddInMemoryCollection(values);
 		IConfiguration configuration = builder.Build();
 
-		return OidcTokenManager.CreateTokenManager(configuration, _tokenStore, [oidcProvider]);
+		return OidcTokenManager.CreateTokenManager(configuration, _tokenStore, [oidcProvider], _loggerFactory);
 	}
 }
diff --git a/Engine/Source/Programs/Shared/EpicGames.Horde/HordeClient.cs b/Engine/Source/Programs/Shared/EpicGames.Horde/HordeClient.cs
--- a/Engine/Source/Programs/Shared/EpicGames.Horde/HordeClient.cs
+++ b/Engine/Source/Programs/Shared/EpicGames.Horde/HordeClient.cs
@@ -351,15 +351,15 @@ namespace EpicGames.Horde
 		{
 			_baseHttpMessageHandler = baseHttpMessageHandler;
 
-			_tokenStore = TokenStoreFactory.CreateTokenStore();
+			_tokenStore = TokenStoreFactory.CreateTokenStore(loggerFactory);
 			_hordeAuthConfigProvider =
 				new HordeAuthConfigProvider(_baseHttpMessageHandler.Instance, serverUrl, loggerFactory.CreateLogger<HordeAuthConfigProvider>());
-			IOidcTokenManager oidcTokenManager = new HordeOidcTokenManager(_hordeAuthConfigProvider, _tokenStore);
+			IOidcTokenManager oidcTokenManager = new HordeOidcTokenManager(_hordeAuthConfigProvider, _tokenStore, loggerFactory.CreateLogger<HordeOidcTokenManager>(), loggerFactory);
 			IStaticAccessTokenProvider staticAccessTokenProvider = new HordeStaticAccessTokenProvider(hordeOptions);
 			HttpOidcAuthHandlerConfig httpOidcAuthHandlerConfig = HordeOptionsToIHttpOidcAuthHandlerConfig(hordeOptions.Value);
 			IClock clock = new DefaultClock();
 			_authHandlerState = new OidcAuthState(oidcTokenManager, staticAccessTokenProvider, clock, _hordeAuthConfigProvider, loggerFactory.CreateLogger<OidcAuthState>());
-			_authHttpMessageHandler = new HttpOidcAuthHandler(_baseHttpMessageHandler.Instance, _authHandlerState, httpOidcAuthHandlerConfig);
+			_authHttpMessageHandler = new HttpOidcAuthHandler(_baseHttpMessageHandler.Instance, _authHandlerState, httpOidcAuthHandlerConfig, loggerFactory.CreateLogger<HttpOidcAuthHandler>());
 
 			_authHandlerState.OnStateChanged += NotifyAuthStateChanged;
 		}
diff --git a/Engine/Source/Programs/Shared/EpicGames.OIDC/HttpOidcAuthHandler.cs b/Engine/Source/Programs/Shared/EpicGames.OIDC/HttpOidcAuthHandler.cs
--- a/Engine/Source/Programs/Shared/EpicGames.OIDC/HttpOidcAuthHandler.cs
+++ b/Engine/Source/Programs/Shared/EpicGames.OIDC/HttpOidcAuthHandler.cs
@@ -6,6 +6,7 @@ using System.Net.Http;
 using System.Net.Http.Headers;
 using System.Threading;
 using System.Threading.Tasks;
+using Microsoft.Extensions.Logging;
 
 namespace EpicGames.OIDC;
 
@@ -17,16 +18,18 @@ public sealed class HttpOidcAuthHandler : DelegatingHandler
 	/// <summary>
 	/// Constructor
 	/// </summary>
-	public HttpOidcAuthHandler(HttpMessageHandler innerHandler, IOidcAuthState authState, HttpOidcAuthHandlerConfig config)
+	public HttpOidcAuthHandler(HttpMessageHandler innerHandler, IOidcAuthState authState, HttpOidcAuthHandlerConfig config, ILogger<HttpOidcAuthHandler>? logger = null)
 		: base(innerHandler)
 	{
 		_authState = authState;
 		_config = config;
+		_logger = logger;
 		_allowInteractiveLogin = new(_config.AllowInteractiveLoginHeaderKey);
 	}
 
 	private readonly IOidcAuthState _authState;
 	private readonly HttpOidcAuthHandlerConfig _config;
+	private readonly ILogger<HttpOidcAuthHandler>? _logger;
 	private readonly HttpRequestOptionsKey<bool> _allowInteractiveLogin;
 
 	/// <inheritdoc/>
@@ -55,9 +58,15 @@ public sealed class HttpOidcAuthHandler : DelegatingHandler
 
 			if (response.StatusCode != HttpStatusCode.Unauthorized || attempt >= _config.MaxAttempts)
 			{
+				if (response.StatusCode == HttpStatusCode.Unauthorized && attempt >= _config.MaxAttempts)
+				{
+					_logger?.LogWarning("Received HTTP 401 from {Uri} after {Attempts} attempt(s). Auth failed.", request.RequestUri, attempt + 1);
+				}
 				return response;
 			}
 
+			_logger?.LogWarning("Received HTTP 401 from {Uri}. Invalidating token and retrying (attempt {Attempt}/{MaxAttempts}).", request.RequestUri, attempt + 1, _config.MaxAttempts);
+
 			// Mark this access token as invalid.
 			if (accessToken != null)
 			{
@@ -68,6 +77,7 @@ public sealed class HttpOidcAuthHandler : DelegatingHandler
 			string? nextAccessToken = await _authState.GetAccessTokenAsync(allowInteractiveLogin, cancellationToken);
 			if (String.Equals(accessToken, nextAccessToken, StringComparison.Ordinal))
 			{
+				_logger?.LogWarning("New access token is identical to the invalidated one. Cannot recover auth for {Uri}.", request.RequestUri);
 				return response;
 			}
 
diff --git a/Engine/Source/Programs/Shared/EpicGames.OIDC/ITokenStore.cs b/Engine/Source/Programs/Shared/EpicGames.OIDC/ITokenStore.cs
--- a/Engine/Source/Programs/Shared/EpicGames.OIDC/ITokenStore.cs
+++ b/Engine/Source/Programs/Shared/EpicGames.OIDC/ITokenStore.cs
@@ -4,6 +4,7 @@ using System;
 using System.Diagnostics.CodeAnalysis;
 using System.Runtime.InteropServices;
 using Microsoft.Extensions.DependencyInjection;
+using Microsoft.Extensions.Logging;
 
 #pragma warning disable CS1591 // Missing XML documentation on public types
 
@@ -11,10 +12,11 @@ namespace EpicGames.OIDC
 {
 	public interface ITokenStore : IDisposable
 	{
+		/// <summary>Always reads fresh from the backing store (disk/keychain).</summary>
 		public bool TryGetRefreshToken(string oidcProvider, [MaybeNullWhen(false)] out string refreshToken);
 
-		public void AddRefreshToken(string providerIdentifier, string refreshToken);
-		public void Save();
+		/// <summary>Persists immediately to the backing store.</summary>
+		public void AddRefreshToken(string oidcProvider, string refreshToken);
 	}
 
 	public static class TokenStoreFactory
@@ -37,11 +39,13 @@ namespace EpicGames.OIDC
 			throw new NotSupportedException("Unknown platform when attempting to create ITokenStore");
 		}
 
-		public static ITokenStore CreateTokenStore()
+		public static ITokenStore CreateTokenStore(ILoggerFactory? loggerFactory = null)
 		{
 			if (RuntimeInformation.IsOSPlatform(OSPlatform.Windows))
 			{
-				return new WindowsTokenStore();
+				return loggerFactory != null
+					? new WindowsTokenStore(loggerFactory.CreateLogger<WindowsTokenStore>())
+					: new WindowsTokenStore();
 			}
 			else if (RuntimeInformation.IsOSPlatform(OSPlatform.OSX))
 			{
diff --git a/Engine/Source/Programs/Shared/EpicGames.OIDC/InMemoryTokenStore.cs b/Engine/Source/Programs/Shared/EpicGames.OIDC/InMemoryTokenStore.cs
--- a/Engine/Source/Programs/Shared/EpicGames.OIDC/InMemoryTokenStore.cs
+++ b/Engine/Source/Programs/Shared/EpicGames.OIDC/InMemoryTokenStore.cs
@@ -1,6 +1,5 @@
 // Copyright Epic Games, Inc. All Rights Reserved.
 
-using System;
 using System.Collections.Generic;
 
 namespace EpicGames.OIDC;
@@ -8,13 +7,12 @@ namespace EpicGames.OIDC;
 /// <summary>
 /// An in-memory token store without encryption, intended for testing
 /// </summary>
-// ReSharper disable once ClassNeverInstantiated.Global
-public class InMemoryTokenStore : ITokenStore
+internal sealed class InMemoryTokenStore : ITokenStore
 {
 	private readonly Dictionary<string, string> _providerToRefreshToken = [];
 
 	/// <inheritdoc/>
-	public virtual bool TryGetRefreshToken(string oidcProvider, out string refreshToken)
+	public bool TryGetRefreshToken(string oidcProvider, out string refreshToken)
 	{
 		if (!_providerToRefreshToken.TryGetValue(oidcProvider, out string? result))
 		{
@@ -27,33 +25,14 @@ public class InMemoryTokenStore : ITokenStore
 	}
 
 	/// <inheritdoc/>
-	public virtual void AddRefreshToken(string providerIdentifier, string refreshToken)
+	public void AddRefreshToken(string oidcProvider, string refreshToken)
 	{
-		_providerToRefreshToken[providerIdentifier] = refreshToken;
-	}
-
-	/// <inheritdoc/>
-	public virtual void Save()
-	{
-		// No need - data exists only in-memory
-	}
-
-	/// <summary>
-	/// Dispose method
-	/// </summary>
-	/// <param name="disposing"></param>
-	protected virtual void Dispose(bool disposing)
-	{
-		if (disposing)
-		{
-			// no resources to release
-		}
+		_providerToRefreshToken[oidcProvider] = refreshToken;
 	}
 
 	/// <inheritdoc/>
 	public void Dispose()
 	{
-		Dispose(true);
-		GC.SuppressFinalize(this);
+		// No unmanaged resources; required by ITokenStore (IDisposable)
 	}
 }
diff --git a/Engine/Source/Programs/Shared/EpicGames.OIDC/MacOSTokenStore.cs b/Engine/Source/Programs/Shared/EpicGames.OIDC/MacOSTokenStore.cs
--- a/Engine/Source/Programs/Shared/EpicGames.OIDC/MacOSTokenStore.cs
+++ b/Engine/Source/Programs/Shared/EpicGames.OIDC/MacOSTokenStore.cs
@@ -9,7 +9,7 @@ using System.Text;
 
 namespace EpicGames.OIDC
 {
-	public sealed class MacOSTokenStore : ITokenStore
+	internal sealed class MacOSTokenStore : ITokenStore
 	{
 		public MacOSTokenStore()
 		{
@@ -29,13 +29,9 @@ namespace EpicGames.OIDC
 			return true;
 		}
 
-		public void AddRefreshToken(string providerIdentifier, string refreshToken)
-		{
-			KeychainHelper.AddKeychainItem(refreshToken, providerIdentifier);
-		}
-
-		public void Save()
+		public void AddRefreshToken(string oidcProvider, string refreshToken)
 		{
+			KeychainHelper.AddKeychainItem(refreshToken, oidcProvider);
 		}
 
 		public void Dispose()
diff --git a/Engine/Source/Programs/Shared/EpicGames.OIDC/OidcAuthState.cs b/Engine/Source/Programs/Shared/EpicGames.OIDC/OidcAuthState.cs
--- a/Engine/Source/Programs/Shared/EpicGames.OIDC/OidcAuthState.cs
+++ b/Engine/Source/Programs/Shared/EpicGames.OIDC/OidcAuthState.cs
@@ -138,6 +138,7 @@ public sealed class OidcAuthState(
 		{
 			if (TryGetCurrentAuthState(out AuthState? authState) && Equals(authState.TokenInfo?.AccessToken, accessToken))
 			{
+				_logger.LogInformation("Access token invalidated (likely received HTTP 401). Will attempt to acquire a new token.");
 				_currentAuthResult = null;
 			}
 		}
@@ -286,14 +287,18 @@ public sealed class OidcAuthState(
 			}
 			catch (Exception ex)
 			{
-				_logger.LogTrace(ex, "Unable to get access token; attempting login: {Message}", ex.Message);
+				_logger.LogWarning(ex, "Unable to get access token for provider '{Provider}': {Message}", oidcProvider, ex.Message);
 			}
 		}
 		if (result == null && interactive)
 		{
-			_logger.LogInformation("Logging in to provider '{Provider}'...", oidcProvider);
+			_logger.LogWarning("No valid access token obtained for provider '{Provider}'. Starting interactive login (browser).", oidcProvider);
 			result = await _oidcTokenManager.LoginAsync(oidcProvider, cancellationToken);
 		}
+		else if (result == null)
+		{
+			_logger.LogWarning("No valid access token obtained for provider '{Provider}' and interactive login is disabled.", oidcProvider);
+		}
 
 		return new AuthState(_clock, result, interactive);
 	}
diff --git a/Engine/Source/Programs/Shared/EpicGames.OIDC/OidcTokenManager.cs b/Engine/Source/Programs/Shared/EpicGames.OIDC/OidcTokenManager.cs
--- a/Engine/Source/Programs/Shared/EpicGames.OIDC/OidcTokenManager.cs
+++ b/Engine/Source/Programs/Shared/EpicGames.OIDC/OidcTokenManager.cs
@@ -25,6 +25,13 @@ using Microsoft.Extensions.Options;
 
 namespace EpicGames.OIDC
 {
+	/// <summary>
+	/// Result of an OIDC token refresh operation.
+	/// </summary>
+	public record TokenRefreshResult(
+		bool IsError, string? Error, string? ErrorDescription,
+		string? AccessToken, string? RefreshToken, DateTimeOffset AccessTokenExpiration);
+
 	public interface IOidcTokenManager
 	{
 		public Task<IOidcTokenInfo> LoginAsync(string providerIdentifier, CancellationToken cancellationToken = default);
@@ -64,15 +71,7 @@ namespace EpicGames.OIDC
 		public OidcTokenManager(IOptionsMonitor<OidcTokenOptions> settings, ITokenStore tokenStore, IOidcTokenClientFactory tokenClientFactory, List<string>? allowedProviders = null)
 			: this(tokenStore)
 		{
-			Dictionary<string, string> refreshTokens = [];
 			List<string> providerKeys = [.. settings.CurrentValue.Providers.Keys.Where(k => allowedProviders == null || allowedProviders.Contains(k))];
-			foreach (string oidcProvider in providerKeys)
-			{
-				if (tokenStore.TryGetRefreshToken(oidcProvider, out string? refreshToken))
-				{
-					refreshTokens.TryAdd(oidcProvider, refreshToken);
-				}
-			}
 
 			foreach ((string key, ProviderInfo providerInfo) in settings.CurrentValue.Providers)
 			{
@@ -82,39 +81,23 @@ namespace EpicGames.OIDC
 				}
 
 				OidcTokenClient tokenClient = tokenClientFactory.CreateTokenClient(key, providerInfo, settings.CurrentValue.LoginTimeout);
-
-				if (refreshTokens.TryGetValue(key, out string? refreshToken))
-				{
-					tokenClient.SetRefreshToken(refreshToken);
-				}
-
 				_tokenClients.Add(key, tokenClient);
 			}
 		}
 
-		public static OidcTokenManager CreateTokenManager(IConfiguration providerConfiguration, ITokenStore tokenStore, List<string>? allowedProviders = null)
+		public static OidcTokenManager CreateTokenManager(IConfiguration providerConfiguration, ITokenStore tokenStore, List<string>? allowedProviders = null, ILoggerFactory? loggerFactory = null)
 		{
-			return new OidcTokenManager(providerConfiguration, tokenStore, allowedProviders);
+			return new OidcTokenManager(providerConfiguration, tokenStore, allowedProviders, loggerFactory);
 		}
 
-		private OidcTokenManager(IConfiguration providerConfiguration, ITokenStore tokenStore, List<string>? allowedProviders = null)
+		private OidcTokenManager(IConfiguration providerConfiguration, ITokenStore tokenStore, List<string>? allowedProviders = null, ILoggerFactory? loggerFactory = null)
 			: this(tokenStore)
 		{
-			Dictionary<string, string> refreshTokens = [];
-
 			OidcTokenOptions options = new();
 			providerConfiguration.Bind(options);
 
 			List<string> providerKeys = [.. options.Providers.Keys.Where(k => allowedProviders == null || allowedProviders.Contains(k))];
 
-			foreach (string oidcProvider in providerKeys)
-			{
-				if (tokenStore.TryGetRefreshToken(oidcProvider, out string? refreshToken))
-				{
-					refreshTokens.TryAdd(oidcProvider, refreshToken);
-				}
-			}
-
 			foreach ((string key, ProviderInfo providerInfo) in options.Providers)
 			{
 				if (providerKeys != null && !providerKeys.Contains(key))
@@ -122,13 +105,7 @@ namespace EpicGames.OIDC
 					continue;
 				}
 
-				OidcTokenClient tokenClient = new(key, providerInfo, options.LoginTimeout, _tokenStore);
-
-				if (refreshTokens.TryGetValue(key, out string? refreshToken))
-				{
-					tokenClient.SetRefreshToken(refreshToken);
-				}
-
+				OidcTokenClient tokenClient = new(key, providerInfo, options.LoginTimeout, _tokenStore, loggerFactory?.CreateLogger<OidcTokenClient>());
 				_tokenClients.Add(key, tokenClient);
 			}
 		}
@@ -371,6 +348,12 @@ namespace EpicGames.OIDC
 
 	public class OidcTokenClient
 	{
+		/// <summary>Max seconds the refresh HTTP call is allowed before cancellation.</summary>
+		private const int RefreshHttpTimeoutSeconds = 30;
+
+		/// <summary>Max seconds to wait for another process to release the refresh lock.</summary>
+		private const int LockAcquireTimeoutSeconds = 60;
+
 		public string Name => _name;
 		public string DisplayName => _providerInfo.DisplayName;
 
@@ -383,7 +366,6 @@ namespace EpicGames.OIDC
 		private readonly string _clientId;
 		private readonly string _scopes;
 
-		private string? _refreshToken;
 		private string? _accessToken;
 		private DateTimeOffset _tokenExpiry;
 
@@ -417,7 +399,11 @@ namespace EpicGames.OIDC
 			_useDiscoveryDocument = providerInfo.UseDiscoveryDocument;
 
 			_genericErrorInformation = providerInfo.GenericErrorInformation;
-			_tokenStore.TryGetRefreshToken(name, out _refreshToken);
+			if (_logger?.IsEnabled(LogLevel.Debug) == true)
+			{
+				_tokenStore.TryGetRefreshToken(name, out string? existingToken);
+				_logger.LogDebug("OidcTokenClient created for provider '{Name}'. Refresh token from store: {HasToken}", name, !String.IsNullOrEmpty(existingToken));
+			}
 		}
 
 		private async Task<OidcClientOptions> BuildClientOptionsAsync(Uri redirectUri, CancellationToken cancellationToken)
@@ -457,6 +443,7 @@ namespace EpicGames.OIDC
 
 		public async Task<IOidcTokenInfo> LoginAsync(CancellationToken cancellationToken)
 		{
+			_logger?.LogInformation("Starting interactive login for provider '{Name}'.", _name);
 			IOidcTokenInfo tokenInfo = await DoLoginAsync(cancellationToken);
 
 			if (String.IsNullOrEmpty(tokenInfo.RefreshToken))
@@ -464,20 +451,18 @@ namespace EpicGames.OIDC
 				throw new Exception("No refresh token was provided in response.");
 			}
 
-			_refreshToken = tokenInfo.RefreshToken;
 			_accessToken = tokenInfo.AccessToken;
 			_tokenExpiry = tokenInfo.TokenExpiry;
 
-			if (_refreshToken != null)
-			{
-				_tokenStore.AddRefreshToken(_name, _refreshToken);
-				_tokenStore.Save();
-			}
+			_logger?.LogDebug("Saving refresh token to store after interactive login for provider '{Name}'.", _name);
+			using IDisposable _ = await AcquireRefreshLockAsync(cancellationToken);
+			_tokenStore.AddRefreshToken(_name, tokenInfo.RefreshToken);
 
+			_logger?.LogInformation("Interactive login completed for provider '{Name}'. Token expiry: {Expiry}", _name, tokenInfo.TokenExpiry);
 			return tokenInfo;
 		}
 
-		public virtual async Task<IOidcTokenInfo> DoLoginAsync(CancellationToken cancellationToken)
+		private async Task<IOidcTokenInfo> DoLoginAsync(CancellationToken cancellationToken)
 		{
 			// setup a local http server to listen for the result of the login
 			LoginResult? loginResult = null;
@@ -736,46 +721,116 @@ namespace EpicGames.OIDC
 			return version.Contains("microsoft", StringComparison.InvariantCultureIgnoreCase) || version.Contains("wsl2", StringComparison.InvariantCultureIgnoreCase);
 		}
 
-		private async Task<OidcTokenInfo?> TryDoRefreshTokenAsync(string inRefreshToken, CancellationToken cancellationToken)
+		/// <summary>
+		/// Performs the OIDC token refresh via HTTP. Virtual for testability.
+		/// </summary>
+		protected virtual async Task<TokenRefreshResult> DoRefreshAsync(string refreshToken, CancellationToken cancellationToken)
 		{
-			// redirect uri is not used for refresh tokens so we can just pick one of them to configure the client
 			OidcClientOptions options = await BuildClientOptionsAsync(_redirectUris.First(), cancellationToken);
 			OidcClient oidcClient = new(options);
+			RefreshTokenResult r = await oidcClient.RefreshTokenAsync(refreshToken, cancellationToken: cancellationToken);
+			return new TokenRefreshResult(r.IsError, r.Error, r.ErrorDescription, r.AccessToken, r.RefreshToken, r.AccessTokenExpiration);
+		}
 
-			// use the refresh token to acquire a new access token
-			RefreshTokenResult refreshTokenResult = await oidcClient!.RefreshTokenAsync(inRefreshToken, cancellationToken: cancellationToken);
+		/// <summary>
+		/// Acquire machine-wide lock for the refresh operation. Virtual for testability.
+		/// </summary>
+		protected virtual async Task<IDisposable> AcquireRefreshLockAsync(CancellationToken cancellationToken)
+		{
+			Mutex mutex = new(false, @"Global\OidcRefreshLock");
+			TimeSpan timeout = TimeSpan.FromSeconds(LockAcquireTimeoutSeconds);
 
-			if (refreshTokenResult.IsError)
+			Stopwatch sw = Stopwatch.StartNew();
+			bool acquired;
+			try
 			{
-				if (refreshTokenResult.Error == "invalid_grant")
+				acquired = await Task.Run(() => mutex.WaitOne(timeout), cancellationToken);
+			}
+			catch (AbandonedMutexException)
+			{
+				// Previous holder crashed — we now own the mutex.
+				acquired = true;
+			}
+			catch
+			{
+				mutex.Dispose();
+				throw;
+			}
+
+			if (!acquired)
+			{
+				mutex.Dispose();
+				throw new TimeoutException($"Timed out after {LockAcquireTimeoutSeconds}s waiting to acquire OIDC refresh lock.");
+			}
+
+			_logger?.LogInformation("Acquired refresh lock for '{Name}' (waited {Elapsed}ms).", _name, sw.ElapsedMilliseconds);
+			return new MutexReleaser(mutex);
+		}
+
+		private sealed class MutexReleaser : IDisposable
+		{
+			private Mutex? _mutex;
+
+			public MutexReleaser(Mutex mutex) => _mutex = mutex;
+
+			public void Dispose()
+			{
+				if (_mutex != null)
 				{
-					// the refresh token is no logger valid, resetting it and treating us as not logged in
-					_refreshToken = null;
-					return null;
+					_mutex.ReleaseMutex();
+					_mutex.Dispose();
+					_mutex = null;
 				}
-				else
+			}
+		}
+
+		private async Task<OidcTokenInfo?> TryDoRefreshTokenAsync(string refreshToken, CancellationToken cancellationToken)
+		{
+			TokenRefreshResult result = await DoRefreshAsync(refreshToken, cancellationToken);
+
+			// Fallback for non-cooperating processes (old binaries without new lock)
+			if (result is { IsError: true, Error: "invalid_grant" })
+			{
+				_logger?.LogWarning("Refresh token for provider '{Name}' returned invalid_grant. May have been rotated by another process.", _name);
+
+				if (_tokenStore.TryGetRefreshToken(_name, out string? freshToken) &&
+				    !String.Equals(freshToken, refreshToken, StringComparison.Ordinal))
 				{
-					throw new Exception($"Error using the refresh token: {refreshTokenResult.Error} , details: {refreshTokenResult.ErrorDescription}");
+					_logger?.LogInformation("Found rotated refresh token for provider '{Name}'. Retrying.", _name);
+					result = await DoRefreshAsync(freshToken, cancellationToken);
 				}
 			}
 
-			_refreshToken = refreshTokenResult.RefreshToken;
-			_accessToken = refreshTokenResult.AccessToken;
-			_tokenExpiry = refreshTokenResult.AccessTokenExpiration;
+			if (result.IsError)
+			{
+				if (result.Error == "invalid_grant")
+				{
+					_logger?.LogWarning("Refresh token for provider '{Name}' still invalid. User must re-authenticate.", _name);
+					return null;
+				}
+				_logger?.LogWarning("Refresh token error for provider '{Name}': {Error} - {Description}", _name, result.Error, result.ErrorDescription);
+				throw new Exception($"Error using the refresh token: {result.Error} , details: {result.ErrorDescription}");
+			}
 
-			// refresh tokens are always one time use only so we need to store this new refresh token we got so it can be used the next time
-			if (String.IsNullOrEmpty(refreshTokenResult.RefreshToken))
+			// Persist new token to disk immediately
+			if (String.IsNullOrEmpty(result.RefreshToken))
 			{
 				throw new Exception("No refresh token was provided in response to refresh.");
 			}
-			_tokenStore.AddRefreshToken(_name, _refreshToken);
-			_tokenStore.Save();
+			_logger?.LogDebug("Saving new refresh token to store for provider '{Name}'.", _name);
+			_tokenStore.AddRefreshToken(_name, result.RefreshToken);
+
+			// Cache access token in memory (short-lived)
+			_accessToken = result.AccessToken;
+			_tokenExpiry = result.AccessTokenExpiration;
+
+			_logger?.LogInformation("Successfully refreshed token for provider '{Name}'. New expiry: {Expiry}", _name, result.AccessTokenExpiration);
 
 			return new OidcTokenInfo
 			{
-				RefreshToken = refreshTokenResult.RefreshToken,
-				AccessToken = refreshTokenResult.AccessToken,
-				TokenExpiry = refreshTokenResult.AccessTokenExpiration
+				RefreshToken = result.RefreshToken,
+				AccessToken = result.AccessToken,
+				TokenExpiry = result.AccessTokenExpiration
 			};
 		}
 
@@ -811,30 +866,48 @@ namespace EpicGames.OIDC
 
 		public async Task<IOidcTokenInfo?> TryGetAccessTokenAsync(CancellationToken cancellationToken)
 		{
-			if (String.IsNullOrEmpty(_refreshToken))
+			// Presence check - reads fresh from store
+			if (!_tokenStore.TryGetRefreshToken(_name, out string? refreshToken) || String.IsNullOrEmpty(refreshToken))
 			{
+				_logger?.LogDebug("No refresh token for provider '{Name}'.", _name);
 				throw new NotLoggedInException();
 			}
 
+			// Check cached access token (short-lived, OK to cache)
 			OidcTokenInfo tokenInfo = new()
 			{
-				RefreshToken = _refreshToken,
+				RefreshToken = refreshToken,
 				AccessToken = _accessToken,
 				TokenExpiry = _tokenExpiry
 			};
 
-			// Ensure token is valid for at least another two minutes
 			if (tokenInfo.IsValid(DateTime.Now.AddMinutes(-2)))
 			{
 				return tokenInfo;
 			}
 
-			return await TryDoRefreshTokenAsync(_refreshToken, cancellationToken);
+			_logger?.LogInformation("Access token for provider '{Name}' expired or expiring. Attempting refresh.", _name);
+
+			// Machine-wide lock - only one process refreshes at a time
+			using IDisposable refreshLock = await AcquireRefreshLockAsync(cancellationToken);
+
+			// Re-read from disk inside lock — another process/thread may have refreshed while we waited
+			if (!_tokenStore.TryGetRefreshToken(_name, out refreshToken) || String.IsNullOrEmpty(refreshToken))
+			{
+				throw new NotLoggedInException();
+			}
+
+			// Bound the HTTP refresh to prevent holding the lock indefinitely
+			using CancellationTokenSource refreshCts = CancellationTokenSource.CreateLinkedTokenSource(cancellationToken);
+			refreshCts.CancelAfter(TimeSpan.FromSeconds(RefreshHttpTimeoutSeconds));
+
+			return await TryDoRefreshTokenAsync(refreshToken, refreshCts.Token);
 		}
 
 		public OidcStatus GetStatus()
 		{
-			return GetStatus(_refreshToken, _accessToken, _tokenExpiry);
+			_tokenStore.TryGetRefreshToken(_name, out string? rt);
+			return GetStatus(rt, _accessToken, _tokenExpiry);
 		}
 
 		public static OidcStatus GetStatus(string? refreshToken, string? accessToken, DateTimeOffset accessTokenExpiry)
@@ -856,12 +929,8 @@ namespace EpicGames.OIDC
 
 			return OidcStatus.Connected;
 		}
-
-		public void SetRefreshToken(string inRefreshToken)
-		{
-			_refreshToken = inRefreshToken;
-		}
 	}
+
 	public class LoginFailedException : Exception
 	{
 		public LoginFailedException(string message, string errorDescription) : base(message)
diff --git a/Engine/Source/Programs/Shared/EpicGames.OIDC/Properties/AssemblyInfo.cs b/Engine/Source/Programs/Shared/EpicGames.OIDC/Properties/AssemblyInfo.cs
new file mode 100644
--- /dev/null
+++ b/Engine/Source/Programs/Shared/EpicGames.OIDC/Properties/AssemblyInfo.cs
@@ -0,0 +1,5 @@
+// Copyright Epic Games, Inc. All Rights Reserved.
+
+using System.Runtime.CompilerServices;
+
+[assembly: InternalsVisibleTo("EpicGames.Horde.Tests")]
diff --git a/Engine/Source/Programs/Shared/EpicGames.OIDC/TokenStore.cs b/Engine/Source/Programs/Shared/EpicGames.OIDC/TokenStore.cs
--- a/Engine/Source/Programs/Shared/EpicGames.OIDC/TokenStore.cs
+++ b/Engine/Source/Programs/Shared/EpicGames.OIDC/TokenStore.cs
@@ -43,15 +43,12 @@ namespace EpicGames.OIDC
 	/// <summary>
 	/// A generic token store that saves the offline token in a file on disk, this works on any platform
 	/// </summary>
-	public class FilesystemTokenStore : ITokenStore, IDisposable
+	internal sealed class FilesystemTokenStore : ITokenStore
 	{
-		private readonly Dictionary<string, byte[]> _providerToRefreshToken;
-		private readonly List<string> _dirtyProviders = [];
 		private readonly SymmetricAlgorithm _crypt = Aes.Create();
 
 		public FilesystemTokenStore()
 		{
-			_providerToRefreshToken = ReadStoreFromDisk();
 		}
 
 		private static FileInfo GetStorePath()
@@ -59,13 +56,14 @@ namespace EpicGames.OIDC
 			return new FileInfo(Path.Combine(Environment.GetFolderPath(Environment.SpecialFolder.ApplicationData), "Epic", "UnrealEngine", "Common", "OidcToken", "tokenStore.dat"));
 		}
 
-		private Dictionary<string, byte[]> ReadStoreFromDisk()
+		private static void ReadStoreFromDisk(out Dictionary<string, byte[]> providers, out byte[]? iv)
 		{
 			FileInfo fi = GetStorePath();
 			if (!fi.Exists)
 			{
-				// if we have no store on disk then we just initialize it to empty
-				return [];
+				providers = [];
+				iv = null;
+				return;
 			}
 
 			using FileStream fs = fi.OpenRead();
@@ -81,22 +79,21 @@ namespace EpicGames.OIDC
 
 			if (state == null)
 			{
-				// if we fail to deserialize the state just drop it, will mean users will need to login again
-				return [];
+				providers = [];
+				iv = null;
+				return;
 			}
-			Dictionary<string, byte[]> providers = [];
 
+			providers = [];
 			foreach ((string key, string value) in state.Providers)
 			{
 				providers[key] = Convert.FromBase64String(value);
 			}
 
-			_crypt.IV = Convert.FromBase64String(state.Key);
-
-			return providers;
+			iv = Convert.FromBase64String(state.Key);
 		}
 
-		private void SaveStoreToDisk()
+		private static void WriteStoreToDisk(byte[] iv, Dictionary<string, byte[]> providers)
 		{
 			FileInfo fi = GetStorePath();
 
@@ -105,67 +102,46 @@ namespace EpicGames.OIDC
 				Directory.CreateDirectory(fi.Directory!.FullName);
 			}
 
-			lock (_dirtyProviders)
+			string tempFile = Path.GetTempFileName();
 			{
-				// no providers have changed, do not touch the state file
-				if (_dirtyProviders.Count == 0)
-				{
-					return;
-				}
-
-				// read back the state of all providers but only overwrite the state of the ones we have actually got new state for (are dirty)
-				Dictionary<string, byte[]> providers = ReadStoreFromDisk();
-
-				foreach (string providerId in _dirtyProviders)
-				{
-					providers[providerId] = _providerToRefreshToken[providerId];
-				}
-
-				string tempFile = Path.GetTempFileName();
-				{
-					using FileStream fs = new(tempFile, FileMode.Create, FileAccess.Write);
-					using Utf8JsonWriter writer = new(fs);
-					JsonSerializer.Serialize<TokenStoreState>(writer, new TokenStoreState(_crypt.IV, providers), TokenStoreStateContext.Default.TokenStoreState);
-				}
-
-				File.Move(tempFile, fi.FullName, true);
-				_dirtyProviders.Clear();
+				using FileStream fs = new(tempFile, FileMode.Create, FileAccess.Write);
+				using Utf8JsonWriter writer = new(fs);
+				JsonSerializer.Serialize<TokenStoreState>(writer, new TokenStoreState(iv, providers), TokenStoreStateContext.Default.TokenStoreState);
 			}
+
+			File.Move(tempFile, fi.FullName, true);
 		}
 
 		public bool TryGetRefreshToken(string oidcProvider, out string refreshToken)
 		{
-			if (!_providerToRefreshToken.TryGetValue(oidcProvider, out byte[]? encryptedToken))
+			ReadStoreFromDisk(out Dictionary<string, byte[]> providers, out byte[]? iv);
+
+			if (!providers.TryGetValue(oidcProvider, out byte[]? encryptedToken))
 			{
 				refreshToken = "";
 				return false;
 			}
 
-			ICryptoTransform decryptor = _crypt.CreateDecryptor(GetSeed(), _crypt.IV);
+			ICryptoTransform decryptor = _crypt.CreateDecryptor(GetSeed(), iv ?? _crypt.IV);
 			byte[] bytes = decryptor.TransformFinalBlock(encryptedToken, 0, encryptedToken.Length);
 			refreshToken = Encoding.Unicode.GetString(bytes);
 			return true;
 		}
 
-		public void AddRefreshToken(string providerIdentifier, string refreshToken)
+		public void AddRefreshToken(string oidcProvider, string refreshToken)
 		{
+			ReadStoreFromDisk(out Dictionary<string, byte[]> providers, out byte[]? iv);
+			byte[] encryptIv = iv ?? _crypt.IV;
+
 			byte[] bytes = Encoding.Unicode.GetBytes(refreshToken);
 #pragma warning disable CA5401 // Do not use CreateEncryptor with non-default IV
-			ICryptoTransform encryptor = _crypt.CreateEncryptor(GetSeed(), _crypt.IV);
+			ICryptoTransform encryptor = _crypt.CreateEncryptor(GetSeed(), encryptIv);
 #pragma warning restore CA5401 // Do not use CreateEncryptor with non-default IV
 
 			byte[] encryptedToken = encryptor.TransformFinalBlock(bytes, 0, bytes.Length);
-			_providerToRefreshToken[providerIdentifier] = encryptedToken;
-
-			lock (_dirtyProviders)
-			{
-				_dirtyProviders.Add(providerIdentifier);
-			}
-		}
 
-		public void Save()
-		{
-			SaveStoreToDisk();
+			providers[oidcProvider] = encryptedToken;
+			WriteStoreToDisk(encryptIv, providers);
 		}
 
 		static byte[] GetSeed()
@@ -177,19 +153,9 @@ namespace EpicGames.OIDC
 			];
 		}
 
-		protected virtual void Dispose(bool disposing)
-		{
-			if (disposing)
-			{
-				SaveStoreToDisk();
-				_crypt.Dispose();
-			}
-		}
-
 		public void Dispose()
 		{
-			Dispose(true);
-			GC.SuppressFinalize(this);
+			_crypt.Dispose();
 		}
 	}
 }
\ No newline at end of file
diff --git a/Engine/Source/Programs/Shared/EpicGames.OIDC/WindowsTokenStore.cs b/Engine/Source/Programs/Shared/EpicGames.OIDC/WindowsTokenStore.cs
--- a/Engine/Source/Programs/Shared/EpicGames.OIDC/WindowsTokenStore.cs
+++ b/Engine/Source/Programs/Shared/EpicGames.OIDC/WindowsTokenStore.cs
@@ -5,6 +5,7 @@ using System.Collections.Generic;
 using System.ComponentModel;
 using System.IO;
 using System.Runtime.InteropServices;
+using System.Security.Cryptography;
 using System.Text;
 using System.Text.Json;
 using System.Text.Json.Serialization;
@@ -40,29 +41,35 @@ namespace EpicGames.OIDC
 		}
 	}
 
-	public class WindowsTokenStore : ITokenStore, IDisposable
+	internal sealed class WindowsTokenStore : ITokenStore
 	{
 		private readonly ILogger<WindowsTokenStore>? _logger = null;
+		private readonly string? _storePathOverride;
 
-		private readonly Dictionary<string, byte[]> _providerToRefreshToken = [];
-		private readonly List<string> _dirtyProviders = [];
+		/// <summary>First 8 hex chars of SHA-256 — enough for log correlation without leaking the token.</summary>
+		private static string TokenFingerprint(string token) => Convert.ToHexString(SHA256.HashData(Encoding.UTF8.GetBytes(token)))[..8];
 
 		public WindowsTokenStore()
 		{
-			_providerToRefreshToken = ReadStoreFromDisk();
 		}
 
 		[ActivatorUtilitiesConstructor]
 		public WindowsTokenStore(ILogger<WindowsTokenStore> logger)
 		{
 			_logger = logger;
+		}
 
-			_providerToRefreshToken = ReadStoreFromDisk();
+		/// <summary>Test-only constructor that redirects the store to a custom path.</summary>
+		internal WindowsTokenStore(ILogger<WindowsTokenStore>? logger, string storePath)
+		{
+			_logger = logger;
+			_storePathOverride = storePath;
 		}
 
-		private static FileInfo GetStorePath()
+		private FileInfo GetStorePath()
 		{
-			return new FileInfo(Path.Combine(Environment.GetFolderPath(Environment.SpecialFolder.LocalApplicationData), "UnrealEngine", "Common", "OidcToken", "oidcTokenStore.dat"));
+			// Always return a fresh FileInfo to avoid stale cached Exists state.
+			return new FileInfo(_storePathOverride ?? Path.Combine(Environment.GetFolderPath(Environment.SpecialFolder.LocalApplicationData), "UnrealEngine", "Common", "OidcToken", "oidcTokenStore.dat"));
 		}
 
 		private Dictionary<string, byte[]> ReadStoreFromDisk()
@@ -71,41 +78,62 @@ namespace EpicGames.OIDC
 			if (!fi.Exists)
 			{
 				_logger?.LogDebug("No existing token store found at {Path}. Assuming empty store.", fi.FullName);
-				// if we have no store on disk then we just initialize it to empty
 				return [];
 			}
 
-			using FileStream fs = fi.Open(FileMode.Open, FileAccess.Read, FileShare.Read);
-			using TextReader tr = new StreamReader(fs);
+			string json;
+			try
+			{
+				using FileStream fs = fi.Open(FileMode.Open, FileAccess.Read, FileShare.Read);
+				using TextReader tr = new StreamReader(fs);
+				json = tr.ReadToEnd();
+			}
+			catch (IOException ex)
+			{
+				_logger?.LogWarning(ex, "Unable to read token store file at {Path}. Assuming empty store.", fi.FullName);
+				return [];
+			}
+			catch (UnauthorizedAccessException ex)
+			{
+				_logger?.LogWarning(ex, "Access denied reading token store file at {Path}. Assuming empty store.", fi.FullName);
+				return [];
+			}
 
 			WindowsTokenStoreState? state;
 			try
 			{
-				state = JsonSerializer.Deserialize(tr.ReadToEnd(), WindowsTokenStoreStateContext.Default.WindowsTokenStoreState);
+				state = JsonSerializer.Deserialize(json, WindowsTokenStoreStateContext.Default.WindowsTokenStoreState);
 			}
-			catch (JsonException)
+			catch (JsonException ex)
 			{
+				_logger?.LogWarning(ex, "Failed to deserialize token store JSON. Dropping the existing state.");
 				state = null;
 			}
 
 			if (state == null)
 			{
-				_logger?.LogDebug("Failed to deserialize state. Dropping the existing state.");
-				// if we fail to deserialize the state just drop it, will mean users will need to login again
+				_logger?.LogDebug("Token store deserialized to null. Dropping the existing state.");
 				return [];
 			}
 
 			Dictionary<string, byte[]> providers = [];
-
 			foreach ((string key, string value) in state.Providers)
 			{
-				providers[key] = Convert.FromBase64String(value);
+				try
+				{
+					providers[key] = Convert.FromBase64String(value);
+				}
+				catch (FormatException ex)
+				{
+					_logger?.LogWarning(ex, "Corrupt Base64 data for provider '{Provider}'. Skipping entry.", key);
+				}
 			}
 
+			_logger?.LogDebug("Loaded token store with {Count} provider(s).", providers.Count);
 			return providers;
 		}
 
-		private void SaveStoreToDisk()
+		private void WriteStoreToDisk(Dictionary<string, byte[]> providers, ILogger? logger)
 		{
 			FileInfo fi = GetStorePath();
 
@@ -114,125 +142,119 @@ namespace EpicGames.OIDC
 				Directory.CreateDirectory(fi.Directory!.FullName);
 			}
 
-			lock (_dirtyProviders)
+			string tempFile = Path.GetTempFileName();
+			try
 			{
-				// no providers have changed, do not touch the state file
-				if (_dirtyProviders.Count == 0)
 				{
-					return;
+					using FileStream fs = new(tempFile, FileMode.Create, FileAccess.Write);
+					using Utf8JsonWriter writer = new(fs);
+					JsonSerializer.Serialize(writer, new WindowsTokenStoreState(providers), WindowsTokenStoreStateContext.Default.WindowsTokenStoreState);
 				}
 
-				using Mutex mutex = new(false, "oidcTokenStoreDat");
-
+				File.Move(tempFile, fi.FullName, true);
+			}
+			catch
+			{
 				try
 				{
-					mutex.WaitOne();
+					File.Delete(tempFile);
 				}
-				catch (AbandonedMutexException)
+				catch (Exception deleteEx)
 				{
-
+					logger?.LogDebug(deleteEx, "Failed to clean up temp file {TempFile}.", tempFile);
 				}
 
-				// read back the state of all providers but only overwrite the state of the ones we have actually got new state for (are dirty)
-				Dictionary<string, byte[]> providers = ReadStoreFromDisk();
-
-				foreach (string providerId in _dirtyProviders)
-				{
-					providers[providerId] = _providerToRefreshToken[providerId];
-				}
-				string tempFile = Path.GetTempFileName();
-				{
-					using FileStream fs = new(tempFile, FileMode.Create, FileAccess.Write);
-					using Utf8JsonWriter writer = new(fs);
-					JsonSerializer.Serialize(writer, new WindowsTokenStoreState(providers), WindowsTokenStoreStateContext.Default.WindowsTokenStoreState);
-				}
-
-				File.Move(tempFile, fi.FullName, true);
-
-				mutex.ReleaseMutex();
-
-				_dirtyProviders.Clear();
+				throw;
 			}
 		}
 
 		public bool TryGetRefreshToken(string oidcProvider, out string refreshToken)
 		{
-			if (!_providerToRefreshToken.TryGetValue(oidcProvider, out byte[]? encryptedToken))
-			{
-				refreshToken = "";
-				return false;
-			}
-
 			try
 			{
-				byte[] bytes = CryptProtectDataHelper.DoCryptUnprotectData(encryptedToken, $"OidcToken-{oidcProvider}", GetEntropy(oidcProvider));
-				refreshToken = Encoding.Unicode.GetString(bytes);
-
-				return true;
-			}
-			catch (Win32Exception e)
-			{
-				if (e.NativeErrorCode == 13) // data is invalid
+				Dictionary<string, byte[]> providers = ReadStoreFromDisk();
+				if (!providers.TryGetValue(oidcProvider, out byte[]? encryptedToken))
 				{
-					// unable to decrypt the data, ignore it
+					_logger?.LogDebug("No refresh token found in store for provider '{Provider}'.", oidcProvider);
 					refreshToken = "";
-					_logger?.LogDebug("Unable to decrypt refresh token. Ignoring.");
 					return false;
 				}
-				if (e.NativeErrorCode == unchecked((int)0x8009000B)) // key not valid for use in specified state
+
+				try
 				{
-					// unable to decrypt the data, ignore it
+					byte[] bytes = CryptProtectDataHelper.DoCryptUnprotectData(encryptedToken, $"OidcToken-{oidcProvider}", GetEntropy(oidcProvider));
+					refreshToken = Encoding.Unicode.GetString(bytes);
+					_logger?.LogInformation("Read refresh token for '{Provider}': {Fingerprint}", oidcProvider, TokenFingerprint(refreshToken));
+					return true;
+				}
+				catch (Win32Exception e)
+				{
+					_logger?.LogWarning(e, "Unable to decrypt refresh token for provider '{Provider}' (Win32 error 0x{ErrorCode:X}). Discarding stored token.", oidcProvider, e.NativeErrorCode);
 					refreshToken = "";
-					_logger?.LogDebug("Unable to decrypt refresh token, key not valid for use in specified state. Ignoring.");
 					return false;
 				}
-				throw;
+			}
+			catch (Exception ex)
+			{
+				_logger?.LogError(ex, "Failed to read refresh token for provider '{Provider}'. Treating as not logged in.", oidcProvider);
+				refreshToken = "";
+				return false;
 			}
 		}
 
 		private static byte[] GetEntropy(string oidcProvider)
 		{
-			byte[] providerBytes = Encoding.UTF8.GetBytes(oidcProvider);
-
-			return providerBytes;
+			return Encoding.UTF8.GetBytes(oidcProvider);
 		}
 
-		public void AddRefreshToken(string providerIdentifier, string refreshToken)
+		public void AddRefreshToken(string oidcProvider, string refreshToken)
 		{
-			byte[] bytes = Encoding.Unicode.GetBytes(refreshToken);
-			byte[] encryptedToken = CryptProtectDataHelper.DoCryptProtectData(bytes, $"OidcToken-{providerIdentifier}", GetEntropy(providerIdentifier));
-
-			_providerToRefreshToken[providerIdentifier] = encryptedToken;
-
-			lock (_dirtyProviders)
+			try
 			{
-				_dirtyProviders.Add(providerIdentifier);
-			}
-		}
+				byte[] bytes = Encoding.Unicode.GetBytes(refreshToken);
+				byte[] encryptedToken = CryptProtectDataHelper.DoCryptProtectData(bytes, $"OidcToken-{oidcProvider}", GetEntropy(oidcProvider));
 
-		public void Save()
-		{
-			SaveStoreToDisk();
-		}
+				// Acquire the same named Mutex that older binaries use for store writes,
+				// so a new binary and an old binary cannot write concurrently.
+				using Mutex mutex = new(false, "oidcTokenStoreDat");
+				try
+				{
+					mutex.WaitOne();
+				}
+				catch (AbandonedMutexException)
+				{
+					// Previous holder crashed — we now own the mutex, safe to proceed.
+				}
 
-		protected virtual void Dispose(bool disposing)
-		{
-			if (disposing)
+				try
+				{
+					Dictionary<string, byte[]> providers = ReadStoreFromDisk();
+					providers[oidcProvider] = encryptedToken;
+					WriteStoreToDisk(providers, _logger);
+				}
+				finally
+				{
+					mutex.ReleaseMutex();
+				}
+
+				_logger?.LogInformation("Saved refresh token for '{Provider}': {Fingerprint}", oidcProvider, TokenFingerprint(refreshToken));
+			}
+			catch (Exception ex)
 			{
-				SaveStoreToDisk();
+				_logger?.LogError(ex, "Failed to persist refresh token for provider '{Provider}'", oidcProvider);
 			}
 		}
 
 		public void Dispose()
 		{
-			Dispose(true);
-			GC.SuppressFinalize(this);
 		}
 	}
 
 #pragma warning disable IDE1006 // Pinvoke code doesnt use the same naming conventions as C#
 	static class CryptProtectDataHelper
 	{
+		[DllImport("kernel32.dll", SetLastError = true)]
+		private static extern IntPtr LocalFree(IntPtr hMem);
 
 		[StructLayout(LayoutKind.Sequential, CharSet = CharSet.Unicode)]
 		private struct DataBlob
@@ -354,9 +376,19 @@ namespace EpicGames.OIDC
 				entropyHandle.Free();
 			}
 
-			byte[] buf = new byte[dataOut.cbData];
-			Marshal.Copy(dataOut.pbData, buf, 0, dataOut.cbData);
-			return buf;
+			try
+			{
+				byte[] buf = new byte[dataOut.cbData];
+				Marshal.Copy(dataOut.pbData, buf, 0, dataOut.cbData);
+				return buf;
+			}
+			finally
+			{
+				if (dataOut.pbData != IntPtr.Zero)
+				{
+					LocalFree(dataOut.pbData);
+				}
+			}
 		}
 
 		public static byte[] DoCryptUnprotectData(byte[] dataToDecrypt, string description, byte[] entropy)
@@ -395,9 +427,19 @@ namespace EpicGames.OIDC
 				entropyHandle.Free();
 			}
 
-			byte[] buf = new byte[dataOut.cbData];
-			Marshal.Copy(dataOut.pbData, buf, 0, dataOut.cbData);
-			return buf;
+			try
+			{
+				byte[] buf = new byte[dataOut.cbData];
+				Marshal.Copy(dataOut.pbData, buf, 0, dataOut.cbData);
+				return buf;
+			}
+			finally
+			{
+				if (dataOut.pbData != IntPtr.Zero)
+				{
+					LocalFree(dataOut.pbData);
+				}
+			}
 		}
 	}
 #pragma warning restore IDE1006 // Naming Styles
diff --git a/Engine/Source/Programs/UnrealBuildTool/Executors/UnrealBuildAccelerator/UBAAgentCoordinatorHorde.cs b/Engine/Source/Programs/UnrealBuildTool/Executors/UnrealBuildAccelerator/UBAAgentCoordinatorHorde.cs
--- a/Engine/Source/Programs/UnrealBuildTool/Executors/UnrealBuildAccelerator/UBAAgentCoordinatorHorde.cs
+++ b/Engine/Source/Programs/UnrealBuildTool/Executors/UnrealBuildAccelerator/UBAAgentCoordinatorHorde.cs
@@ -138,7 +138,15 @@ namespace UnrealBuildTool
 				string? accessToken = firstConfig.HordeToken;
 
 				ServiceCollection services = new();
-				services.AddLogging(builder => builder.AddEpicDefault());
+				services.AddLogging(builder =>
+				{
+					builder.AddEpicDefault();
+					if (firstConfig.VerboseAuthLogging)
+					{
+						builder.AddFilter("EpicGames.OIDC", LogLevel.Debug);
+						builder.AddFilter("EpicGames.Horde.Auth", LogLevel.Debug);
+					}
+				});
 				services.AddHorde(options =>
 				{
 					options.ServerUrl = serverUrl;
diff --git a/Engine/Source/Programs/UnrealBuildTool/Executors/UnrealBuildAccelerator/UnrealBuildAcceleratorHordeConfig.cs b/Engine/Source/Programs/UnrealBuildTool/Executors/UnrealBuildAccelerator/UnrealBuildAcceleratorHordeConfig.cs
--- a/Engine/Source/Programs/UnrealBuildTool/Executors/UnrealBuildAccelerator/UnrealBuildAcceleratorHordeConfig.cs
+++ b/Engine/Source/Programs/UnrealBuildTool/Executors/UnrealBuildAccelerator/UnrealBuildAcceleratorHordeConfig.cs
@@ -238,11 +238,19 @@ namespace UnrealBuildTool
 		public int CacheDesiredConnectionCount { get; set; } = 1;
 
 		/// <summary>
-		/// Set which platform the agents requested should have. If not set, default will be host platform 
+		/// Set which platform the agents requested should have. If not set, default will be host platform
 		/// </summary>
 		[ConfigFile(ConfigHierarchyType.Engine, "Uba.Provider.Horde", "UbaAgentPlatform")]
 		[XmlConfigFile(Category = "Horde", Name = "AgentPlatform")]
 		public string? AgentPlatform { get; set; }
+
+		/// <summary>
+		/// When true, enables verbose OIDC auth logging for diagnostics
+		/// </summary>
+		[ConfigFile(ConfigHierarchyType.Engine, "Uba.Provider.Horde", "UbaVerboseAuthLogging")]
+		[XmlConfigFile(Category = "Horde", Name = "VerboseAuthLogging")]
+		[CommandLine("-UBAHordeVerboseAuthLogging")]
+		public bool VerboseAuthLogging { get; set; } = true;
 	}
 
 	#region Extension methods
