EpicGames.OIDC: Use SingleInstanceMutex to avoid cross-thread lock release

diff --git a/Engine/Source/Programs/Shared/EpicGames.Horde/Auth/HordeOidcTokenManager.cs b/Engine/Source/Programs/Shared/EpicGames.Horde/Auth/HordeOidcTokenManager.cs
--- a/Engine/Source/Programs/Shared/EpicGames.Horde/Auth/HordeOidcTokenManager.cs
+++ b/Engine/Source/Programs/Shared/EpicGames.Horde/Auth/HordeOidcTokenManager.cs
@@ -1,6 +1,7 @@
 // Copyright Epic Games, Inc. All Rights Reserved.
 
 using System;
+using System.Collections.Concurrent;
 using System.Collections.Generic;
 using System.Linq;
 using System.Threading;
@@ -28,7 +29,7 @@ internal sealed class HordeOidcTokenManager(
 			return new AnonymousOidcTokenInfo();
 		}
 
-		OidcTokenManager oidcTokenManager = CreateOidcTokenManager(providerIdentifier, authConfig);
+		OidcTokenManager oidcTokenManager = GetOrCreateOidcTokenManager(providerIdentifier, authConfig);
 		return await oidcTokenManager.GetAccessToken(providerIdentifier, cancellationToken);
 	}
 
@@ -41,7 +42,7 @@ internal sealed class HordeOidcTokenManager(
 			return OidcStatus.Connected;
 		}
 
-		OidcTokenManager oidcTokenManager = CreateOidcTokenManager(providerIdentifier, authConfig);
+		OidcTokenManager oidcTokenManager = GetOrCreateOidcTokenManager(providerIdentifier, authConfig);
 		return await oidcTokenManager.GetStatusForProvider(providerIdentifier, cancellationToken);
 	}
 
@@ -55,7 +56,7 @@ internal sealed class HordeOidcTokenManager(
 		}
 
 		_logger?.LogDebug("Interactive login requested for OIDC provider '{Provider}'. A browser window will open.", providerIdentifier);
-		OidcTokenManager oidcTokenManager = CreateOidcTokenManager(providerIdentifier, authConfig);
+		OidcTokenManager oidcTokenManager = GetOrCreateOidcTokenManager(providerIdentifier, authConfig);
 		return await oidcTokenManager.LoginAsync(providerIdentifier, cancellationToken);
 	}
 
@@ -69,7 +70,7 @@ internal sealed class HordeOidcTokenManager(
 			return new AnonymousOidcTokenInfo();
 		}
 
-		OidcTokenManager oidcTokenManager = CreateOidcTokenManager(providerIdentifier, authConfig);
+		OidcTokenManager oidcTokenManager = GetOrCreateOidcTokenManager(providerIdentifier, authConfig);
 		return await oidcTokenManager.TryGetAccessToken(providerIdentifier, cancellationToken);
 	}
 
@@ -77,6 +78,12 @@ internal sealed class HordeOidcTokenManager(
 	private readonly ITokenStore _tokenStore = tokenStore;
 	private readonly ILogger<HordeOidcTokenManager>? _logger = logger;
 	private readonly ILoggerFactory? _loggerFactory = loggerFactory;
+	private readonly ConcurrentDictionary<string, OidcTokenManager> _cachedManagers = new(StringComparer.OrdinalIgnoreCase);
+
+	private OidcTokenManager GetOrCreateOidcTokenManager(string providerIdentifier, GetAuthConfigResponse authConfig)
+	{
+		return _cachedManagers.GetOrAdd(providerIdentifier, _ => CreateOidcTokenManager(providerIdentifier, authConfig));
+	}
 
 	private OidcTokenManager CreateOidcTokenManager(string providerIdentifier, GetAuthConfigResponse authConfig)
 	{
@@ -92,7 +99,7 @@ internal sealed class HordeOidcTokenManager(
 			throw new Exception("Provider identifier mismatch");
 		}
 
-		_logger?.LogDebug("Creating new OidcTokenManager for provider '{Provider}' (note: in-memory token cache is not preserved across calls).", oidcProvider);
+		_logger?.LogDebug("Creating OidcTokenManager for provider '{Provider}'.", oidcProvider);
 
 		Dictionary<string, string?> values = [];
 		values[$"Providers:{oidcProvider}:DisplayName"] = oidcProvider;
diff --git a/Engine/Source/Programs/Shared/EpicGames.OIDC/OidcTokenManager.cs b/Engine/Source/Programs/Shared/EpicGames.OIDC/OidcTokenManager.cs
--- a/Engine/Source/Programs/Shared/EpicGames.OIDC/OidcTokenManager.cs
+++ b/Engine/Source/Programs/Shared/EpicGames.OIDC/OidcTokenManager.cs
@@ -15,6 +15,7 @@ using System.Threading.Tasks;
 using IdentityModel.Client;
 using IdentityModel.OidcClient;
 using IdentityModel.OidcClient.Results;
+using EpicGames.Core;
 using Microsoft.Extensions.Configuration;
 using Microsoft.Extensions.DependencyInjection;
 using Microsoft.Extensions.Logging;
@@ -737,51 +738,22 @@ namespace EpicGames.OIDC
 		/// </summary>
 		protected virtual async Task<IDisposable> AcquireRefreshLockAsync(CancellationToken cancellationToken)
 		{
-			Mutex mutex = new(false, @"Global\OidcRefreshLock");
-			TimeSpan timeout = TimeSpan.FromSeconds(LockAcquireTimeoutSeconds);
+			using CancellationTokenSource timeoutCts = new(TimeSpan.FromSeconds(LockAcquireTimeoutSeconds));
+			using CancellationTokenSource linkedCts = CancellationTokenSource.CreateLinkedTokenSource(cancellationToken, timeoutCts.Token);
 
 			Stopwatch sw = Stopwatch.StartNew();
-			bool acquired;
+			IDisposable handle;
 			try
 			{
-				acquired = await Task.Run(() => mutex.WaitOne(timeout), cancellationToken);
+				handle = await SingleInstanceMutex.AcquireAsync(@"Global\OidcRefreshLock", linkedCts.Token);
 			}
-			catch (AbandonedMutexException)
+			catch (OperationCanceledException) when (timeoutCts.IsCancellationRequested && !cancellationToken.IsCancellationRequested)
 			{
-				// Previous holder crashed — we now own the mutex.
-				acquired = true;
-			}
-			catch
-			{
-				mutex.Dispose();
-				throw;
-			}
-
-			if (!acquired)
-			{
-				mutex.Dispose();
 				throw new TimeoutException($"Timed out after {LockAcquireTimeoutSeconds}s waiting to acquire OIDC refresh lock.");
 			}
 
 			_logger?.LogInformation("Acquired refresh lock for '{Name}' (waited {Elapsed}ms).", _name, sw.ElapsedMilliseconds);
-			return new MutexReleaser(mutex);
-		}
-
-		private sealed class MutexReleaser : IDisposable
-		{
-			private Mutex? _mutex;
-
-			public MutexReleaser(Mutex mutex) => _mutex = mutex;
-
-			public void Dispose()
-			{
-				if (_mutex != null)
-				{
-					_mutex.ReleaseMutex();
-					_mutex.Dispose();
-					_mutex = null;
-				}
-			}
+			return handle;
 		}
 
 		private async Task<OidcTokenInfo?> TryDoRefreshTokenAsync(string refreshToken, CancellationToken cancellationToken)
@@ -866,7 +838,7 @@ namespace EpicGames.OIDC
 
 		public async Task<IOidcTokenInfo?> TryGetAccessTokenAsync(CancellationToken cancellationToken)
 		{
-			// Presence check - reads fresh from store
+			// Unlocked fast-path check; authoritative read happens inside lock below
 			if (!_tokenStore.TryGetRefreshToken(_name, out string? refreshToken) || String.IsNullOrEmpty(refreshToken))
 			{
 				_logger?.LogDebug("No refresh token for provider '{Name}'.", _name);
