[LogMatchers] Add more generic xcode parsing to find pre-compile errors arising from the .xcodeproj files.

#Synced-CL 51672510

diff --git a/Engine/Source/Programs/Shared/EpicGames.LogMatchers/CompileEventMatcher.cs b/Engine/Source/Programs/Shared/EpicGames.LogMatchers/CompileEventMatcher.cs
--- a/Engine/Source/Programs/Shared/EpicGames.LogMatchers/CompileEventMatcher.cs
+++ b/Engine/Source/Programs/Shared/EpicGames.LogMatchers/CompileEventMatcher.cs
@@ -196,11 +196,23 @@ namespace EpicGames.LogMatchers
 			return null;
 		}
 
-		static readonly Regex s_xcodeRebuildPCHPattern = new(@"(?<severity>(?:[Ff]atal )?[Ee]rror): file '(?<file>.*)' has been modified since the precompiled header '(?<pch>.*)' was built");
+		static readonly Regex s_xcodeProjectErrorPattern = new(@"^\s*(?<file>.+?\.xcodeproj):\s*(?<severity>[Ee]rror|[Ww]arning):\s*(?<message>.+)$");
 
+		static readonly Regex s_xcodeRebuildPCHPattern = new(@"(?<severity>(?:[Ff]atal )?[Ee]rror): file '(?<file>.*)' has been modified since the precompiled header '(?<pch>.*)' was built");
+		
 		static bool TryMatchXcodeCppEvent(LogEventBuilder builder, [NotNullWhen(true)] out LogEventMatch? outEvent)
 		{
 			Match? match;
+			
+			if (builder.Current.TryMatch(s_xcodeProjectErrorPattern, out match))
+			{
+				LogLevel level = GetLogLevelFromSeverity(match);
+				builder.Annotate(match.Groups["severity"], LogEventMarkup.Severity);
+				builder.AnnotateSourceFile(match.Groups["file"], null);
+				outEvent = builder.ToMatch(LogEventPriority.High, level, KnownLogEvents.Systemic_XCode);
+				return true;
+			}
+
 			if (builder.Current.TryMatch(s_xcodeRebuildPCHPattern, out match))
 			{
 				builder.Annotate(match.Groups["severity"], LogEventMarkup.Severity);
