Cloud Storage - Implemented the wrapper for listing file contents using ICloudStorage

#Synced-CL 51251964

diff --git a/Engine/Source/Programs/Shared/EpicGames.ProjectStore/CloudStorage/CloudStorage.cs b/Engine/Source/Programs/Shared/EpicGames.ProjectStore/CloudStorage/CloudStorage.cs
--- a/Engine/Source/Programs/Shared/EpicGames.ProjectStore/CloudStorage/CloudStorage.cs
+++ b/Engine/Source/Programs/Shared/EpicGames.ProjectStore/CloudStorage/CloudStorage.cs
@@ -394,45 +394,103 @@ namespace EpicGames.ProjectStore
 		}
 
 		/// <inheritdoc />
-		public async Task<List<BuildContents>> ListBuildContentsAsync(StringId bucketId, CbObjectId buildId, CloudBuildOptionsOverride? optionsOverride = null, CancellationToken cancellationToken = default)
+		public async Task<BuildContents> ListBuildContentsAsync(StringId bucketId, CbObjectId buildId, CloudBuildOptionsOverride? optionsOverride = null, CancellationToken cancellationToken = default)
 		{
-			Uri uri = new($"{_options.Value.Host}/api/v2/builds/{GetCurrentNamespace(optionsOverride)}/{bucketId}/{buildId}");
-			string cmdline = $"builds ls --cloud-url={uri} --quiet";
-			// TODO: we need a way to ask this to get written to file
-			int exitCode = await RunZenCliAsync(cmdline, cancellationToken: cancellationToken);
+			string? ns = GetCurrentNamespace(optionsOverride);
+			if (ns == null)
+			{
+				throw new Exception("Unable to determine namespace, so unable to list contents");
+			}
 
-			EnsureZenExitCodeSuccess(exitCode);
+			string? hostString = optionsOverride?.OverrideHost ?? optionsOverride?.DefaultHost ?? _options.Value.Host;
+			if (String.IsNullOrEmpty(hostString))
+			{
+				throw new ArgumentException("Expected host option to be set either in the global config or the per invocation options override to be able to list contents");
+			}
+			Uri host = new(hostString, UriKind.Absolute);
+			Uri uri = new(host,$"/api/v2/builds/{ns}/{bucketId}/{buildId}");
 
-			// TODO: we need structured output from zen for this listing
+			string tempOutputFile = Path.Combine(Path.GetTempPath(), Path.GetRandomFileName());
+			tempOutputFile = Path.ChangeExtension(tempOutputFile, "cbo");
 
-			Dictionary<string, FileDescription> files = [];
-			/*foreach (string line in stdout)
+			CbObject o;
+			try
 			{
-				string[] segments = line.Split('\t');
-				string fileName = segments[0];
-				ulong rawSize = UInt64.Parse(segments[1]);
-				IoHash rawHash = IoHash.Parse(segments[2]);
+				string cmdline = $"builds ls --cloud-url={uri} --result-path \"{tempOutputFile}\"";
+				int exitCode = await RunZenCliAsync(cmdline, optionsOverride: optionsOverride, cancellationToken: cancellationToken);
 
-				// TODO: Attributes are missing from zen output
-				files.Add(fileName, new FileDescription
+				EnsureZenExitCodeSuccess(exitCode);
+
+				byte[] b = await File.ReadAllBytesAsync(tempOutputFile, cancellationToken);
+				o = new(b);
+			}
+			finally
+			{
+				if (File.Exists(tempOutputFile))
 				{
-					RawSize = rawSize,
-					RawHash = rawHash,
-					Attributes = 0,
-					Mode = 0
-				});
-			}*/
-			List<BuildContents> contents =
-			[
-				new BuildContents
+					File.Delete(tempOutputFile);
+				}
+			}
+
+			CbField parts = o["parts"];
+			if (parts.Equals(CbField.Empty))
+			{
+				throw new Exception("Failed to parse output file when listing content, unable to find parts field");
+			}
+
+			Dictionary<string, PartContents> partsDict = [];
+
+			foreach (CbField part in parts.AsArray())
+			{
+				Dictionary<string, FileDescription> foundFiles = [];
+
+				CbObject objectPart = part.AsObject();
+				CbObjectId partId = objectPart["id"].AsObjectId();
+				string partName = objectPart["partName"].AsString();
+				ulong partSize = 0;
+
+				foreach (CbField files in objectPart["files"].AsArray())
 				{
-					Parts = new Dictionary<string, PartContents>
+					CbObject filesObject = files.AsObject();
+					string path = filesObject["path"].AsString();
+
+					CbField attributesField = filesObject.Find("attributes");
+					int attributes = 0;
+					if (!attributesField.Equals(CbField.Empty))
+					{
+						attributes = attributesField.AsInt32();
+					}
+
+					CbField modeField = filesObject.Find("mode");
+					int mode = 0;
+					if (!modeField.Equals(CbField.Empty))
+					{
+						mode = modeField.AsInt32();
+					}
+					ulong rawSize = filesObject["rawSize"].AsUInt64();
+
+					IoHash rawHash = IoHash.Zero;
+
+					CbField rawHashField = filesObject.Find("rawHash");
+					if (!rawHashField.Equals(CbField.Empty))
 					{
-						{ "default", new PartContents() { PartName = "default", TotalSize = 0, Files = files } }
+						// zen builds ls does not currently output the raw hash but will in a future release so making this optional for now
+						rawHash = rawHashField.AsHash();
 					}
+					partSize += rawSize;
+					foundFiles.Add(path, new FileDescription {Attributes = attributes, RawHash = rawHash, Mode = mode, RawSize = rawSize});
 				}
-			];
-			return contents;
+
+				partsDict.Add(partName, new PartContents
+				{
+					PartId = partId,
+					PartName = partName,
+					TotalSize = partSize,
+					Files = foundFiles
+				});
+			}
+
+			return new BuildContents { Parts = partsDict };
 		}
 
 		private void EnsureZenExitCodeSuccess(int exitCode)
diff --git a/Engine/Source/Programs/Shared/EpicGames.ProjectStore/CloudStorage/ICloudStorage.cs b/Engine/Source/Programs/Shared/EpicGames.ProjectStore/CloudStorage/ICloudStorage.cs
--- a/Engine/Source/Programs/Shared/EpicGames.ProjectStore/CloudStorage/ICloudStorage.cs
+++ b/Engine/Source/Programs/Shared/EpicGames.ProjectStore/CloudStorage/ICloudStorage.cs
@@ -67,13 +67,18 @@ namespace EpicGames.ProjectStore
 	/// <summary>
 	/// The contents available for a part
 	/// </summary>
-	public readonly record struct PartContents(string PartName, ulong TotalSize, Dictionary<string, FileDescription> Files)
+	public readonly record struct PartContents(string PartName, CbObjectId PartId, ulong TotalSize, Dictionary<string, FileDescription> Files)
 	{
 		/// <summary>
 		/// The name of the part
 		/// </summary>
 		public required string PartName { get; init; } = PartName;
 
+		/// <summary>
+		/// The id of the part, usually used in apis to uniquely identify this part
+		/// </summary>
+		public CbObjectId PartId { get; init; } = PartId;
+
 		/// <summary>
 		/// The total uncompressed size of the part
 		/// </summary>
@@ -365,7 +370,7 @@ namespace EpicGames.ProjectStore
 		/// <param name="optionsOverride">Optional per invocation overrides of options</param>		
 		/// <param name="cancellationToken">A cancellation token</param>
 		/// <returns></returns>
-		Task<List<BuildContents>> ListBuildContentsAsync(StringId bucketId, CbObjectId buildId, CloudBuildOptionsOverride? optionsOverride = null, CancellationToken cancellationToken = default);
+		Task<BuildContents> ListBuildContentsAsync(StringId bucketId, CbObjectId buildId, CloudBuildOptionsOverride? optionsOverride = null, CancellationToken cancellationToken = default);
 
 		/// <summary>
 		/// Override the configuration of this instance of cloud storage, mostly useful if you get configuration from another source then the default inis / environment variables
