EpicGames.Core: Improve Windows directory deletion diagnostics and fix read-only directory handling

diff --git a/Engine/Source/Programs/Shared/EpicGames.Core.Tests/FileUtilsTests.cs b/Engine/Source/Programs/Shared/EpicGames.Core.Tests/FileUtilsTests.cs
--- a/Engine/Source/Programs/Shared/EpicGames.Core.Tests/FileUtilsTests.cs
+++ b/Engine/Source/Programs/Shared/EpicGames.Core.Tests/FileUtilsTests.cs
@@ -3,7 +3,9 @@
 using System;
 using System.Collections.Generic;
 using System.IO;
+using System.Runtime.InteropServices;
 using System.Runtime.Versioning;
+using Microsoft.Win32.SafeHandles;
 using Microsoft.VisualStudio.TestTools.UnitTesting;
 
 namespace EpicGames.Core.Tests
@@ -171,6 +173,60 @@ namespace EpicGames.Core.Tests
 			StringAssert.Contains(msg, $"{Environment.ProcessId}: ", StringComparison.Ordinal);
 		}
 
+		[DllImport("kernel32.dll", SetLastError = true, CharSet = CharSet.Unicode)]
+		static extern SafeFileHandle CreateFileW(string lpFileName, uint dwDesiredAccess, uint dwShareMode, IntPtr lpSecurityAttributes, uint dwCreationDisposition, uint dwFlagsAndAttributes, IntPtr hTemplateFile);
+
+		static SafeFileHandle HoldDirectoryHandle(string path)
+		{
+			SafeFileHandle handle = CreateFileW(path, 0x80000000 /* GENERIC_READ */, 0x01 /* FILE_SHARE_READ */, IntPtr.Zero, 3 /* OPEN_EXISTING */, 0x02000000 /* FILE_FLAG_BACKUP_SEMANTICS */, IntPtr.Zero);
+			Assert.IsFalse(handle.IsInvalid, $"Failed to open directory handle on {path}");
+			return handle;
+		}
+
+		static void AssertDiagnosticInfo(string message)
+		{
+			Assert.IsTrue(
+				message.Contains("Directory is empty but could not be removed", StringComparison.Ordinal)
+				|| message.Contains("Processes with open handles", StringComparison.Ordinal)
+				|| message.Contains("Failed to gather lock diagnostics", StringComparison.Ordinal),
+				$"Expected diagnostic info in message but got: {message}");
+		}
+
+		[TestMethod]
+		public void ForceDeleteDirectory_ReadOnlySubdirectory_Deleted()
+		{
+			string subDir = Path.Join(_tempDir.FullName, "Parent");
+			string childDir = Path.Join(subDir, "ReadOnlyChild");
+			Directory.CreateDirectory(childDir);
+			File.SetAttributes(childDir, File.GetAttributes(childDir) | FileAttributes.ReadOnly);
+
+			FileUtils.ForceDeleteDirectory(subDir);
+			Assert.IsFalse(Directory.Exists(subDir));
+		}
+
+		[TestMethod]
+		public void ForceDeleteDirectory_ReadOnlyDirectory_Deleted()
+		{
+			string subDir = Path.Join(_tempDir.FullName, "ReadOnlyDir");
+			Directory.CreateDirectory(subDir);
+			File.SetAttributes(subDir, File.GetAttributes(subDir) | FileAttributes.ReadOnly);
+
+			FileUtils.ForceDeleteDirectory(subDir);
+			Assert.IsFalse(Directory.Exists(subDir));
+		}
+
+		[TestMethod]
+		public void ForceDeleteDirectory_OpenDirectoryHandle_ThrowsWithDiagnostics()
+		{
+			string subDir = Path.Join(_tempDir.FullName, "HandleDir");
+			Directory.CreateDirectory(subDir);
+
+			using SafeFileHandle _ = HoldDirectoryHandle(subDir);
+
+			WrappedFileOrDirectoryException ex = Assert.ThrowsExactly<WrappedFileOrDirectoryException>(() => FileUtils.ForceDeleteDirectory(subDir));
+			AssertDiagnosticInfo(ex.Message);
+		}
+
 		[TestMethod]
 		public void GetFileLockInfoWin32_ReturnsCommandLineAndParentInfo()
 		{
diff --git a/Engine/Source/Programs/Shared/EpicGames.Core/FileUtils.cs b/Engine/Source/Programs/Shared/EpicGames.Core/FileUtils.cs
--- a/Engine/Source/Programs/Shared/EpicGames.Core/FileUtils.cs
+++ b/Engine/Source/Programs/Shared/EpicGames.Core/FileUtils.cs
@@ -213,9 +213,9 @@ namespace EpicGames.Core
 				{
 					lockInfoList = GetFileLockInfoWin32(fileName);
 				}
-				catch
+				catch (Exception ex)
 				{
-					lockInfoList = null;
+					return $"Failed to query file lock info: {ex.Message}";
 				}
 
 				if (lockInfoList != null && lockInfoList.Count > 0)
@@ -646,6 +646,7 @@ namespace EpicGames.Core
 		const int ERROR_PATH_NOT_FOUND = 3;
 		const int ERROR_ACCESS_DENIED = 5;
 		const int ERROR_SHARING_VIOLATION = 32;
+		const int ERROR_DIR_NOT_EMPTY = 145;
 #pragma warning restore IDE1006 // Naming Styles
 
 		[SupportedOSPlatform("windows")]
@@ -675,6 +676,10 @@ namespace EpicGames.Core
 						{
 							if (findData.cFileName != "." && findData.cFileName != "..")
 							{
+								if ((findData.dwFileAttributes & FILE_ATTRIBUTE_READONLY) != 0)
+								{
+									SetFileAttributesW(fullName, (uint)(findData.dwFileAttributes & ~FILE_ATTRIBUTE_READONLY));
+								}
 								ForceDeleteLongDirectoryWin32(fullName);
 							}
 						}
@@ -705,6 +710,13 @@ namespace EpicGames.Core
 		{
 			ForceDeleteLongDirectoryContentsWin32(dirName);
 
+			// Clear read-only on the directory itself before attempting removal
+			int attrs = GetFileAttributesW(dirName);
+			if (attrs != -1 && (attrs & FILE_ATTRIBUTE_READONLY) != 0)
+			{
+				SetFileAttributesW(dirName, (uint)(attrs & ~FILE_ATTRIBUTE_READONLY));
+			}
+
 			if (!RemoveDirectory(dirName))
 			{
 				int errorCode = Marshal.GetLastWin32Error();
@@ -713,7 +725,7 @@ namespace EpicGames.Core
 					Win32Exception ex = new(errorCode);
 					StringBuilder message = new($"Unable to delete {dirName} - {ex.Message}");
 
-					if (errorCode == ERROR_ACCESS_DENIED)
+					if (errorCode is ERROR_ACCESS_DENIED or ERROR_SHARING_VIOLATION or ERROR_DIR_NOT_EMPTY)
 					{
 						string? lockInfo = GetDirectoryLockInfo(dirName);
 						if (lockInfo != null)
@@ -831,9 +843,9 @@ namespace EpicGames.Core
 
 				return message.Length > 0 ? message.ToString() : null;
 			}
-			catch
+			catch (Exception ex)
 			{
-				return null;
+				return $"Failed to gather lock diagnostics: {ex.Message}";
 			}
 		}
 
