EpicGames.OIDC: Use new name for token state file

Recent bug fixes to this library fixed a race condition where long running processes would save stale token data. By changing the name of this file, any apps using the outdated version of this library can no longer interfere. This adds an initial browser login to authenticate but after that, new and old can better co-exist. MacOS provider is not affected since it addresses token data directly by provider name (no saving/loading of file with mixed provider state).

diff --git a/Engine/Source/Programs/Shared/EpicGames.OIDC/TokenStore.cs b/Engine/Source/Programs/Shared/EpicGames.OIDC/TokenStore.cs
--- a/Engine/Source/Programs/Shared/EpicGames.OIDC/TokenStore.cs
+++ b/Engine/Source/Programs/Shared/EpicGames.OIDC/TokenStore.cs
@@ -53,7 +53,7 @@ namespace EpicGames.OIDC
 
 		private static FileInfo GetStorePath()
 		{
-			return new FileInfo(Path.Combine(Environment.GetFolderPath(Environment.SpecialFolder.ApplicationData), "Epic", "UnrealEngine", "Common", "OidcToken", "tokenStore.dat"));
+			return new FileInfo(Path.Combine(Environment.GetFolderPath(Environment.SpecialFolder.ApplicationData), "Epic", "UnrealEngine", "Common", "OidcToken", "tokenStore-v2.dat"));
 		}
 
 		private static void ReadStoreFromDisk(out Dictionary<string, byte[]> providers, out byte[]? iv)
diff --git a/Engine/Source/Programs/Shared/EpicGames.OIDC/WindowsTokenStore.cs b/Engine/Source/Programs/Shared/EpicGames.OIDC/WindowsTokenStore.cs
--- a/Engine/Source/Programs/Shared/EpicGames.OIDC/WindowsTokenStore.cs
+++ b/Engine/Source/Programs/Shared/EpicGames.OIDC/WindowsTokenStore.cs
@@ -69,7 +69,7 @@ namespace EpicGames.OIDC
 		private FileInfo GetStorePath()
 		{
 			// Always return a fresh FileInfo to avoid stale cached Exists state.
-			return new FileInfo(_storePathOverride ?? Path.Combine(Environment.GetFolderPath(Environment.SpecialFolder.LocalApplicationData), "UnrealEngine", "Common", "OidcToken", "oidcTokenStore.dat"));
+			return new FileInfo(_storePathOverride ?? Path.Combine(Environment.GetFolderPath(Environment.SpecialFolder.LocalApplicationData), "UnrealEngine", "Common", "OidcToken", "oidcTokenStore-v2.dat"));
 		}
 
 		private Dictionary<string, byte[]> ReadStoreFromDisk()
