UnrealBuildTool: Fix in-memory file caching to pretend that a directory and its ancestors are extant when it's pre-populated with files.

#Synced-CL 51819546

diff --git a/Engine/Source/Programs/Shared/EpicGames.Build/System/DirectoryItem.cs b/Engine/Source/Programs/Shared/EpicGames.Build/System/DirectoryItem.cs
--- a/Engine/Source/Programs/Shared/EpicGames.Build/System/DirectoryItem.cs
+++ b/Engine/Source/Programs/Shared/EpicGames.Build/System/DirectoryItem.cs
@@ -125,7 +125,7 @@ namespace UnrealBuildBase
 		/// <param name="files"></param>
 		public void PrepopulateCache(List<string> files)
 		{
-			Dictionary<string, FileItem> cacheFiles = new(files.Count);
+			Dictionary<string, FileItem> cacheFiles = new(files.Count, FileSystemReference.Comparer);
 			foreach (string file in files)
 			{
 				FileReference fileRef = FileReference.Combine(Location, file);
@@ -134,7 +134,29 @@ namespace UnrealBuildBase
 				cacheFiles.Add(file, fileItem);
 			}
 
-			_cache = new(([], cacheFiles));
+			_cache = new((new(FileSystemReference.Comparer), cacheFiles));
+
+			// Update ourself and all our ancestors - if we're pretending we have extant files then they must all exist too.
+			DirectoryItem? toUpdate = this;
+			DirectoryItem? childDir = null;
+			while (toUpdate is not null)
+			{
+				if (childDir is not null)
+				{
+					lock (toUpdate)
+					{
+						toUpdate.Directories[childDir.Name] = childDir;
+					}
+				}
+				// This has to come after the above because it triggers Scan which resets Exists back to false.
+				if (!toUpdate.Exists)
+				{
+					toUpdate.Update(new DirectoryEnumerator.Entry { Attributes = FileAttributes.Directory }, force: true);
+				}
+
+				childDir = toUpdate;
+				toUpdate = toUpdate.GetParentDirectoryItem();
+			}
 		}
 
 		/// <summary>
@@ -229,6 +251,14 @@ namespace UnrealBuildBase
 			return newInfo;
 		}
 
+		internal void Update(DirectoryEnumerator.Entry entry, bool force)
+		{
+			if (_info == null || force)
+			{
+				_info = new() { Exists = entry.Attributes != FileAttributes.None, LastWriteTimeUtc = entry.LastWriteTimeUtc, CreationTimeUtc = entry.CreationTimeUtc };
+			}
+		}
+
 		/// <summary>
 		/// Scans the directory for directories and files, used for lazy initialization.
 		/// </summary>
@@ -237,7 +267,7 @@ namespace UnrealBuildBase
 		{
 			if (_info != null && !_info.Exists)
 			{
-				return ([], []);
+				return (new(FileSystemReference.Comparer), new(FileSystemReference.Comparer));
 			}
 
 			try
@@ -247,7 +277,7 @@ namespace UnrealBuildBase
 				if (entries == null)
 				{
 					_info = new Info() { Exists = false };
-					return ([], []);
+					return (new(FileSystemReference.Comparer), new(FileSystemReference.Comparer));
 				}
 				int dirCount = 0;
 				foreach (DirectoryEnumerator.Entry entry in entries)
