EpicGames.Build: PrepopulateCache should use TryAdd rather than Add

#rnx
#Synced-CL 51916670

diff --git a/Engine/Source/Programs/Shared/EpicGames.Build/System/DirectoryItem.cs b/Engine/Source/Programs/Shared/EpicGames.Build/System/DirectoryItem.cs
--- a/Engine/Source/Programs/Shared/EpicGames.Build/System/DirectoryItem.cs
+++ b/Engine/Source/Programs/Shared/EpicGames.Build/System/DirectoryItem.cs
@@ -131,7 +131,7 @@ namespace UnrealBuildBase
 				FileReference fileRef = FileReference.Combine(Location, file);
 				FileItem fileItem = FileItem.GetItemByFileReference(fileRef);
 				fileItem.Update(new() { Attributes = FileAttributes.Normal, Size = long.MaxValue }, true);
-				cacheFiles.Add(file, fileItem);
+				cacheFiles.TryAdd(file, fileItem);
 			}
 
 			_cache = new((new(FileSystemReference.Comparer), cacheFiles));
