Ignore error codes from AutoSDK hooks under Wine.

This change ignores AutoSDK hook failures when running under Wine, because error propagation behaves differently for batch scripts when compared to Windows. This sometimes results in non-zero exit codes after a script completes successfully.

#rnx

diff --git a/Engine/Source/Programs/Shared/EpicGames.Build/Platform/UEBuildPlatformSDK.cs b/Engine/Source/Programs/Shared/EpicGames.Build/Platform/UEBuildPlatformSDK.cs
--- a/Engine/Source/Programs/Shared/EpicGames.Build/Platform/UEBuildPlatformSDK.cs
+++ b/Engine/Source/Programs/Shared/EpicGames.Build/Platform/UEBuildPlatformSDK.cs
@@ -1569,7 +1569,14 @@ namespace EpicGames.Core
 					if (hookProcess.ExitCode != 0)
 					{
 						Logger.LogDebug("Hook exited uncleanly (returned {ExitCode}), considering it failed.", hookProcess.ExitCode);
-						return false;
+
+						// Don't fail on Wine, because error propagation behaves differently
+						// for batch scripts when compared to Windows. This sometimes results
+						// in non-zero exit codes after a script completes successfully.
+						if (!RuntimePlatform.IsRunningOnWine())
+						{
+							return false;
+						}
 					}
 
 					return true;
