EpicGames.OIDC: Prevent refresh token loss when process shuts down during token rotation

diff --git a/Engine/Source/Programs/Shared/EpicGames.Horde.Tests/OidcTokenRefreshTests.cs b/Engine/Source/Programs/Shared/EpicGames.Horde.Tests/OidcTokenRefreshTests.cs
--- a/Engine/Source/Programs/Shared/EpicGames.Horde.Tests/OidcTokenRefreshTests.cs
+++ b/Engine/Source/Programs/Shared/EpicGames.Horde.Tests/OidcTokenRefreshTests.cs
@@ -11,27 +11,47 @@ namespace EpicGames.Horde.Tests;
 /// <summary>
 /// Testable subclass of OidcTokenClient that overrides DoRefreshAsync (no real OIDC backend)
 /// and AcquireRefreshLock (in-process SemaphoreSlim instead of file lock).
+/// Optionally blocks mid-refresh until signaled, enabling deterministic cancellation testing.
 /// </summary>
 sealed class TestableOidcTokenClient(
 	string name,
 	ITokenStore tokenStore,
-	Func<string, TokenRefreshResult> refreshHandler)
-	: OidcTokenClient(name, CreateDummyProviderInfo(), TimeSpan.FromMinutes(20), tokenStore)
+	Func<string, TokenRefreshResult> refreshHandler,
+	bool asyncGated = false)
+	: OidcTokenClient(name, CreateDummyProviderInfo(), TimeSpan.FromMinutes(20), tokenStore), IDisposable
 {
-	private static readonly SemaphoreSlim s_testLock = new(1, 1);
+	private readonly TaskCompletionSource? _enteredRefresh = asyncGated ? new() : null;
+	private readonly TaskCompletionSource? _proceedWithRefresh = asyncGated ? new() : null;
+	private readonly SemaphoreSlim _testLock = new(1, 1);
 	public int RefreshCallCount { get; private set; }
 
-	protected override Task<TokenRefreshResult> DoRefreshAsync(
+	/// <summary>Completes when DoRefreshAsync has been entered (only when constructed with async gating).</summary>
+	public Task EnteredRefreshTask => _enteredRefresh?.Task ?? Task.CompletedTask;
+
+	/// <summary>Signal the refresh to continue (call after cancelling external CTS).</summary>
+	public void AllowRefreshToComplete() => _proceedWithRefresh?.TrySetResult();
+
+	public void Dispose() => _testLock.Dispose();
+
+	protected override async Task<TokenRefreshResult> DoRefreshAsync(
 		string refreshToken, CancellationToken cancellationToken)
 	{
 		RefreshCallCount++;
-		return Task.FromResult(refreshHandler(refreshToken));
+
+		if (_enteredRefresh != null)
+		{
+			_enteredRefresh.TrySetResult();
+			await _proceedWithRefresh!.Task;
+			cancellationToken.ThrowIfCancellationRequested();
+		}
+
+		return await Task.FromResult(refreshHandler(refreshToken));
 	}
 
 	protected override async Task<IDisposable> AcquireRefreshLockAsync(CancellationToken cancellationToken)
 	{
-		await s_testLock.WaitAsync(cancellationToken);
-		return new SemaphoreReleaser(s_testLock);
+		await _testLock.WaitAsync(cancellationToken);
+		return new SemaphoreReleaser(_testLock);
 	}
 
 	private sealed class SemaphoreReleaser(SemaphoreSlim s) : IDisposable
@@ -39,7 +59,7 @@ sealed class TestableOidcTokenClient(
 		public void Dispose() => s.Release();
 	}
 
-	private static ProviderInfo CreateDummyProviderInfo() => new()
+	internal static ProviderInfo CreateDummyProviderInfo() => new()
 	{
 		ServerUri = new Uri("https://fake-oidc.test"),
 		ClientId = "test-client",
@@ -64,11 +84,11 @@ public sealed class OidcTokenRefreshTests
 			AccessToken: null, RefreshToken: null, AccessTokenExpiration: default);
 
 	private static (InMemoryTokenStore Store, TestableOidcTokenClient Client) CreateClient(
-		string initialToken, Func<string, TokenRefreshResult> refreshHandler)
+		string initialToken, Func<string, TokenRefreshResult> refreshHandler, bool asyncGated = false)
 	{
 		InMemoryTokenStore store = new();
 		store.AddRefreshToken(Provider, initialToken);
-		return (store, new TestableOidcTokenClient(Provider, store, refreshHandler));
+		return (store, new TestableOidcTokenClient(Provider, store, refreshHandler, asyncGated));
 	}
 
 	/// <summary>
@@ -80,7 +100,7 @@ public sealed class OidcTokenRefreshTests
 	{
 		using InMemoryTokenStore store = new();
 		store.AddRefreshToken(Provider, "tokenA");
-		TestableOidcTokenClient client = new(Provider, store, refreshToken =>
+		using TestableOidcTokenClient client = new(Provider, store, refreshToken =>
 		{
 			if (refreshToken == "tokenA")
 			{
@@ -174,7 +194,7 @@ public sealed class OidcTokenRefreshTests
 
 		using InMemoryTokenStore store = new();
 		store.AddRefreshToken(Provider, "tokenA");
-		TestableOidcTokenClient client = new(Provider, store, refreshToken =>
+		using TestableOidcTokenClient client = new(Provider, store, refreshToken =>
 		{
 			int current = Interlocked.Increment(ref concurrentCount);
 			int snapshot;
@@ -202,4 +222,55 @@ public sealed class OidcTokenRefreshTests
 			Assert.AreEqual(1, maxConcurrent, "Refresh calls should have been serialized by the lock");
 		}
 	}
+
+	/// <summary>
+	/// When external cancellation fires mid-refresh (e.g. process shutdown), the new token
+	/// must still be persisted. This is the key bug scenario: Okta rotates the token but
+	/// OperationCanceledException skips AddRefreshToken, losing the new token permanently.
+	/// </summary>
+	[TestMethod]
+	public async Task RefreshToken_ExternalCancellationDuringRefresh_StillPersistsTokenAsync()
+	{
+		using CancellationTokenSource externalCts = new();
+		(InMemoryTokenStore store, TestableOidcTokenClient client) =
+			CreateClient("tokenA", _ => SuccessResult("newAccess", "tokenB"), asyncGated: true);
+		using (store)
+		{
+			Task<IOidcTokenInfo?> refreshTask = client.TryGetAccessTokenAsync(externalCts.Token);
+			await client.EnteredRefreshTask;
+
+			// Simulate process shutdown — cancel AFTER Okta has processed the rotation
+			await externalCts.CancelAsync();
+			client.AllowRefreshToComplete();
+
+			IOidcTokenInfo? result = await refreshTask;
+
+			Assert.IsNotNull(result);
+			Assert.IsTrue(store.TryGetRefreshToken(Provider, out string? storedToken));
+			Assert.AreEqual("tokenB", storedToken, "New token must be persisted even when external cancellation fires");
+		}
+	}
+
+	/// <summary>
+	/// If cancellation is already requested before refresh starts, the client should
+	/// throw without attempting the HTTP call.
+	/// </summary>
+	[TestMethod]
+	public async Task RefreshToken_CancelledBeforeRefresh_DoesNotCallRefreshAsync()
+	{
+		(InMemoryTokenStore store, TestableOidcTokenClient client) =
+			CreateClient("tokenA", _ => SuccessResult());
+		using (store)
+		{
+			using CancellationTokenSource cts = new();
+			await cts.CancelAsync();
+
+			await Assert.ThrowsExactlyAsync<TaskCanceledException>(
+				() => client.TryGetAccessTokenAsync(cts.Token));
+
+			Assert.AreEqual(0, client.RefreshCallCount);
+			Assert.IsTrue(store.TryGetRefreshToken(Provider, out string? token));
+			Assert.AreEqual("tokenA", token, "Original token must remain on disk");
+		}
+	}
 }
diff --git a/Engine/Source/Programs/Shared/EpicGames.Horde.Tests/WindowsTokenStoreTests.cs b/Engine/Source/Programs/Shared/EpicGames.Horde.Tests/WindowsTokenStoreTests.cs
--- a/Engine/Source/Programs/Shared/EpicGames.Horde.Tests/WindowsTokenStoreTests.cs
+++ b/Engine/Source/Programs/Shared/EpicGames.Horde.Tests/WindowsTokenStoreTests.cs
@@ -316,7 +316,7 @@ public sealed class TokenStoreFailureTests
 		store.AddRefreshToken(Provider, "initialToken");
 
 		int callCount = 0;
-		TestableOidcTokenClient client = new(Provider, store, _ =>
+		using TestableOidcTokenClient client = new(Provider, store, _ =>
 		{
 			callCount++;
 			store.SetDropWrites(true);
@@ -336,7 +336,7 @@ public sealed class TokenStoreFailureTests
 	public async Task TryGetAccessToken_EmptyStore_ThrowsNotLoggedInAsync()
 	{
 		using FailingTokenStore store = new();
-		TestableOidcTokenClient client = new(Provider, store, _ => SuccessResult());
+		using TestableOidcTokenClient client = new(Provider, store, _ => SuccessResult());
 
 		await Assert.ThrowsExactlyAsync<NotLoggedInException>(
 			() => client.TryGetAccessTokenAsync(CancellationToken.None));
diff --git a/Engine/Source/Programs/Shared/EpicGames.OIDC/ITokenStore.cs b/Engine/Source/Programs/Shared/EpicGames.OIDC/ITokenStore.cs
--- a/Engine/Source/Programs/Shared/EpicGames.OIDC/ITokenStore.cs
+++ b/Engine/Source/Programs/Shared/EpicGames.OIDC/ITokenStore.cs
@@ -3,6 +3,8 @@
 using System;
 using System.Diagnostics.CodeAnalysis;
 using System.Runtime.InteropServices;
+using System.Security.Cryptography;
+using System.Text;
 using Microsoft.Extensions.DependencyInjection;
 using Microsoft.Extensions.Logging;
 
@@ -19,6 +21,15 @@ namespace EpicGames.OIDC
 		public void AddRefreshToken(string oidcProvider, string refreshToken);
 	}
 
+	/// <summary>
+	/// Shared utilities for safe token logging.
+	/// </summary>
+	public static class TokenUtils
+	{
+		/// <summary>First 8 hex chars of SHA-256 - enough for log correlation without leaking the token.</summary>
+		public static string Fingerprint(string token) => Convert.ToHexString(SHA256.HashData(Encoding.UTF8.GetBytes(token)))[..8];
+	}
+
 	public static class TokenStoreFactory
 	{
 		public static ITokenStore CreateTokenStore(IServiceProvider provider)
diff --git a/Engine/Source/Programs/Shared/EpicGames.OIDC/OidcTokenManager.cs b/Engine/Source/Programs/Shared/EpicGames.OIDC/OidcTokenManager.cs
--- a/Engine/Source/Programs/Shared/EpicGames.OIDC/OidcTokenManager.cs
+++ b/Engine/Source/Programs/Shared/EpicGames.OIDC/OidcTokenManager.cs
@@ -456,7 +456,9 @@ namespace EpicGames.OIDC
 			_tokenExpiry = tokenInfo.TokenExpiry;
 
 			_logger?.LogDebug("Saving refresh token to store after interactive login for provider '{Name}'.", _name);
-			using IDisposable _ = await AcquireRefreshLockAsync(cancellationToken);
+			// After browser login, token MUST be persisted regardless of external cancellation.
+			using CancellationTokenSource loginPersistCts = new(TimeSpan.FromSeconds(LockAcquireTimeoutSeconds));
+			using IDisposable _ = await AcquireRefreshLockAsync(loginPersistCts.Token);
 			_tokenStore.AddRefreshToken(_name, tokenInfo.RefreshToken);
 
 			_logger?.LogInformation("Interactive login completed for provider '{Name}'. Token expiry: {Expiry}", _name, tokenInfo.TokenExpiry);
@@ -758,52 +760,69 @@ namespace EpicGames.OIDC
 
 		private async Task<OidcTokenInfo?> TryDoRefreshTokenAsync(string refreshToken, CancellationToken cancellationToken)
 		{
-			TokenRefreshResult result = await DoRefreshAsync(refreshToken, cancellationToken);
-
-			// Fallback for non-cooperating processes (old binaries without new lock)
-			if (result is { IsError: true, Error: "invalid_grant" })
+			try
 			{
-				_logger?.LogWarning("Refresh token for provider '{Name}' returned invalid_grant. May have been rotated by another process.", _name);
+				_logger?.LogDebug("Starting OIDC refresh HTTP call for '{Name}' (token: {Fingerprint}).",
+					_name, TokenUtils.Fingerprint(refreshToken));
+
+				TokenRefreshResult result = await DoRefreshAsync(refreshToken, cancellationToken);
+				_logger?.LogDebug("OIDC refresh HTTP call completed for '{Name}'. IsError={IsError}.", _name, result.IsError);
 
-				if (_tokenStore.TryGetRefreshToken(_name, out string? freshToken) &&
-				    !String.Equals(freshToken, refreshToken, StringComparison.Ordinal))
+				// Fallback for non-cooperating processes (old binaries without new lock)
+				if (result is { IsError: true, Error: "invalid_grant" })
 				{
-					_logger?.LogInformation("Found rotated refresh token for provider '{Name}'. Retrying.", _name);
-					result = await DoRefreshAsync(freshToken, cancellationToken);
+					_logger?.LogWarning("Refresh token for provider '{Name}' returned invalid_grant. May have been rotated by another process.", _name);
+
+					if (_tokenStore.TryGetRefreshToken(_name, out string? freshToken) &&
+					    !String.Equals(freshToken, refreshToken, StringComparison.Ordinal))
+					{
+						_logger?.LogInformation("Found rotated refresh token for provider '{Name}'. Retrying.", _name);
+						result = await DoRefreshAsync(freshToken, cancellationToken);
+					}
 				}
-			}
 
-			if (result.IsError)
-			{
-				if (result.Error == "invalid_grant")
+				if (result.IsError)
 				{
-					_logger?.LogWarning("Refresh token for provider '{Name}' still invalid. User must re-authenticate.", _name);
-					return null;
+					if (result.Error == "invalid_grant")
+					{
+						_logger?.LogWarning("Refresh token for provider '{Name}' still invalid. User must re-authenticate.", _name);
+						return null;
+					}
+					_logger?.LogWarning("Refresh token error for provider '{Name}': {Error} - {Description}", _name, result.Error, result.ErrorDescription);
+					throw new Exception($"Error using the refresh token: {result.Error} , details: {result.ErrorDescription}");
 				}
-				_logger?.LogWarning("Refresh token error for provider '{Name}': {Error} - {Description}", _name, result.Error, result.ErrorDescription);
-				throw new Exception($"Error using the refresh token: {result.Error} , details: {result.ErrorDescription}");
-			}
 
-			// Persist new token to disk immediately
-			if (String.IsNullOrEmpty(result.RefreshToken))
-			{
-				throw new Exception("No refresh token was provided in response to refresh.");
-			}
-			_logger?.LogDebug("Saving new refresh token to store for provider '{Name}'.", _name);
-			_tokenStore.AddRefreshToken(_name, result.RefreshToken);
+				// Persist new token to disk immediately
+				if (String.IsNullOrEmpty(result.RefreshToken))
+				{
+					throw new Exception("No refresh token was provided in response to refresh.");
+				}
+				_logger?.LogDebug("Persisting new refresh token to store for '{Name}' (token: {Fingerprint}).",
+					_name, TokenUtils.Fingerprint(result.RefreshToken));
+				_tokenStore.AddRefreshToken(_name, result.RefreshToken);
+				_logger?.LogDebug("Refresh token persisted successfully for '{Name}'.", _name);
 
-			// Cache access token in memory (short-lived)
-			_accessToken = result.AccessToken;
-			_tokenExpiry = result.AccessTokenExpiration;
+				// Cache access token in memory (short-lived)
+				_accessToken = result.AccessToken;
+				_tokenExpiry = result.AccessTokenExpiration;
 
-			_logger?.LogInformation("Successfully refreshed token for provider '{Name}'. New expiry: {Expiry}", _name, result.AccessTokenExpiration);
+				_logger?.LogInformation("Successfully refreshed token for provider '{Name}'. New expiry: {Expiry}", _name, result.AccessTokenExpiration);
 
-			return new OidcTokenInfo
+				return new OidcTokenInfo
+				{
+					RefreshToken = result.RefreshToken,
+					AccessToken = result.AccessToken,
+					TokenExpiry = result.AccessTokenExpiration
+				};
+			}
+			catch (OperationCanceledException) when (cancellationToken.IsCancellationRequested)
 			{
-				RefreshToken = result.RefreshToken,
-				AccessToken = result.AccessToken,
-				TokenExpiry = result.AccessTokenExpiration
-			};
+				_logger?.LogWarning(
+					"OIDC refresh for '{Name}' timed out after {TimeoutSeconds}s (token: {Fingerprint}). "
+					+ "If Okta already rotated the token, the old token on disk may now be invalid.",
+					_name, RefreshHttpTimeoutSeconds, TokenUtils.Fingerprint(refreshToken));
+				throw;
+			}
 		}
 
 		private async Task<DiscoveryDocumentResponse> GetDiscoveryDocumentAsync(CancellationToken cancellationToken)
@@ -869,8 +888,30 @@ namespace EpicGames.OIDC
 				throw new NotLoggedInException();
 			}
 
-			// Bound the HTTP refresh to prevent holding the lock indefinitely
-			using CancellationTokenSource refreshCts = CancellationTokenSource.CreateLinkedTokenSource(cancellationToken);
+			// If another in-process caller refreshed while we waited for the lock, skip the redundant rotation.
+			OidcTokenInfo freshCheck = new()
+			{
+				RefreshToken = refreshToken,
+				AccessToken = _accessToken,
+				TokenExpiry = _tokenExpiry
+			};
+			if (freshCheck.IsValid(DateTime.Now.AddMinutes(-2)))
+			{
+				_logger?.LogInformation("Access token for '{Name}' became valid while waiting for lock. Skipping refresh.", _name);
+				return freshCheck;
+			}
+
+			// Once committed to refresh: only a timeout can cancel, NOT external cancellation.
+			// Prevents losing a rotated refresh token if the caller cancels mid-HTTP-call (e.g., process shutdown via BackgroundTask disposal).
+			if (cancellationToken.IsCancellationRequested)
+			{
+				_logger?.LogInformation("Cancellation requested before starting refresh for '{Name}'. Releasing lock without refreshing.", _name);
+				cancellationToken.ThrowIfCancellationRequested();
+			}
+
+			_logger?.LogInformation("Starting refresh for '{Name}' with timeout-only cancellation ({Timeout}s)", _name, RefreshHttpTimeoutSeconds);
+
+			using CancellationTokenSource refreshCts = new();
 			refreshCts.CancelAfter(TimeSpan.FromSeconds(RefreshHttpTimeoutSeconds));
 
 			return await TryDoRefreshTokenAsync(refreshToken, refreshCts.Token);
diff --git a/Engine/Source/Programs/Shared/EpicGames.OIDC/WindowsTokenStore.cs b/Engine/Source/Programs/Shared/EpicGames.OIDC/WindowsTokenStore.cs
--- a/Engine/Source/Programs/Shared/EpicGames.OIDC/WindowsTokenStore.cs
+++ b/Engine/Source/Programs/Shared/EpicGames.OIDC/WindowsTokenStore.cs
@@ -5,7 +5,6 @@ using System.Collections.Generic;
 using System.ComponentModel;
 using System.IO;
 using System.Runtime.InteropServices;
-using System.Security.Cryptography;
 using System.Text;
 using System.Text.Json;
 using System.Text.Json.Serialization;
@@ -46,9 +45,6 @@ namespace EpicGames.OIDC
 		private readonly ILogger<WindowsTokenStore>? _logger = null;
 		private readonly string? _storePathOverride;
 
-		/// <summary>First 8 hex chars of SHA-256 — enough for log correlation without leaking the token.</summary>
-		private static string TokenFingerprint(string token) => Convert.ToHexString(SHA256.HashData(Encoding.UTF8.GetBytes(token)))[..8];
-
 		public WindowsTokenStore()
 		{
 		}
@@ -184,7 +180,7 @@ namespace EpicGames.OIDC
 				{
 					byte[] bytes = CryptProtectDataHelper.DoCryptUnprotectData(encryptedToken, $"OidcToken-{oidcProvider}", GetEntropy(oidcProvider));
 					refreshToken = Encoding.Unicode.GetString(bytes);
-					_logger?.LogInformation("Read refresh token for '{Provider}': {Fingerprint}", oidcProvider, TokenFingerprint(refreshToken));
+					_logger?.LogInformation("Read refresh token for '{Provider}': {Fingerprint}", oidcProvider, TokenUtils.Fingerprint(refreshToken));
 					return true;
 				}
 				catch (Win32Exception e)
@@ -237,7 +233,7 @@ namespace EpicGames.OIDC
 					mutex.ReleaseMutex();
 				}
 
-				_logger?.LogInformation("Saved refresh token for '{Provider}': {Fingerprint}", oidcProvider, TokenFingerprint(refreshToken));
+				_logger?.LogInformation("Saved refresh token for '{Provider}': {Fingerprint}", oidcProvider, TokenUtils.Fingerprint(refreshToken));
 			}
 			catch (Exception ex)
 			{
