[UHT]
* Fixed code gen when partial only has editor-only properties or functions

#Synced-CL 51957558

diff --git a/Engine/Source/Programs/Shared/EpicGames.Core/UnrealEngineTypes.cs b/Engine/Source/Programs/Shared/EpicGames.Core/UnrealEngineTypes.cs
--- a/Engine/Source/Programs/Shared/EpicGames.Core/UnrealEngineTypes.cs
+++ b/Engine/Source/Programs/Shared/EpicGames.Core/UnrealEngineTypes.cs
@@ -612,7 +612,7 @@ namespace EpicGames.Core
 		BlueprintPure = 0x10000000,
 
 		/// <summary>
-		/// function can only be called from an editor scrippt.
+		/// function can only be called from an editor script.
 		/// </summary>
 		EditorOnly = 0x20000000,
 
diff --git a/Engine/Source/Programs/Shared/EpicGames.UHT/Exporters/CodeGen/UhtHeaderCodeGeneratorCppFile.cs b/Engine/Source/Programs/Shared/EpicGames.UHT/Exporters/CodeGen/UhtHeaderCodeGeneratorCppFile.cs
--- a/Engine/Source/Programs/Shared/EpicGames.UHT/Exporters/CodeGen/UhtHeaderCodeGeneratorCppFile.cs
+++ b/Engine/Source/Programs/Shared/EpicGames.UHT/Exporters/CodeGen/UhtHeaderCodeGeneratorCppFile.cs
@@ -1414,24 +1414,24 @@ namespace EpicGames.UHT.Exporters.CodeGen
 			// Forward declare the partial class for partial function compilation
 			if (partial.Functions.Any())
 			{
-				builder.Append($"// Forward declaration for partial function compilation\r\n");
 				builder.Append($"class {partial.OwnerClass!.Namespace.FullSourceName}{partial.OwnerClass.SourceName};\r\n\r\n");
 			}
 
 			// Generate function arrays for partial functions (after functions are generated)
 			if (!Session.IsUsingCompiledInObjectFormat(UhtCompiledInObjectFormat.ConstInit) && partial.Functions.Any())
 			{
+				bool allIsEditorOnly = partial.Functions.All(f => f.FunctionFlags.HasAnyFlags(EFunctionFlags.EditorOnly));
+				using UhtConditionalMacroBlock editorBlock = new(builder, UhtNames.WithEditorOnlyData, allIsEditorOnly);
+
 				// Function constructor pointers (for creating UFunctions)
-				builder.Append($"// Function constructor pointers for {partial.SourceName}\r\n");
 				builder.Append($"static UFunction* (*PartialFunctionConstructors_{partial.SourceName}[])(ETypeConstructPhase) = {{\r\n");
 				foreach (UhtFunction function in partial.Functions)
 				{
 					builder.Append($"\t&{GetSingletonName(function, UhtSingletonType.Registered)},\r\n");
 				}
-				builder.Append("};\r\n\r\n");
+				builder.Append("};\r\n");
 
 				// Native function bindings (for registering exec functions)
-				builder.Append($"// Native function bindings for {partial.SourceName}\r\n");
 				builder.Append($"static constexpr UE::CodeGen::FClassNativeFunction PartialNativeFunctions_{partial.SourceName}[] = {{\r\n");
 				foreach (UhtFunction function in partial.Functions)
 				{
@@ -1470,11 +1470,15 @@ namespace EpicGames.UHT.Exporters.CodeGen
 				builder.Append($"\t.Class = {GetSingletonName(partial.OwnerClass, UhtSingletonType.Registered)},\r\n");
 				if (partial.Properties.Any())
 				{
+					bool allIsEditorOnly = partial.Properties.All(p => p.PropertyFlags.HasAnyFlags(EPropertyFlags.EditorOnly));
+					using UhtConditionalMacroBlock editorBlock = new(builder, UhtNames.WithEditorOnlyData, allIsEditorOnly);
 					builder.Append($"\t.Properties = UHT_STATICS::PropPointers,\r\n");
 					builder.Append($"\t.NumProperties = UE_ARRAY_COUNT(UHT_STATICS::PropPointers),\r\n");
 				}
 				if (partial.Functions.Any())
 				{
+					bool allIsEditorOnly = partial.Functions.All(f => f.FunctionFlags.HasAnyFlags(EFunctionFlags.EditorOnly));
+					using UhtConditionalMacroBlock editorBlock = new(builder, UhtNames.WithEditorOnlyData, allIsEditorOnly);
 					builder.Append($"\t.FunctionConstructors = PartialFunctionConstructors_{partial.SourceName},\r\n");
 					builder.Append($"\t.NativeFunctions = PartialNativeFunctions_{partial.SourceName},\r\n");
 					builder.Append($"\t.NumFunctions = {partial.Functions.Count()},\r\n");
diff --git a/Engine/Source/Programs/Shared/EpicGames.UHT/Parsers/UhtFunctionParser.cs b/Engine/Source/Programs/Shared/EpicGames.UHT/Parsers/UhtFunctionParser.cs
--- a/Engine/Source/Programs/Shared/EpicGames.UHT/Parsers/UhtFunctionParser.cs
+++ b/Engine/Source/Programs/Shared/EpicGames.UHT/Parsers/UhtFunctionParser.cs
@@ -387,6 +387,11 @@ namespace EpicGames.UHT.Parsers
 						function.FunctionFlags |= EFunctionFlags.EditorOnly;
 						function.DefineScope |= UhtDefineScope.Editor;
 					}
+					else if ((compilerDirective & UhtCompilerDirective.WithEditorOnlyData) != 0)
+					{
+						function.FunctionFlags |= EFunctionFlags.EditorOnly;
+						function.DefineScope |= UhtDefineScope.EditorOnlyData;
+					}
 					function.DefineScope |= (compilerDirective & ~(UhtCompilerDirective.WithEditorOnlyData | UhtCompilerDirective.WithEditor)).GetDefaultDefineScopes();
 
 					specifierParser.ParseDeferred();
diff --git a/Engine/Source/Programs/Shared/EpicGames.UHT/Types/UhtPartial.cs b/Engine/Source/Programs/Shared/EpicGames.UHT/Types/UhtPartial.cs
--- a/Engine/Source/Programs/Shared/EpicGames.UHT/Types/UhtPartial.cs
+++ b/Engine/Source/Programs/Shared/EpicGames.UHT/Types/UhtPartial.cs
@@ -272,7 +272,7 @@ namespace EpicGames.UHT.Types
 			}
 
 			// Validate owner class is specified
-			if (string.IsNullOrEmpty(OwnerClassName))
+			if (String.IsNullOrEmpty(OwnerClassName))
 			{
 				this.LogError("Partial must specify the class it extends via UPARTIAL(ClassName) macro");
 			}
