UE: Make PropertiesSize more accurate for native types by actually assigning it to the unaligned size of the struct. This is important so that functions like FProperty::IsInContainer and GetDefaultValue() can accurately determine if a class contains a property when delta serializing and other similar operations. Properties in children classes can have offsets that would fall within tail padding of the parent class and sizeof() does not measure the data size of the class, but instead the tail-padded size (to reach the struct's alignment)

#Synced-CL 51724748

diff --git a/Engine/Source/Programs/Shared/EpicGames.UHT/Exporters/CodeGen/UhtHeaderCodeGeneratorCppFile.cs b/Engine/Source/Programs/Shared/EpicGames.UHT/Exporters/CodeGen/UhtHeaderCodeGeneratorCppFile.cs
--- a/Engine/Source/Programs/Shared/EpicGames.UHT/Exporters/CodeGen/UhtHeaderCodeGeneratorCppFile.cs
+++ b/Engine/Source/Programs/Shared/EpicGames.UHT/Exporters/CodeGen/UhtHeaderCodeGeneratorCppFile.cs
@@ -923,7 +923,7 @@ namespace EpicGames.UHT.Exporters.CodeGen
 				}
 
 				// Declare a function to access the size/alignment for NoExport cases as well, consteval should prevent this function having code in the binary
-				builder.Append("\tstatic inline consteval int32 GetStructSize() { return sizeof(").Append(scriptStruct.NamespaceExportName).Append(scriptStruct.SourceName).Append("); }\r\n");
+				builder.Append("\tstatic inline consteval int32 GetStructSize() { return DataSizeOf<").Append(scriptStruct.NamespaceExportName).Append(scriptStruct.SourceName).Append(">(); }\r\n");
 				builder.Append("\tstatic inline consteval int16 GetStructAlignment() { return alignof(").Append(scriptStruct.NamespaceExportName).Append(scriptStruct.SourceName).Append("); }\r\n");
 
 				// Meta data
@@ -1163,7 +1163,7 @@ namespace EpicGames.UHT.Exporters.CodeGen
 					builder.Append('\t').AppendUTF8LiteralString(scriptStruct.EngineName).Append(",\r\n");
 					builder.AppendArrayPtrLine(properties, UhtNames.StaticsPropPointersId, 1, ",\r\n");
 					builder.AppendArrayCountLine(properties, UhtNames.StaticsPropPointersId, 1, ",\r\n");
-					builder.Append("\tsizeof(").Append(scriptStruct.NamespaceExportName).Append(scriptStruct.SourceName).Append("),\r\n");
+					builder.Append("\tDataSizeOf<").Append(scriptStruct.NamespaceExportName).Append(scriptStruct.SourceName).Append(">(),\r\n");
 					builder.Append("\talignof(").Append(scriptStruct.NamespaceExportName).Append(scriptStruct.SourceName).Append("),\r\n");
 					builder.Append("\tRF_Public|RF_Transient|RF_MarkAsNative,\r\n");
 					builder.Append($"\tEStructFlags(0x{(uint)(scriptStruct.ScriptStructFlags & ~EStructFlags.ComputedFlags):X8}),\r\n");
@@ -1576,6 +1576,7 @@ namespace EpicGames.UHT.Exporters.CodeGen
 				string functionEngineName = isVerseNativeCallable ? function.GetEncodedVerseScopeAndName(UhtVerseNameMode.Default) : function.EngineName;
 
 				string? sizeOfStatic = null;
+				string? dataSizeOfStatic = null;
 				string? alignOfStatic = null;
 				if (function.Children.Count > 0)
 				{
@@ -1587,6 +1588,7 @@ namespace EpicGames.UHT.Exporters.CodeGen
 
 					string eventStructName = GetEventStructParametersName(tempFunction.Outer, tempFunction.StrippedFunctionName);
 					UhtCppIdentifier identifier = paramsInStatic ? new UhtCppIdentifier(eventStructName).MakeStatics() : new UhtCppIdentifier(eventStructName);
+					dataSizeOfStatic = $"DataSizeOf<{identifier}>()";
 					sizeOfStatic = $"sizeof({identifier})";
 					alignOfStatic = $"alignof({identifier})";
 				}
@@ -1602,7 +1604,7 @@ namespace EpicGames.UHT.Exporters.CodeGen
 					builder.Append("\tUE::CodeGen::ConstInit::FStructParams{\r\n")
 						.AppendLine($"\t\t.Super = {ConstInitSingletonRef(this, function.SuperFunction)}, ")
 						.Append("\t\t.ChildProperties = UHT_STATICS::GetFirstProperty(),\r\n")
-						.Append(sizeOfStatic is null ? null : $"\t\t.PropertiesSize = {sizeOfStatic},\r\n")
+						.Append(dataSizeOfStatic is null ? null : $"\t\t.PropertiesSize = {dataSizeOfStatic},\r\n")
 						.Append(alignOfStatic is null ? null : $"\t\t.MinAlignment = {alignOfStatic},\r\n")
 						.Append("\t\tIF_WITH_METADATA(.MetaData = ").AppendMetaDataArrayView(function.MetaData, UhtNames.StaticsTypeId, 0, ",)\r\n")
 						.Append("\t},\r\n");
@@ -1654,7 +1656,7 @@ namespace EpicGames.UHT.Exporters.CodeGen
 						.AppendUTF8LiteralString(functionEngineName).Append(", ")
 						.AppendArrayPtrLine(properties, UhtNames.StaticsPropPointersId, 0, ", ")
 						.AppendArrayCountLine(properties, UhtNames.StaticsPropPointersId, 0, ", ")
-						.Append(sizeOfStatic ?? "0").Append(", ")
+						.Append(dataSizeOfStatic ?? "0").Append(", ")
 						.Append("RF_Public|RF_Transient|RF_MarkAsNative, ")
 						.Append($"(EFunctionFlags)0x{(uint)function.FunctionFlags:X8}, ")
 						.Append(isNet ? function.RPCId : 0).Append(", ")
@@ -3202,7 +3204,7 @@ namespace EpicGames.UHT.Exporters.CodeGen
 						.Append("\t\t.FirstChild = UHT_STATICS::GetFirstFunction(),\r\n")
 						.Append("\t\t.BaseChain = MakeArrayView(UHT_STATICS::StructBases),\r\n")
 						.Append("\t\t.ChildProperties = UHT_STATICS::GetFirstProperty(),\r\n")
-						.Append($"\t\t.PropertiesSize = sizeof({classObj.Namespace.FullSourceName}{classObj.SourceName}),\r\n")
+						.Append($"\t\t.PropertiesSize = DataSizeOf<{classObj.Namespace.FullSourceName}{classObj.SourceName}>(),\r\n")
 						.Append($"\t\t.MinAlignment = alignof({classObj.Namespace.FullSourceName}{classObj.SourceName}),\r\n")
 						.Append("\t\tIF_WITH_METADATA(.MetaData = ").AppendMetaDataArrayView(classObj.MetaData, UhtNames.StaticsTypeId, 0, ",)\r\n")
 						.Append("\t},\r\n")
@@ -3333,7 +3335,7 @@ namespace EpicGames.UHT.Exporters.CodeGen
 				builder.Append("\t\t\t\tTEXT(\"").Append(classObj.EngineName).Append("\"),\r\n");
 				builder.Append("\t\t\t\t").Append(registrationName).Append(".InnerSingleton,\r\n");
 				builder.Append("\t\t\t\t").Append(registerNativesFuncName).Append(",\r\n");
-				builder.Append("\t\t\t\tsizeof(TClass),\r\n");
+				builder.Append("\t\t\t\tDataSizeOf<TClass>(),\r\n");
 				builder.Append("\t\t\t\talignof(TClass),\r\n");
 				builder.Append("\t\t\t\tTClass::StaticClassFlags,\r\n");
 				builder.Append("\t\t\t\tTClass::StaticClassCastFlags(),\r\n");
diff --git a/Engine/Source/Runtime/Core/Public/Templates/AlignmentTemplates.h b/Engine/Source/Runtime/Core/Public/Templates/AlignmentTemplates.h
--- a/Engine/Source/Runtime/Core/Public/Templates/AlignmentTemplates.h
+++ b/Engine/Source/Runtime/Core/Public/Templates/AlignmentTemplates.h
@@ -5,6 +5,7 @@
 #include "CoreTypes.h"
 #include "Concepts/Integral.h"
 #include "Concepts/Pointer.h"
+#include <type_traits>
 
 /**
  * Aligns a value to the nearest higher multiple of 'Alignment', which must be a power of two.
@@ -66,6 +67,51 @@ UE_FORCEINLINE_HINT constexpr T AlignArbitrary(T Val, uint64 Alignment)
 	return (T)((((uint64)Val + Alignment - 1) / Alignment) * Alignment);
 }
 
+/**
+ * Gets the size of the properties of a type, without any trailing padding. Used to populate PropertiesSize for UStructs for native types.
+ * This is important because we use the offset of properties to determine what class they are in (i.e. an offset that falls within a parent
+ * class is assumed to be a property of the parent class). Since compilers can reuse the tail padding of a class in derived classes, PropertiesSize
+ * must be no more than the offset of the first byte of a derived class for non-final classes.
+ * Unfortunately, this implementation does not work for final types, so we fall back to just using sizeof(T) which what was used previously
+ * Note that for types that have a virtual destructor, we declare a pure virtual destructor in Derived in order to avoid calling the destructor
+ *     in T and triggering compile errors related to private destructors or destructors for incomplete types
+ *
+ * @return sizeof(T) minus the unused tail alignment bytes that subclasses may reuse. Unaligned.
+ */
+
+template <typename T>
+consteval std::size_t DataSizeOf()
+{
+	if constexpr (std::is_final_v<T>)
+	{
+		// Final type. We can't check the offset of a child property so we just use sizeof
+		return sizeof(T);
+	}
+	else if constexpr(std::is_polymorphic_v<T>)
+	{
+		// Polymorhpic type. Use a pure virtual destructor to avoid compiler errors related to accessing private destructors
+#if defined(_MSC_VER) && (_MSC_VER < 1950)
+		// There was a bug in msvc that incorrectly reported "= 0" as not meaning "pure virtual" in structs defined in templates
+		__pragma(warning(push)) __pragma(warning(disable : 5288))
+#endif
+			struct PLATFORM_EMPTY_BASES Derived : T { virtual ~Derived() = 0; char Byte; };
+#if defined(_MSC_VER) && (_MSC_VER < 1950)
+		__pragma(warning(pop))
+#endif
+		constexpr std::size_t DataSize = offsetof(Derived, Byte);
+		static_assert(DataSize <= sizeof(T));
+		return (DataSize < 1) ? std::size_t(1) : DataSize;
+	}
+	else
+	{
+		// Non-polymorhpic type. Delete the destructor to avoid compiler errors related to accessing private destructors
+		struct PLATFORM_EMPTY_BASES Derived : T { ~Derived() = delete; char Byte; };
+		constexpr std::size_t DataSize = offsetof(Derived, Byte);
+		static_assert(DataSize <= sizeof(T));
+		return (DataSize < 1) ? std::size_t(1) : DataSize;
+	}
+}
+
 #if UE_ENABLE_INCLUDE_ORDER_DEPRECATED_IN_5_8
 #include "Templates/IsIntegral.h"
 #include "Templates/IsPointer.h"
diff --git a/Engine/Source/Runtime/CoreUObject/Private/UObject/Class.cpp b/Engine/Source/Runtime/CoreUObject/Private/UObject/Class.cpp
--- a/Engine/Source/Runtime/CoreUObject/Private/UObject/Class.cpp
+++ b/Engine/Source/Runtime/CoreUObject/Private/UObject/Class.cpp
@@ -3703,7 +3703,7 @@ void UScriptStruct::CopyScriptStruct(void* InDest, void const* InSrc, int32 Arra
 		check(!(StructFlags & STRUCT_IsPlainOldData)); // should not have both
 		UScriptStruct::ICppStructOps* TheCppStructOps = GetCppStructOps();
 		check(TheCppStructOps);
-		check(Stride == TheCppStructOps->GetSize() && PropertiesSize == Stride);
+		check(Stride == TheCppStructOps->GetSize() && PropertiesSize <= Stride);
 		if (TheCppStructOps->Copy(Dest, Src, ArrayDim))
 		{
 			return;
@@ -3796,7 +3796,7 @@ void UScriptStruct::InitializeStruct(void* InDest, int32 ArrayDim) const
 
 		InitializedSize = TheCppStructOps->GetSize();
 		// here we want to make sure C++ and the property system agree on the size
-		check(Stride == InitializedSize && PropertiesSize == InitializedSize);
+		check(Stride == InitializedSize && PropertiesSize <= InitializedSize);
 	}
 
 	if (PropertiesSize > InitializedSize)
@@ -3899,7 +3899,7 @@ void UScriptStruct::ClearScriptStruct(void* Dest, int32 ArrayDim) const
 		}
 		ClearedSize = TheCppStructOps->GetSize();
 		// here we want to make sure C++ and the property system agree on the size
-		check(Stride == ClearedSize && PropertiesSize == ClearedSize);
+		check(Stride == ClearedSize && PropertiesSize <= ClearedSize);
 	}
 	if ( PropertiesSize > ClearedSize )
 	{
@@ -3948,8 +3948,8 @@ void UScriptStruct::DestroyStruct(void* Dest, int32 ArrayDim) const
 		}
 		ClearedSize = TheCppStructOps->GetSize();
 		// here we want to make sure C++ and the property system agree on the size
-		checkf(Stride == ClearedSize && PropertiesSize == ClearedSize, TEXT("C++ and the property system struct size mismatch for %s (C++ Size: %d, Property Size: %d)"),
-			*GetPathName(), ClearedSize, Stride);
+		checkf(Stride == ClearedSize && PropertiesSize <= ClearedSize, TEXT("C++ and the property system struct size mismatch for %s (C++ Size: %d, Aligned Property Size: %d, Unaligned Property Size: %d)"),
+			*GetPathName(), ClearedSize, Stride, PropertiesSize);
 	}
 
 	if (PropertiesSize > ClearedSize)
diff --git a/Engine/Source/Runtime/CoreUObject/Public/UObject/ObjectMacros.h b/Engine/Source/Runtime/CoreUObject/Public/UObject/ObjectMacros.h
--- a/Engine/Source/Runtime/CoreUObject/Public/UObject/ObjectMacros.h
+++ b/Engine/Source/Runtime/CoreUObject/Public/UObject/ObjectMacros.h
@@ -2301,7 +2301,7 @@ private:
 				(TCHAR*)TEXT(#TClass) + 1 + ((TClass::StaticClassFlags & CLASS_Deprecated) ? 11 : 0), \
 				Z_Registration_Info_UClass_##TClass.InnerSingleton, \
 				TClass::StaticRegisterNatives##TClass, \
-				sizeof(TClass), \
+				DataSizeOf<TClass>(), \
 				alignof(TClass), \
 				TClass::StaticClassFlags, \
 				TClass::StaticClassCastFlags(), \
@@ -2394,7 +2394,7 @@ private:
 		UE::CodeGen::ConstInit::FStructParams{ \
 			.Super = &Z_ConstInit_UClass_##TSuperClass, \
 			.BaseChain = TIntrinsicClassBaseChain<TClass>::ToArrayView(), \
-			.PropertiesSize = sizeof(TClass), \
+			.PropertiesSize = DataSizeOf<TClass>(), \
 			.MinAlignment = alignof(TClass), \
 		}, \
 		UE::CodeGen::ConstInit::FClassParams{ \
