EpicGames.Core: Don't duplicate the same user agent into the default http headers for telemetry

#rnx
#Synced-CL 51936141

diff --git a/Engine/Source/Programs/Shared/EpicGames.Core/Telemetry/DataRouterTelemetryService.cs b/Engine/Source/Programs/Shared/EpicGames.Core/Telemetry/DataRouterTelemetryService.cs
--- a/Engine/Source/Programs/Shared/EpicGames.Core/Telemetry/DataRouterTelemetryService.cs
+++ b/Engine/Source/Programs/Shared/EpicGames.Core/Telemetry/DataRouterTelemetryService.cs
@@ -385,17 +385,20 @@ namespace EpicGames.Core.Telemetry
 		/// </summary>
 		private void SetHttpClientHeaders()
 		{
-			Assembly? entryAssembly = Assembly.GetEntryAssembly();
-			if (entryAssembly != null)
+			if (HttpClient.DefaultRequestHeaders.UserAgent.Count == 0)
 			{
-				string productName = String.IsNullOrEmpty(entryAssembly.GetName().Name) ? GetType().Name : entryAssembly.GetName().Name!;
-				string productVersion = entryAssembly.GetName().Version == null ? "1.0.0.0" : entryAssembly.GetName().Version!.ToString();
-				ProductInfoHeaderValue userAgent = new(productName, productVersion);
-				HttpClient.DefaultRequestHeaders.UserAgent.Add(userAgent);
-			}
-			else
-			{
-				HttpClient.DefaultRequestHeaders.UserAgent.Add(new ProductInfoHeaderValue(GetType().Name, "1.0.0.0"));
+				Assembly? entryAssembly = Assembly.GetEntryAssembly();
+				if (entryAssembly != null)
+				{
+					string productName = String.IsNullOrEmpty(entryAssembly.GetName().Name) ? GetType().Name : entryAssembly.GetName().Name!;
+					string productVersion = entryAssembly.GetName().Version == null ? "1.0.0.0" : entryAssembly.GetName().Version!.ToString();
+					ProductInfoHeaderValue userAgent = new(productName, productVersion);
+					HttpClient.DefaultRequestHeaders.UserAgent.Add(userAgent);
+				}
+				else
+				{
+					HttpClient.DefaultRequestHeaders.UserAgent.Add(new ProductInfoHeaderValue(GetType().Name, "1.0.0.0"));
+				}
 			}
 		}
 
