UHT: do not generate header for _implementation functions if they are already defined
This allow to declare them as non-virtual, or as override

#rnx
#Synced-CL 52059414

diff --git a/Engine/Source/Programs/Shared/EpicGames.UHT/Exporters/CodeGen/UhtHeaderCodeGeneratorHFile.cs b/Engine/Source/Programs/Shared/EpicGames.UHT/Exporters/CodeGen/UhtHeaderCodeGeneratorHFile.cs
--- a/Engine/Source/Programs/Shared/EpicGames.UHT/Exporters/CodeGen/UhtHeaderCodeGeneratorHFile.cs
+++ b/Engine/Source/Programs/Shared/EpicGames.UHT/Exporters/CodeGen/UhtHeaderCodeGeneratorHFile.cs
@@ -1294,8 +1294,14 @@ namespace EpicGames.UHT.Exporters.CodeGen
 				{
 					continue;
 				}
-				AppendNetValidateDeclaration(builder, classObj, function);
-				AppendNativeFunctionHeader(builder, function, UhtPropertyTextType.ClassFunctionArgOrRetVal, true, null, null, UhtFunctionExportFlags.None, 1, "; \\\r\n");
+				if (!function.FunctionExportFlags.HasAnyFlags(UhtFunctionExportFlags.ValidationImplFound))
+				{
+					AppendNetValidateDeclaration(builder, classObj, function);
+				}
+				if (!function.FunctionExportFlags.HasAnyFlags(UhtFunctionExportFlags.ImplFound))
+				{
+					AppendNativeFunctionHeader(builder, function, UhtPropertyTextType.ClassFunctionArgOrRetVal, true, null, null, UhtFunctionExportFlags.None, 1, "; \\\r\n");
+				}
 			}
 			return builder;
 		}
@@ -1821,7 +1827,7 @@ namespace EpicGames.UHT.Exporters.CodeGen
 
 			// we'll need a way to get to the UObject portion of a native interface, so that we can safely pass native interfaces
 			// to script VM functions
-			if (classObj.IsVerseField || (classObj.SuperClass != null && classObj.SuperClass.IsChildOf(Session.UInterface)))
+			if (classObj.IsVerseField || (classObj.SuperClass != null && classObj.SuperClass == Session.UInterface))
 			{
 				// Note: This used to be declared as a pure virtual function, but it was changed here in order to allow the Blueprint nativization process
 				// to detect C++ interface classes that explicitly declare pure virtual functions via type traits. This code will no longer trigger that check.
