Ignore systemic unactionable warnings for update-alternatives when building docker containers

#rnx
#Synced-CL 52066884

diff --git a/Engine/Source/Programs/Shared/EpicGames.Core/KnownLogEvents.cs b/Engine/Source/Programs/Shared/EpicGames.Core/KnownLogEvents.cs
--- a/Engine/Source/Programs/Shared/EpicGames.Core/KnownLogEvents.cs
+++ b/Engine/Source/Programs/Shared/EpicGames.Core/KnownLogEvents.cs
@@ -356,6 +356,11 @@ namespace EpicGames.Core
 		/// </summary>
 		public static EventId Systemic_UbaCacheResponseTimeout { get; } = new EventId(770);
 
+		/// <summary>
+		/// Docker messages that are unactionable and can be ignored
+		/// </summary>
+		public static EventId Systemic_Docker { get; } = new EventId(771);
+
 		/// <summary>
 		/// Maximum code for systemic events. Add new events in the 700-799 range.
 		/// </summary>
diff --git a/Engine/Source/Programs/Shared/EpicGames.LogMatchers/DockerEventMatcher.cs b/Engine/Source/Programs/Shared/EpicGames.LogMatchers/DockerEventMatcher.cs
--- a/Engine/Source/Programs/Shared/EpicGames.LogMatchers/DockerEventMatcher.cs
+++ b/Engine/Source/Programs/Shared/EpicGames.LogMatchers/DockerEventMatcher.cs
@@ -12,7 +12,9 @@ namespace EpicGames.LogMatchers
 	class DockerEventMatcher : ILogEventMatcher
 	{
 		static readonly Regex s_noSpacePattern = new(": no space left on device$");
+		static readonly Regex s_updateAlternativesPattern = new(@"^update-alternatives: warning:");
 
+		/// <inheritdoc/>
 		public LogEventMatch? Match(ILogCursor cursor)
 		{
 			if (cursor.IsMatch(s_noSpacePattern))
@@ -20,6 +22,11 @@ namespace EpicGames.LogMatchers
 				LogEventBuilder builder = new(cursor);
 				return builder.ToMatch(LogEventPriority.Normal, LogLevel.Error, KnownLogEvents.Systemic_OutOfDiskSpace);
 			}
+			if (cursor.IsMatch(s_updateAlternativesPattern))
+			{
+				LogEventBuilder builder = new(cursor);
+				return builder.ToMatch(LogEventPriority.Normal, LogLevel.Information, KnownLogEvents.Systemic_Docker);
+			}
 			return null;
 		}
 	}
