Add error code for when we fail to query which processes have a file handle open.
Also added code to handle if more than ten processes have the file handle open.

#Synced-CL 49593614

diff --git a/Engine/Source/Programs/Shared/EpicGames.Core/FileUtils.cs b/Engine/Source/Programs/Shared/EpicGames.Core/FileUtils.cs
--- a/Engine/Source/Programs/Shared/EpicGames.Core/FileUtils.cs
+++ b/Engine/Source/Programs/Shared/EpicGames.Core/FileUtils.cs
@@ -1016,9 +1016,18 @@ namespace EpicGames.Core
 				uint reason = 0;
 				RM_PROCESS_INFO[] processInfoArray = new RM_PROCESS_INFO[nProcInfo];
 				result = RmGetList(sessionHandle, out nProcInfoNeeded, ref nProcInfo, processInfoArray, ref reason);
+				
+				// ERROR_MORE_DATA thrown when more processes have an open file handle
+				if (result == 234)
+				{
+					nProcInfo = Math.Min(nProcInfoNeeded, 256); // Cap to 256 processes just in case
+					processInfoArray = new RM_PROCESS_INFO[nProcInfo];
+					result = RmGetList(sessionHandle, out nProcInfoNeeded, ref nProcInfo, processInfoArray, ref reason);
+				}
+
 				if (result != 0)
 				{
-					throw new Win32Exception(result, "Unable to query processes with file handle open");
+					throw new Win32Exception(result, $"(error {result}) - Unable to query processes with file handle open");
 				}
 
 				List<FileLockInfoWin32> fileLocks = [];
