P4VUtils: Treat empty or whitespace P4 config values as null to make sure they get handled correctly when passing values to P4. This improves error messages when, for example, trying to run preflights on a CL when no workspace has been selected.

#Synced-CL 52128377

diff --git a/Engine/Extras/P4VUtils/Commands/FastReconcileEdits.cs b/Engine/Extras/P4VUtils/Commands/FastReconcileEdits.cs
--- a/Engine/Extras/P4VUtils/Commands/FastReconcileEdits.cs
+++ b/Engine/Extras/P4VUtils/Commands/FastReconcileEdits.cs
@@ -78,7 +78,7 @@ namespace P4VUtils.Commands
 			ClientRecord Client = await Perforce.GetClientAsync(ClientName, CancellationToken.None);
 			if(Client.Stream == null)
 			{
-				Logger.LogError("Not being run from a stream client");
+				Logger.LogError("Not being run from a stream client. If running through P4V, make sure you've selected the correct workspace for these paths.");
 				return 1;
 			}
 
diff --git a/Engine/Extras/P4VUtils/Commands/PreflightCommand.cs b/Engine/Extras/P4VUtils/Commands/PreflightCommand.cs
--- a/Engine/Extras/P4VUtils/Commands/PreflightCommand.cs
+++ b/Engine/Extras/P4VUtils/Commands/PreflightCommand.cs
@@ -84,7 +84,7 @@ namespace P4VUtils.Commands
 			Client = await Perforce.GetClientAsync(ClientName, CancellationToken.None);
 			if (Client.Stream == null)
 			{
-				Logger.LogError("Not being run from a stream client");
+				Logger.LogError("Not being run from a stream client. If running through P4V, make sure you've selected the correct workspace for this changelist.");
 				return false;
 			}
 
diff --git a/Engine/Extras/P4VUtils/Commands/RobomergePreviewCommand.cs b/Engine/Extras/P4VUtils/Commands/RobomergePreviewCommand.cs
--- a/Engine/Extras/P4VUtils/Commands/RobomergePreviewCommand.cs
+++ b/Engine/Extras/P4VUtils/Commands/RobomergePreviewCommand.cs
@@ -37,7 +37,7 @@ namespace P4VUtils.Commands
 			ClientRecord Client = await Perforce.GetClientAsync(ClientName, CancellationToken.None);
 			if (Client.Stream == null)
 			{
-				Logger.LogError("Not being run from a stream client");
+				Logger.LogError("Not being run from a stream client. If running through P4V, make sure you've selected the correct workspace for this changelist.");
 				return 1;
 			}
 
diff --git a/Engine/Extras/P4VUtils/Commands/SnapShotCommand.cs b/Engine/Extras/P4VUtils/Commands/SnapShotCommand.cs
--- a/Engine/Extras/P4VUtils/Commands/SnapShotCommand.cs
+++ b/Engine/Extras/P4VUtils/Commands/SnapShotCommand.cs
@@ -45,7 +45,7 @@ namespace P4VUtils.Commands
 			ClientRecord Client = await Perforce.GetClientAsync(ClientName, CancellationToken.None);
 			if(Client.Stream == null)
 			{
-				Logger.LogError("Not being run from a stream client");
+				Logger.LogError("Not being run from a stream client. If running through P4V, make sure you've selected the correct workspace for this changelist.");
 				return 1;
 			}
 
diff --git a/Engine/Source/Programs/Shared/EpicGames.Perforce/PerforceConnection.cs b/Engine/Source/Programs/Shared/EpicGames.Perforce/PerforceConnection.cs
--- a/Engine/Source/Programs/Shared/EpicGames.Perforce/PerforceConnection.cs
+++ b/Engine/Source/Programs/Shared/EpicGames.Perforce/PerforceConnection.cs
@@ -259,11 +259,11 @@ namespace EpicGames.Perforce
 		/// <param name="logger">Interface for logging</param>
 		public PerforceConnection(string? serverAndPort, string? userName, string? clientName, string? appName, string? appVersion, ILogger logger, Tracer? tracer = null)
 		{
-			ServerAndPort = serverAndPort;
-			UserName = userName;
-			ClientName = clientName;
-			AppName = appName;
-			AppVersion = appVersion;
+			ServerAndPort = String.IsNullOrWhiteSpace(serverAndPort) ? null : serverAndPort;
+			UserName = String.IsNullOrWhiteSpace(userName) ? null : userName;
+			ClientName = String.IsNullOrWhiteSpace(clientName) ? null : clientName;
+			AppName = String.IsNullOrWhiteSpace(appName) ? null : appName;
+			AppVersion = String.IsNullOrWhiteSpace(appVersion) ? null : appVersion;
 			Logger = logger;
 			Tracer = tracer ?? DefaultTracer;
 		}
