UnrealAutomationTool: Use AssemblyLoadContext.Default.LoadFromAssemblyPath to load assemblies in our assembly resolver, which puts them in a shared context and doesn't result in multiple, incompatible copies of the same assembly being loaded and causing errors in automation scripts.

Additionally, remove the specialized resolver for Ionic.Zip.Reduced, and just copy the assembly locally, which should speed up assembly resolves for all other uses that don't involve this assembly and make our assembly resolving logic a bit more standard. This was also used for SwarmAgent, but that's never actually loaded by anything.

#Synced-CL 52202922

diff --git a/Engine/Source/Programs/AutomationTool/BuildGraph/BuildGraph.Automation.csproj b/Engine/Source/Programs/AutomationTool/BuildGraph/BuildGraph.Automation.csproj
--- a/Engine/Source/Programs/AutomationTool/BuildGraph/BuildGraph.Automation.csproj
+++ b/Engine/Source/Programs/AutomationTool/BuildGraph/BuildGraph.Automation.csproj
@@ -18,7 +18,7 @@
     </Reference>
     <Reference Include="Ionic.Zip.Reduced">
       <HintPath>..\..\..\..\Binaries\DotNET\Ionic.Zip.Reduced.dll</HintPath>
-      <Private>False</Private>
+      <Private>True</Private>
     </Reference>
   </ItemGroup>
   <ItemGroup>
diff --git a/Engine/Source/Programs/AutomationTool/IOS/IOS.Automation.csproj b/Engine/Source/Programs/AutomationTool/IOS/IOS.Automation.csproj
--- a/Engine/Source/Programs/AutomationTool/IOS/IOS.Automation.csproj
+++ b/Engine/Source/Programs/AutomationTool/IOS/IOS.Automation.csproj
@@ -19,6 +19,7 @@
     </Reference>
     <Reference Include="Ionic.Zip.Reduced">
       <HintPath>..\..\..\..\Binaries\DotNET\Ionic.Zip.Reduced.dll</HintPath>
+      <Private>True</Private>
     </Reference>
   </ItemGroup>
 
diff --git a/Engine/Source/Programs/AutomationTool/Program.cs b/Engine/Source/Programs/AutomationTool/Program.cs
--- a/Engine/Source/Programs/AutomationTool/Program.cs
+++ b/Engine/Source/Programs/AutomationTool/Program.cs
@@ -448,7 +448,6 @@ namespace AutomationToolDriver
 
 				// Ensure we can resolve any external assemblies as necessary.
 				DirectoryReference AutomationToolLocation = new DirectoryReference(Path.GetDirectoryName(Assembly.GetEntryAssembly().GetOriginalLocation()));
-				AssemblyUtils.InstallAssemblyResolver(AutomationToolLocation.FullName);
 				AssemblyUtils.InstallRecursiveAssemblyResolver(AutomationToolLocation.FullName);
 
 				// Log the operating environment. Since we usually compile to AnyCPU, we may be executed using different system paths under WOW64.
@@ -522,9 +521,9 @@ namespace AutomationToolDriver
 					Logger);
 
 			if (!bBuildSuccess)
-            {
+			{
 				return ExitCode.Error_Unknown;
-            }
+			}
 
 			// when the engine is installed, or UAT is invoked with -NoCompile, we expect to find at least one script module (AutomationUtils is a necessity)
 			if (ScriptModuleAssemblyPaths.Count == 0)
@@ -538,12 +537,11 @@ namespace AutomationToolDriver
 			Assembly AutomationUtilsAssembly = AssemblyPath != null ? Assembly.LoadFrom(AssemblyPath.FullName) : null;
 
 			if (AutomationUtilsAssembly == null)
-            {
+			{
 				throw new Exception("Did not find an AutomationUtils.Automation.dll");
-            }
-
-			// Call into AutomationTool.Automation.Process()
+			}
 
+			// Call into AutomationTool.Automation.ProcessAsync()
 			Type AutomationTools_Automation = AutomationUtilsAssembly.GetType("AutomationTool.Automation");
 			MethodInfo Automation_Process = AutomationTools_Automation.GetMethod("ProcessAsync");
 			Logger.LogInformation("Total script module initialization time: {InitTime:0.00} s.", Timer.Elapsed.TotalSeconds);
diff --git a/Engine/Source/Programs/Shared/EpicGames.Core/AssemblyUtils.cs b/Engine/Source/Programs/Shared/EpicGames.Core/AssemblyUtils.cs
--- a/Engine/Source/Programs/Shared/EpicGames.Core/AssemblyUtils.cs
+++ b/Engine/Source/Programs/Shared/EpicGames.Core/AssemblyUtils.cs
@@ -4,8 +4,8 @@ using System;
 using System.Collections.Generic;
 using System.Diagnostics;
 using System.IO;
-using System.Linq;
 using System.Reflection;
+using System.Runtime.Loader;
 
 namespace EpicGames.Core
 {
@@ -16,7 +16,7 @@ namespace EpicGames.Core
 	{
 		/// <summary>
 		/// Gets the original location (path and filename) of an assembly.
-		/// This method is using Assembly.CodeBase property to properly resolve original
+		/// This method is using Assembly.Location property to properly resolve original
 		/// assembly path in case shadow copying is enabled.
 		/// </summary>
 		/// <returns>Absolute path and filename to the assembly.</returns>
@@ -30,59 +30,16 @@ namespace EpicGames.Core
 		/// </summary>
 		public static FileVersionInfo ExecutableVersion => FileVersionInfo.GetVersionInfo(Assembly.GetEntryAssembly()!.GetOriginalLocation());
 
-		/// <summary>
-		/// Installs an assembly resolver. Mostly used to get shared assemblies that we don't want copied around to various output locations as happens when "Copy Local" is set to true
-		/// for an assembly reference (which is the default).
-		/// </summary>
-		public static void InstallAssemblyResolver(string pathToBinariesDotNET)
-		{
-			string[] knownAssemblyNames = [ "SwarmAgent.exe", "../ThirdParty/Ionic/Ionic.Zip.Reduced.dll", ];
-			AppDomain.CurrentDomain.AssemblyResolve += (sender, args) =>
-			{
-				// Name is fully qualified assembly definition - e.g. "p4dn, Version=1.0.0.0, Culture=neutral, PublicKeyToken=ff968dc1933aba6f"
-				string assemblyName = args.Name!.Split(',')[0];
-
-				return (
-					from knownAssemblyName in knownAssemblyNames
-					where assemblyName.Equals(Path.GetFileNameWithoutExtension(knownAssemblyName), StringComparison.OrdinalIgnoreCase)
-					let resolvedAssemblyFilename = Path.Combine(pathToBinariesDotNET, knownAssemblyName)
-					// check if the file exists first. If we just try to load it, we correctly throw an exception, but it's a generic
-					// FileNotFoundException, which is not informative. Better to return null.
-					select File.Exists(resolvedAssemblyFilename) ? Assembly.LoadFile(resolvedAssemblyFilename) : null
-					).FirstOrDefault();
-			};
-		}
-
 		/// <summary>
 		/// Installs an assembly resolver, which will load *any* assembly which exists recursively within the supplied folder.
 		/// </summary>
 		/// <param name="rootDirectory">The directory to enumerate.</param>
 		public static void InstallRecursiveAssemblyResolver(string rootDirectory)
-		{
-			RefreshAssemblyCache(rootDirectory);
-
-			AppDomain.CurrentDomain.AssemblyResolve += (sender, args) =>
-			{
-				// Name is fully qualified assembly definition - e.g. "p4dn, Version=1.0.0.0, Culture=neutral, PublicKeyToken=ff968dc1933aba6f"
-				string assemblyName = args.Name!.Split(',')[0];
-
-				// The assembly wasn't found by other resolvers, though may have been compiled or copied as a dependency
-				RefreshAssemblyCache(rootDirectory, String.Format("{0}.dll", assemblyName));
-				if (s_assemblyLocationCache.TryGetValue(assemblyName, out string? assemblyLocation))
-				{
-					return Assembly.LoadFile(assemblyLocation);
-				}
-
-				return null;
-			};
-		}
-
-		private static void RefreshAssemblyCache(string rootDirectory, string pattern = "*.dll")
 		{
 			if (Directory.Exists(rootDirectory))
 			{
 				// Initialize our cache of assemblies by enumerating all files in the given folder.
-				foreach (string discoveredAssembly in Directory.EnumerateFiles(rootDirectory, pattern, SearchOption.AllDirectories))
+				foreach (string discoveredAssembly in Directory.EnumerateFiles(rootDirectory, "*.dll", SearchOption.AllDirectories))
 				{
 					AddFileToAssemblyCache(discoveredAssembly);
 				}
@@ -90,7 +47,7 @@ namespace EpicGames.Core
 		}
 
 		/// <summary>
-		/// Adds a file to the cache
+		/// Adds a file to the assembly cache.
 		/// </summary>
 		/// <param name="assemblyPath"></param>
 		public static void AddFileToAssemblyCache(string assemblyPath)
@@ -121,6 +78,11 @@ namespace EpicGames.Core
 				}
 			}
 
+			AddAssemblyCacheResolverIfNecessary();
+		}
+
+		private static void AddAssemblyCacheResolverIfNecessary()
+		{
 			if (!s_addedToAssemblyResolver)
 			{
 				AppDomain.CurrentDomain.AssemblyResolve += (sender, args) =>
@@ -134,8 +96,8 @@ namespace EpicGames.Core
 							// We have this assembly in our folder.
 							if (File.Exists(assemblyLocation))
 							{
-								// The assembly still exists, so load it.
-								return Assembly.LoadFile(assemblyLocation);
+								// The assembly still exists, so load it into the default, shared context, so we only get one version of the assembly.
+								return AssemblyLoadContext.Default.LoadFromAssemblyPath(assemblyLocation);
 							}
 							else
 							{
diff --git a/Engine/Source/Programs/UnrealAndroidFileTool/Program.cs b/Engine/Source/Programs/UnrealAndroidFileTool/Program.cs
--- a/Engine/Source/Programs/UnrealAndroidFileTool/Program.cs
+++ b/Engine/Source/Programs/UnrealAndroidFileTool/Program.cs
@@ -1,18 +1,11 @@
 // Copyright Epic Games, Inc. All Rights Reserved.
 
 using System;
-using System.Threading;
-using System.Threading.Tasks;
-using System.Diagnostics;
-using System.Reflection;
-using System.IO;
 using System.Collections.Generic;
-using System.Text;
-using System.Linq;
+using System.IO;
 using AutomationTool;
 using EpicGames.Core;
 using Microsoft.Extensions.Logging;
-using Microsoft.CodeAnalysis.CSharp.Syntax;
 
 namespace UnrealAndroidFileTool
 {
@@ -1206,18 +1199,6 @@ namespace UnrealAndroidFileTool
 			string PackageName = "";
 			string Token = "";
 
-			// Ensure we can resolve any external assemblies as necessary.
-			Assembly? assembly = Assembly.GetEntryAssembly();
-			if (assembly != null)
-			{
-				string? PathToBinariesDotNET = Path.GetDirectoryName(AppContext.BaseDirectory);//assembly.GetOriginalLocation());
-				if (PathToBinariesDotNET != null)
-				{
-					AssemblyUtils.InstallAssemblyResolver(PathToBinariesDotNET);
-					AssemblyUtils.InstallRecursiveAssemblyResolver(PathToBinariesDotNET);
-				}
-			}
-
 			if (args.Length == 0 || (args.Length == 1 && args[0] == "help"))
 			{
 				ShowHelp("");
diff --git a/Engine/Source/Programs/UnrealBuildTool/UnrealBuildTool.cs b/Engine/Source/Programs/UnrealBuildTool/UnrealBuildTool.cs
--- a/Engine/Source/Programs/UnrealBuildTool/UnrealBuildTool.cs
+++ b/Engine/Source/Programs/UnrealBuildTool/UnrealBuildTool.cs
@@ -286,9 +286,6 @@ namespace UnrealBuildTool
 				// Configure the progress writer
 				ProgressWriter.bWriteMarkup = Options.WriteProgressMarkup;
 
-				// Ensure we can resolve any external assemblies that are not in the same folder as our assembly.
-				AssemblyUtils.InstallAssemblyResolver(Path.GetDirectoryName(Assembly.GetEntryAssembly()!.GetOriginalLocation())!);
-
 				// Add the application directory to PATH
 				DirectoryReference.AddDirectoryToPath(Unreal.UnrealBuildToolDllPath.Directory);
 
