[Restore]  CL 51877267

#was-backout 52016559

#was-backout 52122917
#Synced-CL 0

#Synced-CL 0

diff --git a/Engine/Source/Programs/Shared/EpicGames.UHT/Utils/UhtSession.cs b/Engine/Source/Programs/Shared/EpicGames.UHT/Utils/UhtSession.cs
--- a/Engine/Source/Programs/Shared/EpicGames.UHT/Utils/UhtSession.cs
+++ b/Engine/Source/Programs/Shared/EpicGames.UHT/Utils/UhtSession.cs
@@ -3101,7 +3101,7 @@ namespace EpicGames.UHT.Utils
 						}
 						if (pluginModule == null)
 						{
-							Logger.LogWarning("Exporter \"{ExporterName}\" skipped because module \"{ModuleName}\" was not found in manifest", exporter.Name, exporter.ModuleName);
+							Logger.LogInformation("Exporter \"{ExporterName}\" skipped because module \"{ModuleName}\" was not found in manifest", exporter.Name, exporter.ModuleName);
 							continue;
 						}
 					}
