EpicGames.Core: Fix GetUtf8String failing on multi-segment Utf8JsonReader input

GetUtf8String only reads reader.ValueSpan and never checks reader.HasValueSequence. When a JSON string token spans multiple buffer segments, ValueSpan is empty per the Utf8JsonReader contract and the method returns an empty byte span, silently corrupting the deserialized value.

diff --git a/Engine/Source/Programs/Shared/EpicGames.Core.Tests/JsonExtensionsTests.cs b/Engine/Source/Programs/Shared/EpicGames.Core.Tests/JsonExtensionsTests.cs
new file mode 100644
--- /dev/null
+++ b/Engine/Source/Programs/Shared/EpicGames.Core.Tests/JsonExtensionsTests.cs
@@ -0,0 +1,72 @@
+// Copyright Epic Games, Inc. All Rights Reserved.
+
+using System;
+using System.Buffers;
+using System.Text;
+using System.Text.Json;
+using Microsoft.VisualStudio.TestTools.UnitTesting;
+
+namespace EpicGames.Core.Tests
+{
+	[TestClass]
+	public class JsonExtensionsTests
+	{
+		sealed class BufferSegment : ReadOnlySequenceSegment<byte>
+		{
+			public BufferSegment(ReadOnlyMemory<byte> memory) => Memory = memory;
+
+			public BufferSegment Append(ReadOnlyMemory<byte> memory)
+			{
+				BufferSegment next = new(memory) { RunningIndex = RunningIndex + Memory.Length };
+				Next = next;
+				return next;
+			}
+		}
+
+		static string GetUtf8StringFromSingleSegment(string jsonLiteral)
+		{
+			byte[] json = Encoding.UTF8.GetBytes(jsonLiteral);
+			Utf8JsonReader reader = new(json);
+			reader.Read();
+			return Encoding.UTF8.GetString(reader.GetUtf8String());
+		}
+
+		static string GetUtf8StringFromMultiSegment(string jsonLiteral, int splitPoint)
+		{
+			byte[] json = Encoding.UTF8.GetBytes(jsonLiteral);
+			if (splitPoint < 0)
+			{
+				splitPoint = json.Length / 2;
+			}
+			byte[] firstPart = json[..splitPoint];
+			byte[] secondPart = json[splitPoint..];
+			BufferSegment first = new(firstPart);
+			BufferSegment last = first.Append(secondPart);
+			ReadOnlySequence<byte> sequence = new(first, 0, last, secondPart.Length);
+
+			Utf8JsonReader reader = new(sequence);
+			reader.Read();
+			return Encoding.UTF8.GetString(reader.GetUtf8String());
+		}
+
+		[TestMethod]
+		[DataRow("\"hello-world\"", "hello-world", DisplayName = "no escapes")]
+		[DataRow("\"hello\\\\world\\\"foo\\/bar\"", "hello\\world\"foo/bar", DisplayName = "backslash escapes")]
+		[DataRow("\"\\u0041\\u0042\\u0043\"", "ABC", DisplayName = "unicode escapes")]
+		[DataRow("\"\"", "", DisplayName = "empty string")]
+		public void SingleSegment(string jsonLiteral, string expected)
+		{
+			Assert.AreEqual(expected, GetUtf8StringFromSingleSegment(jsonLiteral));
+		}
+
+		[TestMethod]
+		[DataRow("\"step-output/test-stream/1/test-artifact/67890abcdef#pkt=4005760,29050&exp=1\"", "step-output/test-stream/1/test-artifact/67890abcdef#pkt=4005760,29050&exp=1", -1, DisplayName = "no escapes")]
+		[DataRow("\"hello\\\\world\\\"foo\\/bar\"", "hello\\world\"foo/bar", -1, DisplayName = "backslash escapes")]
+		[DataRow("\"AB\\u0043DE\"", "ABCDE", 4, DisplayName = "split at escape boundary")]
+		[DataRow("\"\"", "", 1, DisplayName = "empty string")]
+		public void MultiSegment(string jsonLiteral, string expected, int splitPoint)
+		{
+			Assert.AreEqual(expected, GetUtf8StringFromMultiSegment(jsonLiteral, splitPoint));
+		}
+	}
+}
diff --git a/Engine/Source/Programs/Shared/EpicGames.Core/JsonExtensions.cs b/Engine/Source/Programs/Shared/EpicGames.Core/JsonExtensions.cs
--- a/Engine/Source/Programs/Shared/EpicGames.Core/JsonExtensions.cs
+++ b/Engine/Source/Programs/Shared/EpicGames.Core/JsonExtensions.cs
@@ -1,6 +1,7 @@
 // Copyright Epic Games, Inc. All Rights Reserved.
 
 using System;
+using System.Buffers;
 using System.Diagnostics;
 using System.Text.Json;
 
@@ -40,7 +41,19 @@ namespace EpicGames.Core
 		/// <returns></returns>
 		public static ReadOnlySpan<byte> GetUtf8String(this Utf8JsonReader reader)
 		{
-			ReadOnlySpan<byte> span = reader.ValueSpan;
+			ReadOnlySpan<byte> span;
+
+			if (reader.HasValueSequence)
+			{
+				ReadOnlySequence<byte> sequence = reader.ValueSequence;
+				byte[] buffer = new byte[checked((int)sequence.Length)];
+				sequence.CopyTo(buffer);
+				span = buffer;
+			}
+			else
+			{
+				span = reader.ValueSpan;
+			}
 
 			// Scan for the first backslash
 			for (int firstIdx = 0; firstIdx < span.Length; firstIdx++)
