Fix multi-segment JSON deserialization in converter classes (ValueSpan -> GetUtf8String)

.NET 10 changed ASP.NET [FromBody] deserialization from Stream-based to PipeReader-based. Converters that read reader.ValueSpan without checking HasValueSequence get empty/corrupt data.

diff --git a/Engine/Source/Programs/Shared/EpicGames.Core.Tests/LogEventTests.cs b/Engine/Source/Programs/Shared/EpicGames.Core.Tests/LogEventTests.cs
--- a/Engine/Source/Programs/Shared/EpicGames.Core.Tests/LogEventTests.cs
+++ b/Engine/Source/Programs/Shared/EpicGames.Core.Tests/LogEventTests.cs
@@ -1,7 +1,10 @@
 // Copyright Epic Games, Inc. All Rights Reserved.
 
 using System;
+using System.Buffers;
 using System.Collections.Generic;
+using System.Text;
+using System.Text.Json;
 using Microsoft.Extensions.Logging;
 using Microsoft.VisualStudio.TestTools.UnitTesting;
 
@@ -48,5 +51,89 @@ namespace EpicGames.Core.Tests
 			string jevs2 = jev2.ToString();
 			Assert.AreEqual(jevs, jevs2);
 		}
+		
+		[TestMethod]
+		[DataRow(5)]   // splits inside "message" (top-level property name)
+		[DataRow(42)]  // splits inside "properties" (top-level property name)
+		[DataRow(52)]  // splits inside "Foo" (nested property in ReadProperties)
+		[DataRow(60)]  // splits inside "$type" (nested property in ReadObjectPropertyValue)
+		[DataRow(77)]  // splits inside "$text" (nested property in ReadObjectPropertyValue)
+		[DataRow(92)]  // splits inside "count" (nested property in ReadProperties)
+		public void ReadMultiSegment(int splitPoint)
+		{
+			string json = "{\"message\":\"hello\",\"format\":\"{Foo}\",\"properties\":{\"Foo\":{\"$type\":\"Symbol\",\"$text\":\"bar\"},\"count\":42}}";
+			byte[] bytes = Encoding.UTF8.GetBytes(json);
+
+			// Reference: deserialize from contiguous span (known good)
+			LogEvent expected = LogEvent.Read(bytes);
+
+			// Test: deserialize from multi-segment sequence (exercises TryReadNextPropertyName across segment boundaries)
+			ReadOnlySequence<byte> sequence = ReadOnlySequence.Create<byte>([bytes.AsMemory(0, splitPoint), bytes.AsMemory(splitPoint)]);
+			Utf8JsonReader reader = new(sequence);
+			reader.Read();
+			LogEvent actual = LogEvent.Read(ref reader);
+
+			Assert.AreEqual(Encoding.UTF8.GetString(expected.ToJsonBytes()), Encoding.UTF8.GetString(actual.ToJsonBytes()));
+		}
+
+		[TestMethod]
+		[DataRow(63)]  // splits \n escape in "Foo" value (between \ and n)
+		[DataRow(83)]  // splits first \" escape in "Bar" value (between \ and ")
+		[DataRow(87)]  // splits second \" escape in "Bar" value (between \ and ")
+		public void ReadMultiSegmentEscapedStrings(int splitPoint)
+		{
+			// String values with JSON escape sequences (\n, \") — exercises reader.GetString()! at line 352
+			string json = """{"message":"hello","format":"{Foo}","properties":{"Foo":"line1\nline2","Bar":"say \"hi\""}}""";
+			byte[] bytes = Encoding.UTF8.GetBytes(json);
+
+			LogEvent expected = LogEvent.Read(bytes);
+
+			ReadOnlySequence<byte> sequence = ReadOnlySequence.Create<byte>([bytes.AsMemory(0, splitPoint), bytes.AsMemory(splitPoint)]);
+			Utf8JsonReader reader = new(sequence);
+			reader.Read();
+			LogEvent actual = LogEvent.Read(ref reader);
+
+			Assert.AreEqual(Encoding.UTF8.GetString(expected.ToJsonBytes()), Encoding.UTF8.GetString(actual.ToJsonBytes()));
+		}
+
+		[TestMethod]
+		[DataRow(59)]  // splits inside "$type" property name
+		[DataRow(75)]  // splits inside "$text" property name
+		[DataRow(87)]  // splits \\ escape in $text value (between first \ and second \)
+		[DataRow(91)]  // splits second \\ escape in $text value
+		public void ReadMultiSegmentEscapedObjectValue(int splitPoint)
+		{
+			// $text with backslash escapes — exercises ReadObjectPropertyValue + TryReadNextPropertyName
+			string json = """{"message":"test","format":"{Msg}","properties":{"Msg":{"$type":"Source","$text":"path\\to\\file"}}}""";
+			byte[] bytes = Encoding.UTF8.GetBytes(json);
+
+			LogEvent expected = LogEvent.Read(bytes);
+
+			ReadOnlySequence<byte> sequence = ReadOnlySequence.Create<byte>([bytes.AsMemory(0, splitPoint), bytes.AsMemory(splitPoint)]);
+			Utf8JsonReader reader = new(sequence);
+			reader.Read();
+			LogEvent actual = LogEvent.Read(ref reader);
+
+			Assert.AreEqual(Encoding.UTF8.GetString(expected.ToJsonBytes()), Encoding.UTF8.GetString(actual.ToJsonBytes()));
+		}
+
+		[TestMethod]
+		[DataRow(52)]  // splits at decimal point in 3.14159 (between 3 and .)
+		[DataRow(54)]  // splits inside decimal digits (between 4 and 1)
+		public void ReadMultiSegmentNumberDouble(int splitPoint)
+		{
+			// Non-integer number exercises the TryGetDouble path (line 358) across segment boundaries
+			string json = """{"message":"test","format":"{X}","properties":{"X":3.14159}}""";
+			byte[] bytes = Encoding.UTF8.GetBytes(json);
+
+			LogEvent expected = LogEvent.Read(bytes);
+
+			ReadOnlySequence<byte> sequence = ReadOnlySequence.Create<byte>([bytes.AsMemory(0, splitPoint), bytes.AsMemory(splitPoint)]);
+			Utf8JsonReader reader = new(sequence);
+			reader.Read();
+			LogEvent actual = LogEvent.Read(ref reader);
+
+			Assert.AreEqual(Encoding.UTF8.GetString(expected.ToJsonBytes()), Encoding.UTF8.GetString(actual.ToJsonBytes()));
+		}
 	}
 }
diff --git a/Engine/Source/Programs/Shared/EpicGames.Core.Tests/Sha1HashTests.cs b/Engine/Source/Programs/Shared/EpicGames.Core.Tests/Sha1HashTests.cs
new file mode 100644
--- /dev/null
+++ b/Engine/Source/Programs/Shared/EpicGames.Core.Tests/Sha1HashTests.cs
@@ -0,0 +1,32 @@
+// Copyright Epic Games, Inc. All Rights Reserved.
+
+using System;
+using System.Buffers;
+using System.Text;
+using System.Text.Json;
+using Microsoft.VisualStudio.TestTools.UnitTesting;
+
+namespace EpicGames.Core.Tests
+{
+	[TestClass]
+	public class Sha1HashTests
+	{
+		static readonly Sha1Hash s_testHash = Sha1Hash.Compute(Encoding.ASCII.GetBytes("test data for sha1"));
+
+		[TestMethod]
+		[DataRow(5)]
+		[DataRow(22)]
+		[DataRow(30)]
+		public void JsonConverterMultiSegment(int splitPoint)
+		{
+			string json = $"\"{s_testHash}\"";
+			byte[] bytes = Encoding.UTF8.GetBytes(json);
+			ReadOnlySequence<byte> sequence = ReadOnlySequence.Create<byte>([bytes.AsMemory(0, splitPoint), bytes.AsMemory(splitPoint)]);
+
+			Utf8JsonReader reader = new(sequence);
+			reader.Read();
+			Sha1Hash result = JsonSerializer.Deserialize<Sha1Hash>(ref reader);
+			Assert.AreEqual(s_testHash, result);
+		}
+	}
+}
diff --git a/Engine/Source/Programs/Shared/EpicGames.Core/JsonExtensions.cs b/Engine/Source/Programs/Shared/EpicGames.Core/JsonExtensions.cs
--- a/Engine/Source/Programs/Shared/EpicGames.Core/JsonExtensions.cs
+++ b/Engine/Source/Programs/Shared/EpicGames.Core/JsonExtensions.cs
@@ -24,7 +24,7 @@ namespace EpicGames.Core
 		{
 			if (reader.Read() && reader.TokenType == JsonTokenType.PropertyName)
 			{
-				propertyName = reader.ValueSpan;
+				propertyName = reader.GetUtf8String();
 				return reader.Read();
 			}
 			else
diff --git a/Engine/Source/Programs/Shared/EpicGames.Core/LogEvent.cs b/Engine/Source/Programs/Shared/EpicGames.Core/LogEvent.cs
--- a/Engine/Source/Programs/Shared/EpicGames.Core/LogEvent.cs
+++ b/Engine/Source/Programs/Shared/EpicGames.Core/LogEvent.cs
@@ -361,7 +361,7 @@ namespace EpicGames.Core
 					}
 					else
 					{
-						return Encoding.UTF8.GetString(reader.ValueSpan);
+						return reader.GetString()!;
 					}
 				default:
 					throw new InvalidOperationException("Unhandled property type");
diff --git a/Engine/Source/Programs/Shared/EpicGames.Core/Sha1Hash.cs b/Engine/Source/Programs/Shared/EpicGames.Core/Sha1Hash.cs
--- a/Engine/Source/Programs/Shared/EpicGames.Core/Sha1Hash.cs
+++ b/Engine/Source/Programs/Shared/EpicGames.Core/Sha1Hash.cs
@@ -215,7 +215,7 @@ namespace EpicGames.Core
 	sealed class Sha1HashJsonConverter : JsonConverter<Sha1Hash>
 	{
 		/// <inheritdoc/>
-		public override Sha1Hash Read(ref Utf8JsonReader reader, Type typeToConvert, JsonSerializerOptions options) => Sha1Hash.Parse(reader.ValueSpan);
+		public override Sha1Hash Read(ref Utf8JsonReader reader, Type typeToConvert, JsonSerializerOptions options) => Sha1Hash.Parse(reader.GetUtf8String());
 
 		/// <inheritdoc/>
 		public override void Write(Utf8JsonWriter writer, Sha1Hash value, JsonSerializerOptions options) => writer.WriteStringValue(value.ToUtf8String().Span);
diff --git a/Engine/Source/Programs/Shared/EpicGames.Horde.Tests/Storage/RefNameTests.cs b/Engine/Source/Programs/Shared/EpicGames.Horde.Tests/Storage/RefNameTests.cs
new file mode 100644
--- /dev/null
+++ b/Engine/Source/Programs/Shared/EpicGames.Horde.Tests/Storage/RefNameTests.cs
@@ -0,0 +1,36 @@
+// Copyright Epic Games, Inc. All Rights Reserved.
+
+using System;
+using System.Buffers;
+using System.Text;
+using System.Text.Json;
+using EpicGames.Horde.Storage;
+using Microsoft.VisualStudio.TestTools.UnitTesting;
+
+namespace EpicGames.Horde.Tests.Storage
+{
+	[TestClass]
+	public class RefNameTests
+	{
+		[TestMethod]
+		[DataRow("my/nested/ref-name", "my/nested/ref-name", 10)]
+		[DataRow("my/nested/ref-name", "my/nested/ref-name", 15)]
+		[DataRow("simple", "simple", 5)]
+		[DataRow("foo-pcb/baz-main/10200300/editor/a1b2c3d4e5f6a7b8c9d0e1f2/0a1b2c3d4e5f6a7b8c9d0e1f2a3b4c5d_1",
+			"foo-pcb/baz-main/10200300/editor/a1b2c3d4e5f6a7b8c9d0e1f2/0a1b2c3d4e5f6a7b8c9d0e1f2a3b4c5d_1", 40)]
+		[DataRow("foo-pcb/baz-main/10200300/editor/a1b2c3d4e5f6a7b8c9d0e1f2/0a1b2c3d4e5f6a7b8c9d0e1f2a3b4c5d_1",
+			"foo-pcb/baz-main/10200300/editor/a1b2c3d4e5f6a7b8c9d0e1f2/0a1b2c3d4e5f6a7b8c9d0e1f2a3b4c5d_1", 80)]
+		public void JsonConverterMultiSegment(string jsonValue, string expectedStr, int splitPoint)
+		{
+			RefName expected = new(expectedStr);
+			string json = $"\"{jsonValue}\"";
+			byte[] bytes = Encoding.UTF8.GetBytes(json);
+			ReadOnlySequence<byte> sequence = EpicGames.Core.ReadOnlySequence.Create<byte>([bytes.AsMemory(0, splitPoint), bytes.AsMemory(splitPoint)]);
+
+			Utf8JsonReader reader = new(sequence);
+			reader.Read();
+			RefName result = JsonSerializer.Deserialize<RefName>(ref reader);
+			Assert.AreEqual(expected, result);
+		}
+	}
+}
diff --git a/Engine/Source/Programs/Shared/EpicGames.Horde/Storage/RefName.cs b/Engine/Source/Programs/Shared/EpicGames.Horde/Storage/RefName.cs
--- a/Engine/Source/Programs/Shared/EpicGames.Horde/Storage/RefName.cs
+++ b/Engine/Source/Programs/Shared/EpicGames.Horde/Storage/RefName.cs
@@ -168,7 +168,7 @@ namespace EpicGames.Horde.Storage
 	sealed class RefNameJsonConverter : JsonConverter<RefName>
 	{
 		/// <inheritdoc/>
-		public override RefName Read(ref Utf8JsonReader reader, Type typeToConvert, JsonSerializerOptions options) => new(new Utf8String(reader.ValueSpan.ToArray()));
+		public override RefName Read(ref Utf8JsonReader reader, Type typeToConvert, JsonSerializerOptions options) => new(new Utf8String(reader.GetUtf8String().ToArray()));
 
 		/// <inheritdoc/>
 		public override void Write(Utf8JsonWriter writer, RefName value, JsonSerializerOptions options) => writer.WriteStringValue(value.Text.Span);
diff --git a/Engine/Source/Programs/Shared/EpicGames.IoHash.Tests/IoHashTests.cs b/Engine/Source/Programs/Shared/EpicGames.IoHash.Tests/IoHashTests.cs
--- a/Engine/Source/Programs/Shared/EpicGames.IoHash.Tests/IoHashTests.cs
+++ b/Engine/Source/Programs/Shared/EpicGames.IoHash.Tests/IoHashTests.cs
@@ -4,6 +4,7 @@ using System;
 using System.Buffers;
 using System.IO;
 using System.Text;
+using System.Text.Json;
 using System.Threading.Tasks;
 using Microsoft.VisualStudio.TestTools.UnitTesting;
 
@@ -71,5 +72,21 @@ namespace EpicGames.IoHash.Tests
 			Assert.IsFalse(Core.IoHash.TryParse(String.Empty, out _));
 			Assert.IsFalse(Core.IoHash.TryParse(Encoding.ASCII.GetBytes(String.Empty), out _));
 		}
+
+		[TestMethod]
+		[DataRow(5)]
+		[DataRow(22)]
+		[DataRow(30)]
+		public void JsonConverterMultiSegment(int splitPoint)
+		{
+			string json = $"\"{s_testHash}\"";
+			byte[] bytes = Encoding.UTF8.GetBytes(json);
+			ReadOnlySequence<byte> sequence = Core.ReadOnlySequence.Create<byte>([bytes.AsMemory(0, splitPoint), bytes.AsMemory(splitPoint)]);
+
+			Utf8JsonReader reader = new(sequence);
+			reader.Read();
+			Core.IoHash result = JsonSerializer.Deserialize<Core.IoHash>(ref reader);
+			Assert.AreEqual(s_testHash, result);
+		}
 	}
 }
diff --git a/Engine/Source/Programs/Shared/EpicGames.IoHash/IoHash.cs b/Engine/Source/Programs/Shared/EpicGames.IoHash/IoHash.cs
--- a/Engine/Source/Programs/Shared/EpicGames.IoHash/IoHash.cs
+++ b/Engine/Source/Programs/Shared/EpicGames.IoHash/IoHash.cs
@@ -455,7 +455,7 @@ namespace EpicGames.Core
 	sealed class IoHashJsonConverter : JsonConverter<IoHash>
 	{
 		/// <inheritdoc/>
-		public override IoHash Read(ref Utf8JsonReader reader, Type typeToConvert, JsonSerializerOptions options) => IoHash.Parse(reader.ValueSpan);
+		public override IoHash Read(ref Utf8JsonReader reader, Type typeToConvert, JsonSerializerOptions options) => IoHash.Parse(reader.GetUtf8String());
 
 		/// <inheritdoc/>
 		public override void Write(Utf8JsonWriter writer, IoHash value, JsonSerializerOptions options) => writer.WriteStringValue(value.ToUtf8String().Span);
diff --git a/Engine/Source/Programs/Shared/EpicGames.Serialization.Tests/CbFieldTests.cs b/Engine/Source/Programs/Shared/EpicGames.Serialization.Tests/CbFieldTests.cs
new file mode 100644
--- /dev/null
+++ b/Engine/Source/Programs/Shared/EpicGames.Serialization.Tests/CbFieldTests.cs
@@ -0,0 +1,51 @@
+// Copyright Epic Games, Inc. All Rights Reserved.
+
+using System;
+using System.Buffers;
+using System.Text;
+using System.Text.Json;
+using EpicGames.Core;
+using Microsoft.VisualStudio.TestTools.UnitTesting;
+
+namespace EpicGames.Serialization.Tests
+{
+	[TestClass]
+	public class CbFieldTests
+	{
+		static readonly IoHash s_testHash = IoHash.Parse("7b0142eafb4d0bed8265978281a1398099b1d727");
+
+		[TestMethod]
+		[DataRow(5)]
+		[DataRow(22)]
+		[DataRow(30)]
+		public void CbBinaryAttachmentMultiSegment(int splitPoint)
+		{
+			CbBinaryAttachment expected = s_testHash;
+			string json = $"\"{s_testHash}\"";
+			byte[] bytes = Encoding.UTF8.GetBytes(json);
+			ReadOnlySequence<byte> sequence = ReadOnlySequence.Create<byte>([bytes.AsMemory(0, splitPoint), bytes.AsMemory(splitPoint)]);
+
+			Utf8JsonReader reader = new(sequence);
+			reader.Read();
+			CbBinaryAttachment result = JsonSerializer.Deserialize<CbBinaryAttachment>(ref reader);
+			Assert.AreEqual(expected, result);
+		}
+
+		[TestMethod]
+		[DataRow(5)]
+		[DataRow(22)]
+		[DataRow(30)]
+		public void CbObjectAttachmentMultiSegment(int splitPoint)
+		{
+			CbObjectAttachment expected = s_testHash;
+			string json = $"\"{s_testHash}\"";
+			byte[] bytes = Encoding.UTF8.GetBytes(json);
+			ReadOnlySequence<byte> sequence = ReadOnlySequence.Create<byte>([bytes.AsMemory(0, splitPoint), bytes.AsMemory(splitPoint)]);
+
+			Utf8JsonReader reader = new(sequence);
+			reader.Read();
+			CbObjectAttachment result = JsonSerializer.Deserialize<CbObjectAttachment>(ref reader);
+			Assert.AreEqual(expected, result);
+		}
+	}
+}
diff --git a/Engine/Source/Programs/Shared/EpicGames.Serialization/CbField.cs b/Engine/Source/Programs/Shared/EpicGames.Serialization/CbField.cs
--- a/Engine/Source/Programs/Shared/EpicGames.Serialization/CbField.cs
+++ b/Engine/Source/Programs/Shared/EpicGames.Serialization/CbField.cs
@@ -245,7 +245,7 @@ namespace EpicGames.Serialization
 	sealed class CbBinaryAttachmentJsonConverter : JsonConverter<CbBinaryAttachment>
 	{
 		/// <inheritdoc/>
-		public override CbBinaryAttachment Read(ref Utf8JsonReader reader, Type typeToConvert, JsonSerializerOptions options) => IoHash.Parse(reader.ValueSpan);
+		public override CbBinaryAttachment Read(ref Utf8JsonReader reader, Type typeToConvert, JsonSerializerOptions options) => IoHash.Parse(reader.GetUtf8String());
 
 		/// <inheritdoc/>
 		public override void Write(Utf8JsonWriter writer, CbBinaryAttachment value, JsonSerializerOptions options) => writer.WriteStringValue(value.Hash.ToUtf8String().Span);
@@ -319,7 +319,7 @@ namespace EpicGames.Serialization
 	sealed class CbObjectAttachmentJsonConverter : JsonConverter<CbObjectAttachment>
 	{
 		/// <inheritdoc/>
-		public override CbObjectAttachment Read(ref Utf8JsonReader reader, Type typeToConvert, JsonSerializerOptions options) => IoHash.Parse(reader.ValueSpan);
+		public override CbObjectAttachment Read(ref Utf8JsonReader reader, Type typeToConvert, JsonSerializerOptions options) => IoHash.Parse(reader.GetUtf8String());
 
 		/// <inheritdoc/>
 		public override void Write(Utf8JsonWriter writer, CbObjectAttachment value, JsonSerializerOptions options) => writer.WriteStringValue(value.Hash.ToUtf8String().Span);
