Shared: Support targeting both net8.0 and net10.0

#rnx
#Synced-CL 52229246

diff --git a/Engine/Source/Programs/Shared/EpicGames.Analytics.Generators/EpicGames.Analytics.Generators.csproj b/Engine/Source/Programs/Shared/EpicGames.Analytics.Generators/EpicGames.Analytics.Generators.csproj
--- a/Engine/Source/Programs/Shared/EpicGames.Analytics.Generators/EpicGames.Analytics.Generators.csproj
+++ b/Engine/Source/Programs/Shared/EpicGames.Analytics.Generators/EpicGames.Analytics.Generators.csproj
@@ -1,6 +1,6 @@
 ﻿<Project Sdk="Microsoft.NET.Sdk">
   <PropertyGroup>
-    <TargetFramework>netstandard2.0</TargetFramework>
+    <TargetFrameworks>netstandard2.0</TargetFrameworks>
     <LangVersion>latest</LangVersion>
     <IsRoslynComponent>true</IsRoslynComponent>
     <EnforceExtendedAnalyzerRules>true</EnforceExtendedAnalyzerRules>
diff --git a/Engine/Source/Programs/Shared/EpicGames.AspNet/EpicGames.AspNet.csproj b/Engine/Source/Programs/Shared/EpicGames.AspNet/EpicGames.AspNet.csproj
--- a/Engine/Source/Programs/Shared/EpicGames.AspNet/EpicGames.AspNet.csproj
+++ b/Engine/Source/Programs/Shared/EpicGames.AspNet/EpicGames.AspNet.csproj
@@ -17,8 +17,4 @@
     <ProjectReference Include="..\..\Shared\EpicGames.Core\EpicGames.Core.csproj" />
     <ProjectReference Include="..\EpicGames.Serialization\EpicGames.Serialization.csproj" />
   </ItemGroup>
-
-  <ItemGroup>
-    <PackageReference Include="System.Text.Json" Version="8.0.6" />
-  </ItemGroup>
 </Project>
diff --git a/Engine/Source/Programs/Shared/EpicGames.BuildGraph.Tests/EpicGames.BuildGraph.Tests.csproj b/Engine/Source/Programs/Shared/EpicGames.BuildGraph.Tests/EpicGames.BuildGraph.Tests.csproj
--- a/Engine/Source/Programs/Shared/EpicGames.BuildGraph.Tests/EpicGames.BuildGraph.Tests.csproj
+++ b/Engine/Source/Programs/Shared/EpicGames.BuildGraph.Tests/EpicGames.BuildGraph.Tests.csproj
@@ -1,8 +1,4 @@
 ﻿<Project Sdk="Microsoft.NET.Sdk">
-  <ItemGroup>
-    <PackageReference Include="System.Text.Json" Version="8.0.6" />
-  </ItemGroup>
-
   <ItemGroup>
     <ProjectReference Include="..\EpicGames.BuildGraph\EpicGames.BuildGraph.csproj" />
   </ItemGroup>
diff --git a/Engine/Source/Programs/Shared/EpicGames.Core/EpicGames.Core.csproj b/Engine/Source/Programs/Shared/EpicGames.Core/EpicGames.Core.csproj
--- a/Engine/Source/Programs/Shared/EpicGames.Core/EpicGames.Core.csproj
+++ b/Engine/Source/Programs/Shared/EpicGames.Core/EpicGames.Core.csproj
@@ -10,6 +10,5 @@
     <PackageReference Include="Microsoft.Extensions.Logging.Console" Version="8.0.1" />
     <PackageReference Include="Microsoft.Extensions.ObjectPool" Version="8.0.20" />
     <PackageReference Include="OpenTracing" Version="0.12.1" />
-    <PackageReference Include="System.Text.Json" Version="8.0.6" />
   </ItemGroup>
 </Project>
diff --git a/Engine/Source/Programs/Shared/EpicGames.Horde.Tests/EpicGames.Horde.Tests.csproj b/Engine/Source/Programs/Shared/EpicGames.Horde.Tests/EpicGames.Horde.Tests.csproj
--- a/Engine/Source/Programs/Shared/EpicGames.Horde.Tests/EpicGames.Horde.Tests.csproj
+++ b/Engine/Source/Programs/Shared/EpicGames.Horde.Tests/EpicGames.Horde.Tests.csproj
@@ -1,5 +1,5 @@
 <Project Sdk="Microsoft.NET.Sdk">
-  <ItemGroup>
+  <ItemGroup Condition="'$(TargetFramework)' == 'net8.0'">
     <PackageReference Include="System.Linq.Async" Version="6.0.3" />
   </ItemGroup>
 
diff --git a/Engine/Source/Programs/Shared/EpicGames.Horde/Compute/Transports/TcpSslTransport.cs b/Engine/Source/Programs/Shared/EpicGames.Horde/Compute/Transports/TcpSslTransport.cs
--- a/Engine/Source/Programs/Shared/EpicGames.Horde/Compute/Transports/TcpSslTransport.cs
+++ b/Engine/Source/Programs/Shared/EpicGames.Horde/Compute/Transports/TcpSslTransport.cs
@@ -28,7 +28,12 @@ public sealed class TcpSslTransport : ComputeTransport
 	/// <param name="socket">Socket to communicate over</param>
 	/// <param name="certData">Certificate used for auth on both server and client</param>
 	/// <param name="isServer">Whether socket is acting as a server or client</param>
-	public TcpSslTransport(Socket socket, byte[] certData, bool isServer) : this(socket, new X509Certificate2(certData), isServer)
+	public TcpSslTransport(Socket socket, byte[] certData, bool isServer)
+#if NET9_0_OR_GREATER
+		: this(socket, X509CertificateLoader.LoadCertificate(certData), isServer)
+#else
+		: this(socket, new X509Certificate2(certData), isServer)
+#endif
 	{
 	}
 
diff --git a/Engine/Source/Programs/Shared/EpicGames.Horde/EpicGames.Horde.csproj b/Engine/Source/Programs/Shared/EpicGames.Horde/EpicGames.Horde.csproj
--- a/Engine/Source/Programs/Shared/EpicGames.Horde/EpicGames.Horde.csproj
+++ b/Engine/Source/Programs/Shared/EpicGames.Horde/EpicGames.Horde.csproj
@@ -21,10 +21,13 @@
       <PrivateAssets>all</PrivateAssets>
       <IncludeAssets>runtime; build; native; contentfiles; analyzers; buildtransitive</IncludeAssets>
     </PackageReference>
-    <PackageReference Include="System.Linq.Async" Version="6.0.3" />
     <PackageReference Include="ZstdSharp.Port" Version="0.8.6" />
   </ItemGroup>
 
+  <ItemGroup Condition="'$(TargetFramework)' == 'net8.0'">
+    <PackageReference Include="System.Linq.Async" Version="6.0.3" />
+  </ItemGroup>
+
   <ItemGroup>
     <ProjectReference Include="..\EpicGames.Analytics.Telemetry\EpicGames.Analytics.Telemetry.csproj" />
     <ProjectReference Include="..\EpicGames.Core\EpicGames.Core.csproj" />
diff --git a/Engine/Source/Programs/Shared/EpicGames.IoHash/EpicGames.IoHash.csproj b/Engine/Source/Programs/Shared/EpicGames.IoHash/EpicGames.IoHash.csproj
--- a/Engine/Source/Programs/Shared/EpicGames.IoHash/EpicGames.IoHash.csproj
+++ b/Engine/Source/Programs/Shared/EpicGames.IoHash/EpicGames.IoHash.csproj
@@ -6,7 +6,6 @@
   <ItemGroup>
     <!-- Blake3 2.0.0 has a massive performance regression on Windows x64 - takes 2.25x as long for large (1MB) inputs -->
     <PackageReference Include="Blake3" Version="1.1.0" />
-    <PackageReference Include="System.Memory" Version="4.6.3" />
   </ItemGroup>
 
   <ItemGroup>
diff --git a/Engine/Source/Programs/Shared/EpicGames.MsBuild/CsProjBuilder.cs b/Engine/Source/Programs/Shared/EpicGames.MsBuild/CsProjBuilder.cs
--- a/Engine/Source/Programs/Shared/EpicGames.MsBuild/CsProjBuilder.cs
+++ b/Engine/Source/Programs/Shared/EpicGames.MsBuild/CsProjBuilder.cs
@@ -128,6 +128,37 @@ namespace EpicGames.MsBuild
 					try
 					{
 						project = new Project(projectPath, globalProperties, toolsVersion: null, projectCollection: projectCollection);
+
+						string? targetFrameworks = project.GetProperty("TargetFrameworks")?.EvaluatedValue;
+						if (!String.IsNullOrEmpty(targetFrameworks))
+						{
+							Dictionary<string, string> innerProps = project.GlobalProperties.ToDictionary(kvp => kvp.Key, kvp => kvp.Value);
+#if NET10_0_OR_GREATER
+							if (!innerProps.ContainsKey("TargetFramework") && targetFrameworks.Contains("net10.0", StringComparison.OrdinalIgnoreCase))
+							{
+								innerProps["TargetFramework"] = "net10.0";
+							}
+#endif
+#if NET9_0_OR_GREATER
+							if (!innerProps.ContainsKey("TargetFramework") && targetFrameworks.Contains("net9.0", StringComparison.OrdinalIgnoreCase))
+							{
+								innerProps["TargetFramework"] = "net9.0";
+							}
+#endif
+#if NET8_0_OR_GREATER
+							if (!innerProps.ContainsKey("TargetFramework") && targetFrameworks.Contains("net8.0", StringComparison.OrdinalIgnoreCase))
+							{
+								innerProps["TargetFramework"] = "net8.0";
+							}
+#endif
+							if (!innerProps.ContainsKey("TargetFramework"))
+							{
+								logger.LogInformation("Skipping {Framework} project {ProjectPath}", targetFrameworks, projectPath);
+								return;
+							}
+
+							project = new Project(project.FullPath, innerProps, toolsVersion: null, project.ProjectCollection);
+						}
 					}
 					catch (Microsoft.Build.Exceptions.InvalidProjectFileException iPFEx)
 					{
diff --git a/Engine/Source/Programs/Shared/EpicGames.OIDC/EpicGames.OIDC.csproj b/Engine/Source/Programs/Shared/EpicGames.OIDC/EpicGames.OIDC.csproj
--- a/Engine/Source/Programs/Shared/EpicGames.OIDC/EpicGames.OIDC.csproj
+++ b/Engine/Source/Programs/Shared/EpicGames.OIDC/EpicGames.OIDC.csproj
@@ -1,6 +1,7 @@
 ﻿<Project Sdk="Microsoft.NET.Sdk">
   <PropertyGroup>
-    <EnableConfigurationBindingGenerator>true</EnableConfigurationBindingGenerator>
+    <EnableConfigurationBindingGenerator Condition="'$(TargetFramework)' == 'net8.0'">true</EnableConfigurationBindingGenerator>
+    <EnableConfigurationBindingGenerator Condition="'$(TargetFramework)' != 'net8.0'">false</EnableConfigurationBindingGenerator>
   </PropertyGroup>
 
   <ItemGroup>
@@ -8,7 +9,6 @@
     <PackageReference Include="Microsoft.Extensions.Configuration.Binder" Version="8.0.2" />
     <PackageReference Include="Microsoft.Extensions.Configuration.Json" Version="8.0.1" />
     <PackageReference Include="Microsoft.Extensions.Options" Version="8.0.2" />
-    <PackageReference Include="System.Text.Json" Version="8.0.6" />
   </ItemGroup>
 
   <ItemGroup>
diff --git a/Engine/Source/Programs/Shared/EpicGames.Perforce.Fixture/EpicGames.Perforce.Fixture.csproj b/Engine/Source/Programs/Shared/EpicGames.Perforce.Fixture/EpicGames.Perforce.Fixture.csproj
--- a/Engine/Source/Programs/Shared/EpicGames.Perforce.Fixture/EpicGames.Perforce.Fixture.csproj
+++ b/Engine/Source/Programs/Shared/EpicGames.Perforce.Fixture/EpicGames.Perforce.Fixture.csproj
@@ -20,7 +20,6 @@
       <PrivateAssets>all</PrivateAssets>
       <IncludeAssets>runtime; build; native; contentfiles; analyzers; buildtransitive</IncludeAssets>
     </PackageReference>
-    <PackageReference Include="System.Text.Json" Version="8.0.6" />
   </ItemGroup>
 
   <ItemGroup>
diff --git a/Engine/Source/Programs/Shared/EpicGames.Perforce/EpicGames.Perforce.csproj b/Engine/Source/Programs/Shared/EpicGames.Perforce/EpicGames.Perforce.csproj
--- a/Engine/Source/Programs/Shared/EpicGames.Perforce/EpicGames.Perforce.csproj
+++ b/Engine/Source/Programs/Shared/EpicGames.Perforce/EpicGames.Perforce.csproj
@@ -35,10 +35,13 @@
   </ItemGroup>
 
   <ItemGroup>
-    <PackageReference Include="System.Linq.Async" Version="6.0.3" />
     <PackageReference Include="OpenTelemetry.Api" Version="1.8.1" />
   </ItemGroup>
 
+  <ItemGroup Condition="'$(TargetFramework)' == 'net8.0'">
+    <PackageReference Include="System.Linq.Async" Version="6.0.3" />
+  </ItemGroup>
+
   <ItemGroup>
     <ProjectReference Include="..\EpicGames.Core\EpicGames.Core.csproj" />
   </ItemGroup>
