BuildGraph: Add OptionalRequires attribute to BuildGraph Node

Allow BuildGraph nodes to declare dependencies using the new OptionalRequires attribute. The optional dependencies are not required for a node's execution and if an optional dependency fails the node will still run; the node's execution will decide how missing inputs are handled instead of the graph.

#Synced-CL 0

diff --git a/Engine/Source/Programs/AutomationTool/BuildGraph/BgNodeExecutor.cs b/Engine/Source/Programs/AutomationTool/BuildGraph/BgNodeExecutor.cs
--- a/Engine/Source/Programs/AutomationTool/BuildGraph/BgNodeExecutor.cs
+++ b/Engine/Source/Programs/AutomationTool/BuildGraph/BgNodeExecutor.cs
@@ -280,7 +280,7 @@ namespace AutomationTool
 				BgNodeOutput? output;
 				if (tagNameToNodeOutput.TryGetValue(readTagName, out output))
 				{
-					if (output != null && output.ProducingNode != Node && !Node.Inputs.Contains(output))
+					if (output != null && output.ProducingNode != Node && !(Node.Inputs.Contains(output) || Node.OptionalInputs.Contains(output)))
 					{
 						logger.LogScriptError(taskInfo.Location, "The tag '{TagName}' is not a dependency of node '{Node}'", readTagName, Node.Name);
 					}
diff --git a/Engine/Source/Programs/AutomationTool/BuildGraph/BgSchema.cs b/Engine/Source/Programs/AutomationTool/BuildGraph/BgSchema.cs
--- a/Engine/Source/Programs/AutomationTool/BuildGraph/BgSchema.cs
+++ b/Engine/Source/Programs/AutomationTool/BuildGraph/BgSchema.cs
@@ -563,6 +563,7 @@ namespace AutomationTool
 			extension.BaseTypeName = GetQualifiedTypeName(ScriptSchemaStandardType.NodeBody);
 			extension.Attributes.Add(CreateSchemaAttribute("Name", ScriptSchemaStandardType.Name, XmlSchemaUse.Required));
 			extension.Attributes.Add(CreateSchemaAttribute("Requires", ScriptSchemaStandardType.NameOrTagList, XmlSchemaUse.Optional));
+			extension.Attributes.Add(CreateSchemaAttribute("OptionalRequires", ScriptSchemaStandardType.NameOrTagList, XmlSchemaUse.Optional));
 			extension.Attributes.Add(CreateSchemaAttribute("Produces", ScriptSchemaStandardType.TagList, XmlSchemaUse.Optional));
 			extension.Attributes.Add(CreateSchemaAttribute("After", ScriptSchemaStandardType.NameOrTagList, XmlSchemaUse.Optional));
 			extension.Attributes.Add(CreateSchemaAttribute("Token", ScriptSchemaStandardType.BalancedString, XmlSchemaUse.Optional));
diff --git a/Engine/Source/Programs/AutomationTool/BuildGraph/BgScriptReader.cs b/Engine/Source/Programs/AutomationTool/BuildGraph/BgScriptReader.cs
--- a/Engine/Source/Programs/AutomationTool/BuildGraph/BgScriptReader.cs
+++ b/Engine/Source/Programs/AutomationTool/BuildGraph/BgScriptReader.cs
@@ -293,8 +293,8 @@ namespace AutomationTool
 		/// <summary>
 		/// Constructor
 		/// </summary>
-		public BgScriptNode(string name, IReadOnlyList<BgNodeOutput> inputs, IReadOnlyList<string> outputNames, IReadOnlyList<BgNodeDef> inputDependencies, IReadOnlyList<BgNodeDef> orderDependencies, IReadOnlyList<FileReference> requiredTokens, IReadOnlyList<string> ignoreModified)
-			: base(name, inputs, outputNames, inputDependencies, orderDependencies, requiredTokens, ignoreModified)
+		public BgScriptNode(string name, IReadOnlyList<BgNodeOutput> inputs, IReadOnlyList<BgNodeOutput> optionalInputs, IReadOnlyList<string> outputNames, IReadOnlyList<BgNodeDef> inputDependencies, IReadOnlyList<BgNodeDef> optionalInputDependencies, IReadOnlyList<BgNodeDef> orderDependencies, IReadOnlyList<FileReference> requiredTokens, IReadOnlyList<string> ignoreModified)
+			: base(name, inputs, optionalInputs, outputNames, inputDependencies, optionalInputDependencies, orderDependencies, requiredTokens, ignoreModified)
 		{
 		}
 	}
@@ -1457,6 +1457,7 @@ namespace AutomationTool
 			if (await EvaluateConditionAsync(element) && TryReadObjectName(element, out name))
 			{
 				string[] requiresNames = ReadListAttribute(element, "Requires");
+				string[] optionalRequiresNames = ReadListAttribute(element, "OptionalRequires");
 				string[] producesNames = ReadListAttribute(element, "Produces");
 				string[] afterNames = ReadListAttribute(element, "After");
 				string[] tokenFileNames = ReadListAttribute(element, "Token");
@@ -1470,22 +1471,54 @@ namespace AutomationTool
 				HashSet<BgNodeOutput> inputs = ResolveInputReferences(element, requiresNames);
 
 				// Gather up all the input dependencies, and check they're all upstream of the current node
-				HashSet<BgNodeDef> inputDependencies = new HashSet<BgNodeDef>();
+				HashSet<BgNodeDef> inputDependencies = [];
 				foreach (BgNodeDef inputDependency in inputs.Select(x => x.ProducingNode).Distinct())
 				{
 					inputDependencies.Add(inputDependency);
 				}
 
+				// Resolve all the optional inputs we depend on
+				HashSet<BgNodeOutput> optionalInputs = ResolveInputReferences(element, optionalRequiresNames);
+
+				// If an input is both required and optional the required taking precedence
+				HashSet<BgNodeOutput> duplicateInputs = new(inputs);
+				duplicateInputs.IntersectWith(optionalInputs);
+				foreach (BgNodeOutput duplicateInput in duplicateInputs)
+				{
+					LogWarning(element, "Optional input '{Tag}' is also a required input. The required input takes precedence", duplicateInput.TagName);
+				}
+				optionalInputs.ExceptWith(duplicateInputs);
+
+				// Gather all the optional input dependencies, and check they're all upstream of the current node
+				HashSet<BgNodeDef> optionalInputDependencies = [];
+				foreach (BgNodeDef optionalInputDependency in optionalInputs.Select(x => x.ProducingNode).Distinct())
+				{
+					optionalInputDependencies.Add(optionalInputDependency);
+				}
+
 				// Remove all the lock names from the list of required names
 				HashSet<FileReference> requiredTokens = new HashSet<FileReference>(tokenFileNames.Select(x => new FileReference(x)));
 
 				// Recursively include all their dependencies too
+				// Intentionally only traverses InputDependencies (required), not OptionalInputDependencies.
+				// A required dependency's own optional deps don't need to succeed for it to succeed, so they
+				// don't belong in the required chain. Execution ordering for optional dependents is handled
+				// separately via OrderDependencies below.
 				foreach (BgNodeDef inputDependency in inputDependencies.ToArray())
 				{
 					requiredTokens.UnionWith(inputDependency.RequiredTokens);
 					inputDependencies.UnionWith(inputDependency.InputDependencies);
 				}
 
+				// Recursively include all their optional dependencies too
+				foreach (BgNodeDef optionalInputDependency in optionalInputDependencies.ToArray())
+				{
+					// A node that optionally depends on another must still respect its resource locks
+					requiredTokens.UnionWith(optionalInputDependency.RequiredTokens);
+					optionalInputDependencies.UnionWith(optionalInputDependency.InputDependencies);
+					optionalInputDependencies.UnionWith(optionalInputDependency.OptionalInputDependencies);
+				}
+
 				// Validate all the outputs
 				List<string> validOutputNames = new List<string>();
 				foreach (string producesName in producesNames)
@@ -1507,6 +1540,7 @@ namespace AutomationTool
 
 				// Gather up all the order dependencies
 				HashSet<BgNodeDef> orderDependencies = new HashSet<BgNodeDef>(inputDependencies);
+				orderDependencies.UnionWith(optionalInputDependencies);
 				orderDependencies.UnionWith(ResolveReferences(element, afterNames));
 
 				// Recursively include all their order dependencies too
@@ -1529,10 +1563,12 @@ namespace AutomationTool
 				if (CheckNameIsUnique(element, name))
 				{
 					// Add it to the node lookup
-					BgScriptNode newNode = new BgScriptNode(name, inputs.ToArray(), validOutputNames.ToArray(), inputDependencies.ToArray(), orderDependencies.ToArray(), requiredTokens.ToArray(), ignoreModified.ToArray());
-					newNode.RunEarly = runEarly;
-					newNode.AllowRetry = allowRetry;
-					newNode.NotifyOnWarnings = notifyOnWarnings;
+					BgScriptNode newNode = new(name, inputs.ToArray(), optionalInputs.ToArray(), validOutputNames.ToArray(), inputDependencies.ToArray(), optionalInputDependencies.ToArray(), orderDependencies.ToArray(), requiredTokens.ToArray(), ignoreModified.ToArray())
+					{
+						RunEarly = runEarly,
+						AllowRetry = allowRetry,
+						NotifyOnWarnings = notifyOnWarnings
+					};
 					foreach ((string key, string value) in annotations)
 					{
 						newNode.Annotations[key] = value;
@@ -2401,7 +2437,7 @@ namespace AutomationTool
 		}
 
 		/// <summary>
-		/// Outputs a warning message to the log and increments the number of errors, referencing the file and line number of the element that caused it.
+		/// Outputs a warning message to the log, referencing the file and line number of the element that caused it.
 		/// </summary>
 		/// <param name="element">The script element causing the error</param>
 		/// <param name="format">Standard String.Format()-style format string</param>
diff --git a/Engine/Source/Programs/AutomationTool/BuildGraph/BgScriptWriter.cs b/Engine/Source/Programs/AutomationTool/BuildGraph/BgScriptWriter.cs
--- a/Engine/Source/Programs/AutomationTool/BuildGraph/BgScriptWriter.cs
+++ b/Engine/Source/Programs/AutomationTool/BuildGraph/BgScriptWriter.cs
@@ -68,6 +68,7 @@ namespace AutomationTool
 
 					HashSet<BgNodeDef> includedNodes = new HashSet<BgNodeDef>(label.IncludedNodes);
 					includedNodes.ExceptWith(label.IncludedNodes.SelectMany(x => x.InputDependencies));
+					includedNodes.ExceptWith(label.IncludedNodes.SelectMany(x => x.OptionalInputDependencies));
 					includedNodes.ExceptWith(label.RequiredNodes);
 					if (includedNodes.Count > 0)
 					{
@@ -76,8 +77,10 @@ namespace AutomationTool
 
 					HashSet<BgNodeDef> excludedNodes = new HashSet<BgNodeDef>(label.IncludedNodes);
 					excludedNodes.UnionWith(label.IncludedNodes.SelectMany(x => x.InputDependencies));
+					excludedNodes.UnionWith(label.IncludedNodes.SelectMany(x => x.OptionalInputDependencies));
 					excludedNodes.ExceptWith(label.IncludedNodes);
 					excludedNodes.ExceptWith(excludedNodes.ToArray().SelectMany(x => x.InputDependencies));
+					excludedNodes.ExceptWith(excludedNodes.ToArray().SelectMany(x => x.OptionalInputDependencies));
 					if (excludedNodes.Count > 0)
 					{
 						writer.WriteAttributeString("Exclude", String.Join(";", excludedNodes.Select(x => x.Name)));
@@ -146,13 +149,19 @@ namespace AutomationTool
 				writer.WriteAttributeString("Requires", String.Join(";", requireNames));
 			}
 
+			string[] optionalRequireNames = node.OptionalInputs.Select(x => x.TagName).ToArray();
+			if (optionalRequireNames.Length > 0)
+			{
+				writer.WriteAttributeString("OptionalRequires", String.Join(";", optionalRequireNames));
+			}
+
 			string[] producesNames = node.Outputs.Where(x => x != node.DefaultOutput).Select(x => x.TagName).ToArray();
 			if (producesNames.Length > 0)
 			{
 				writer.WriteAttributeString("Produces", String.Join(";", producesNames));
 			}
 
-			string[] afterNames = node.GetDirectOrderDependencies().Except(node.InputDependencies).Select(x => x.Name).ToArray();
+			string[] afterNames = node.GetDirectOrderDependencies().Except(node.InputDependencies).Except(node.OptionalInputDependencies).Select(x => x.Name).ToArray();
 			if (afterNames.Length > 0)
 			{
 				writer.WriteAttributeString("After", String.Join(";", afterNames));
diff --git a/Engine/Source/Programs/AutomationTool/BuildGraph/BuildGraph.cs b/Engine/Source/Programs/AutomationTool/BuildGraph/BuildGraph.cs
--- a/Engine/Source/Programs/AutomationTool/BuildGraph/BuildGraph.cs
+++ b/Engine/Source/Programs/AutomationTool/BuildGraph/BuildGraph.cs
@@ -212,6 +212,7 @@ namespace AutomationTool
 	[ParamHelp("SharedStorageDir=<DirName>", "Sets the directory to use to transfer build products between agents in a build farm", ParamType = typeof(Directory))]
 	[ParamHelp("SingleNode=<Name>", "Run only the given node. Intended for use on a build system after running with -Export.")]
 	[ParamHelp("WriteToSharedStorage", "Allow writing to shared storage. If not set, but -SharedStorageDir is specified, build products will read but not written", ParamType = typeof(bool))]
+	[ParamHelp("IgnoreOptionalNodeFailures", "Continue execution when nodes that are only optional dependencies fail", ParamType = typeof(bool))]
 	public class BuildGraph : BuildCommand
 	{
 		/// <summary>
@@ -240,6 +241,7 @@ namespace AutomationTool
 			bool writeToSharedStorage = ParseParam("WriteToSharedStorage") || CommandUtils.IsBuildMachine;
 			bool publicTasksOnly = ParseParam("PublicTasksOnly");
 			bool skipValidation = ParseParam("SkipValidation");
+			bool ignoreOptionalNodeFailures = ParseParam("IgnoreOptionalNodeFailures");
 			string? reportName = ParseParamValue("ReportName", null);
 			string? branchOverride = ParseParamValue("Branch", null);
 
@@ -751,7 +753,7 @@ namespace AutomationTool
 				}
 				else
 				{
-					if (!await BuildAllNodesAsync(graph, nodeToExecutor, storage))
+					if (!await BuildAllNodesAsync(graph, nodeToExecutor, storage, ignoreOptionalNodeFailures))
 					{
 						return ExitCode.Error_Unknown;
 					}
@@ -978,8 +980,9 @@ namespace AutomationTool
 		/// <param name="graph">The graph instance</param>
 		/// <param name="nodeToExecutor">Map from node to executor</param>
 		/// <param name="storage">The temp storage backend which stores the shared state</param>
+		/// <param name="ignoreOptionalNodeFailures">Continue execution when nodes that are only optional dependencies fail</param>
 		/// <returns>True if everything built successfully</returns>
-		async Task<bool> BuildAllNodesAsync(BgGraphDef graph, Dictionary<BgNodeDef, BgNodeExecutor> nodeToExecutor, TempStorage storage)
+		async Task<bool> BuildAllNodesAsync(BgGraphDef graph, Dictionary<BgNodeDef, BgNodeExecutor> nodeToExecutor, TempStorage storage, bool ignoreOptionalNodeFailures)
 		{
 			// Build a flat list of nodes to execute, in order
 			BgNodeDef[] nodesToExecute = graph.Agents.SelectMany(x => x.Nodes).ToArray();
@@ -990,13 +993,33 @@ namespace AutomationTool
 			foreach (BgNodeDef nodeToExecute in nodesToExecute)
 			{
 				FileFilter ignoreModifiedFilter = new FileFilter(nodeToExecute.IgnoreModified);
-				if (nodeToExecute.InputDependencies.Any(x => cleanedNodes.Contains(x)) || !storage.CheckLocalIntegrity(nodeToExecute.Name, nodeToExecute.Outputs.Select(x => x.TagName), ignoreModifiedFilter))
+				if (nodeToExecute.InputDependencies.Any(x => cleanedNodes.Contains(x)) ||
+				    nodeToExecute.OptionalInputDependencies.Any(x => cleanedNodes.Contains(x)) ||
+				    !storage.CheckLocalIntegrity(nodeToExecute.Name, nodeToExecute.Outputs.Select(x => x.TagName), ignoreModifiedFilter))
 				{
 					storage.CleanLocalNode(nodeToExecute.Name);
 					cleanedNodes.Add(nodeToExecute);
 				}
 			}
 
+			// Nodes that are only optional dependencies, which can fail without stopping the build
+			HashSet<BgNodeDef> nodesThatCanFail = [];
+			foreach (BgNodeDef node in nodesToExecute)
+			{
+				nodesThatCanFail.UnionWith(node.OptionalInputDependencies);
+			}
+
+			// Remove nodes that are also required or order-only dependencies; these must complete successfully
+			foreach (BgNodeDef node in nodesToExecute)
+			{
+				nodesThatCanFail.ExceptWith(node.InputDependencies);
+
+				// Exclude After dependencies
+				nodesThatCanFail.ExceptWith(node.GetDirectOrderDependencies().Except(node.InputDependencies).Except(node.OptionalInputDependencies));
+			}
+
+			bool hasErrors = false;
+
 			// ExecuteAsync them in order
 			int nodeIdx = 0;
 			foreach (BgNodeDef nodeToExecute in nodesToExecute)
@@ -1005,14 +1028,27 @@ namespace AutomationTool
 				if (!storage.IsComplete(nodeToExecute.Name))
 				{
 					Logger.LogInformation("");
-					if (!await BuildNodeAsync(graph, nodeToExecute, nodeToExecutor, storage, withBanner: false))
+
+					try
+					{
+						// Node executors signal failure via AutomationException. Other exception types 
+						// or returning false indicate unexpected errors and should always fail the build
+						if (!await BuildNodeAsync(graph, nodeToExecute, nodeToExecutor, storage, withBanner: false))
+						{
+							return false;
+						}
+					}
+					catch (AutomationException ex) when (ignoreOptionalNodeFailures && nodesThatCanFail.Contains(nodeToExecute))
 					{
-						return false;
+						Logger.LogError(ex, "Node '{Node}' failed {Message}. Continuing because node is only an optional dependency", nodeToExecute, ex.Message);
+						hasErrors = true;
 					}
+
 					Logger.LogInformation("");
 				}
 			}
-			return true;
+
+			return !hasErrors;
 		}
 
 		/// <summary>
@@ -1089,6 +1125,24 @@ namespace AutomationTool
 				inputStorageBlocks.UnionWith(fileList.Blocks);
 			}
 
+			// Read all the optional input tags for this node, and build a list of referenced input storage blocks
+			foreach (BgNodeOutput optionalInput in node.OptionalInputs)
+			{
+				try
+				{
+					TempStorageTagManifest? fileList = storage.ReadTagFileList(optionalInput.ProducingNode.Name, optionalInput.TagName);
+					if (fileList != null)
+					{
+						tagNameToFileSet[optionalInput.TagName] = fileList.ToFileSet(rootDir);
+						inputStorageBlocks.UnionWith(fileList.Blocks);
+					}
+				}
+				catch (AutomationException ex)
+				{
+					Logger.LogWarning(ex, "Failed to read optional tagged input {Input} from temp storage: {Message}. The producing step may have failed. Continuing without it", optionalInput, ex.Message);
+				}
+			}
+
 			// Read the manifests for all the input storage blocks
 			Dictionary<TempStorageBlockRef, TempStorageBlockManifest> inputManifests = new Dictionary<TempStorageBlockRef, TempStorageBlockManifest>();
 			using (IScope scope = GlobalTracer.Instance.BuildSpan("TempStorage").WithTag("resource", "read").StartActive())
@@ -1191,6 +1245,14 @@ namespace AutomationTool
 						{
 							referencedOutputFiles.UnionWith(tagNameToFileSet[input.TagName]);
 						}
+
+						foreach (BgNodeOutput optionalInput in otherNode.OptionalInputs.Where(x => x.ProducingNode == node))
+						{
+							if (tagNameToFileSet.TryGetValue(optionalInput.TagName, out HashSet<FileReference>? fileSet))
+							{
+								referencedOutputFiles.UnionWith(fileSet);
+							}
+						}
 					}
 				}
 			}
diff --git a/Engine/Source/Programs/Shared/EpicGames.BuildGraph.Tests/GraphTests.cs b/Engine/Source/Programs/Shared/EpicGames.BuildGraph.Tests/GraphTests.cs
--- a/Engine/Source/Programs/Shared/EpicGames.BuildGraph.Tests/GraphTests.cs
+++ b/Engine/Source/Programs/Shared/EpicGames.BuildGraph.Tests/GraphTests.cs
@@ -14,6 +14,7 @@ namespace EpicGames.BuildGraph.Tests
 		static Task UpdateVersionFiles() => Task.CompletedTask;
 		static Task CompileShooterGameWin64() => Task.CompletedTask;
 		static Task CookShooterGameWin64() => Task.CompletedTask;
+		static Task PackageShooterGameWin64() => Task.CompletedTask;
 
 		static object Evaluate(BgExpr expr)
 		{
@@ -36,16 +37,183 @@ namespace EpicGames.BuildGraph.Tests
 			Assert.AreEqual(node3.Name, "Cook Shooter Game Win64");
 			Assert.AreEqual(2, node3.InputDependencies.Count);
 			Assert.AreEqual(1, node3.Inputs.Count);
+			Assert.AreEqual(0, node3.OptionalInputs.Count);
 
 			BgNodeDef node2 = node3.Inputs[0].ProducingNode;
 			Assert.AreEqual(node2.Name, "Compile Shooter Game Win64");
 			Assert.AreEqual(1, node2.InputDependencies.Count);
 			Assert.AreEqual(1, node2.Inputs.Count);
+			Assert.AreEqual(0, node2.OptionalInputs.Count);
 
 			BgNodeDef node1 = node2.Inputs[0].ProducingNode;
 			Assert.AreEqual(node1.Name, "Update Version Files");
 			Assert.AreEqual(0, node1.InputDependencies.Count);
 			Assert.AreEqual(0, node1.Inputs.Count);
+			Assert.AreEqual(0, node1.OptionalInputs.Count);
+		}
+
+		[TestMethod]
+		public void NodeWithOptionalRequiresTest()
+		{
+			BgAgent agent = new("name", "type");
+			BgNode nodeSpec1 = agent.AddNode(x => UpdateVersionFiles());
+
+			BgNode nodeSpec2 = agent.AddNode(x => CompileShooterGameWin64()).OptionalRequires(nodeSpec1);
+
+			BgNode nodeSpec3 = agent.AddNode(x => CookShooterGameWin64()).OptionalRequires(BgList<BgFileSet>.Create(nodeSpec2.DefaultOutput));
+
+			BgNodeDef node3 = ((BgObjectDef)Evaluate(nodeSpec3)).Deserialize<BgNodeExpressionDef>();
+			Assert.AreEqual(node3.Name, "Cook Shooter Game Win64");
+			Assert.AreEqual(0, node3.InputDependencies.Count);
+			Assert.AreEqual(0, node3.Inputs.Count);
+			Assert.AreEqual(1, node3.OptionalInputs.Count);
+
+			BgNodeDef node2 = node3.OptionalInputs[0].ProducingNode;
+			Assert.AreEqual(node2.Name, "Compile Shooter Game Win64");
+			Assert.AreEqual(0, node2.InputDependencies.Count);
+			Assert.AreEqual(0, node2.Inputs.Count);
+			Assert.AreEqual(1, node2.OptionalInputs.Count);
+
+			BgNodeDef node1 = node2.OptionalInputs[0].ProducingNode;
+			Assert.AreEqual(node1.Name, "Update Version Files");
+			Assert.AreEqual(0, node1.InputDependencies.Count);
+			Assert.AreEqual(0, node1.Inputs.Count);
+			Assert.AreEqual(0, node1.OptionalInputs.Count);
+		}
+
+		[TestMethod]
+		public void NodeWithRequiresAndOptionalRequiresTest()
+		{
+			BgAgent agent = new("name", "type");
+			BgNode nodeSpec1 = agent.AddNode(x => UpdateVersionFiles());
+
+			BgNode nodeSpec2 = agent.AddNode(x => CompileShooterGameWin64()).Requires(nodeSpec1);
+
+			BgNode nodeSpec3 = agent.AddNode(x => CookShooterGameWin64()).Requires(nodeSpec1);
+
+			BgNode nodeSpec4 = agent.AddNode(x => PackageShooterGameWin64()).
+				Requires(BgList<BgFileSet>.Create(nodeSpec2.DefaultOutput)).
+				OptionalRequires(BgList<BgFileSet>.Create(nodeSpec3.DefaultOutput));
+
+			BgNodeDef node4 = ((BgObjectDef)Evaluate(nodeSpec4)).Deserialize<BgNodeExpressionDef>();
+			Assert.AreEqual(node4.Name, "Package Shooter Game Win64");
+			Assert.AreEqual(2, node4.InputDependencies.Count);
+			Assert.AreEqual(1, node4.Inputs.Count);
+			Assert.AreEqual(1, node4.OptionalInputs.Count);
+
+			BgNodeDef node3 = node4.OptionalInputs[0].ProducingNode;
+			Assert.AreEqual(node3.Name, "Cook Shooter Game Win64");
+			Assert.AreEqual(1, node3.InputDependencies.Count);
+			Assert.AreEqual(1, node3.Inputs.Count);
+			Assert.AreEqual(0, node3.OptionalInputs.Count);
+
+			BgNodeDef node2 = node4.Inputs[0].ProducingNode;
+			Assert.AreEqual(node2.Name, "Compile Shooter Game Win64");
+			Assert.AreEqual(1, node2.InputDependencies.Count);
+			Assert.AreEqual(1, node2.Inputs.Count);
+			Assert.AreEqual(0, node2.OptionalInputs.Count);
+
+			BgNodeDef node1 = node2.Inputs[0].ProducingNode;
+			Assert.AreEqual(node1.Name, "Update Version Files");
+			Assert.AreEqual(0, node1.InputDependencies.Count);
+			Assert.AreEqual(0, node1.Inputs.Count);
+			Assert.AreEqual(0, node1.OptionalInputs.Count);
+		}
+
+		[TestMethod]
+		public void OptionalRequiresChainsThruRequiredDependenciesTest()
+		{
+			BgAgent agent = new("name", "type");
+			BgNode nodeSpec1 = agent.AddNode(x => UpdateVersionFiles());
+			BgNode nodeSpec2 = agent.AddNode(x => CompileShooterGameWin64()).Requires(nodeSpec1);
+			BgNode nodeSpec3 = agent.AddNode(x => CookShooterGameWin64()).OptionalRequires(nodeSpec2);
+
+			BgNodeDef node3 = ((BgObjectDef)Evaluate(nodeSpec3)).Deserialize<BgNodeExpressionDef>();
+
+			// Node 3 has no required dependencies
+			Assert.IsEmpty(node3.InputDependencies);
+			Assert.IsEmpty(node3.Inputs);
+
+			// Flattened outputs from Node 2; includes Node 2's output + Node 1's output
+			Assert.HasCount(2, node3.OptionalInputs);
+			Assert.IsTrue(node3.OptionalInputs.Any(x => x.ProducingNode.Name == "Update Version Files"));
+			Assert.IsTrue(node3.OptionalInputs.Any(x => x.ProducingNode.Name == "Compile Shooter Game Win64"));
+
+			// Recursive chain; Node 1 + Node 2
+			Assert.HasCount(2, node3.OptionalInputDependencies);
+			Assert.IsTrue(node3.OptionalInputDependencies.Any(x => x.Name == "Update Version Files"));
+			Assert.IsTrue(node3.OptionalInputDependencies.Any(x => x.Name == "Compile Shooter Game Win64"));
+		}
+
+		[TestMethod]
+		public void RequiredDependencyChainsThruOptionalDependenciesTest()
+		{
+			BgAgent agent = new("name", "type");
+			BgNode nodeSpec1 = agent.AddNode(x => UpdateVersionFiles());
+			BgNode nodeSpec2 = agent.AddNode(x => CompileShooterGameWin64()).OptionalRequires(nodeSpec1);
+			BgNode nodeSpec3 = agent.AddNode(x => CookShooterGameWin64()).Requires(nodeSpec2);
+
+			BgNodeDef node3 = ((BgObjectDef)Evaluate(nodeSpec3)).Deserialize<BgNodeExpressionDef>();
+
+			// When Node 3 requires Node 2 includes:
+			// - Node 2 InputDependencies outputs (none)
+			// - Node 2 OptionalInputDependencies outputs (Node 1)
+			// - Node 2 Outputs
+			// Node 3's InputDependencies includes both Node 1 and Node 2
+			Assert.HasCount(2, node3.InputDependencies);
+			Assert.IsTrue(node3.InputDependencies.Any(x => x.Name == "Compile Shooter Game Win64"));
+			Assert.IsTrue(node3.InputDependencies.Any(x => x.Name == "Update Version Files"));
+
+			// Node 3 should have no optional dependencies
+			Assert.IsEmpty(node3.OptionalInputDependencies);
+
+			// Node 2's optional dependency on Node 1
+			BgNodeDef node2 = node3.InputDependencies.First(x => x.Name == "Compile Shooter Game Win64");
+			Assert.IsEmpty(node2.InputDependencies);
+			Assert.HasCount(1, node2.OptionalInputDependencies);
+			Assert.IsTrue(node2.OptionalInputDependencies.Any(x => x.Name == "Update Version Files"));
+		}
+
+		[TestMethod]
+		public void GetDirectOptionalInputDependenciesExcludesTransitiveRequiredTest()
+		{
+			// A --req--> B --opt--> C
+			BgAgent agent = new("name", "type");
+			BgNode nodeSpec1 = agent.AddNode(x => UpdateVersionFiles());
+			BgNode nodeSpec2 = agent.AddNode(x => CompileShooterGameWin64()).Requires(nodeSpec1);
+			BgNode nodeSpec3 = agent.AddNode(x => CookShooterGameWin64()).OptionalRequires(nodeSpec2);
+
+			BgNodeDef node3 = ((BgObjectDef)Evaluate(nodeSpec3)).Deserialize<BgNodeExpressionDef>();
+
+			// Flattened contains both
+			Assert.HasCount(2, node3.OptionalInputDependencies);
+
+			// Direct should only contain B (Compile Shooter Game Win64)
+			// A (Update Version Files) is transitive through B's required chain
+			List<BgNodeDef> directOptionalDeps = [.. node3.GetDirectOptionalInputDependencies()];
+			Assert.HasCount(1, directOptionalDeps);
+			Assert.AreEqual("Compile Shooter Game Win64", directOptionalDeps[0].Name);
+		}
+
+		[TestMethod]
+		public void GetDirectInputDependenciesExcludesTransitiveOptionalTest()
+		{
+			// A --opt--> B --req--> C
+			BgAgent agent = new("name", "type");
+			BgNode nodeSpec1 = agent.AddNode(x => UpdateVersionFiles());
+			BgNode nodeSpec2 = agent.AddNode(x => CompileShooterGameWin64()).OptionalRequires(nodeSpec1);
+			BgNode nodeSpec3 = agent.AddNode(x => CookShooterGameWin64()).Requires(nodeSpec2);
+
+			BgNodeDef node3 = ((BgObjectDef)Evaluate(nodeSpec3)).Deserialize<BgNodeExpressionDef>();
+
+			// Flattened contains both
+			Assert.HasCount(2, node3.InputDependencies);
+
+			// Direct should only contain B (Compile Shooter Game Win64)
+			// A (Update Version Files) is transitive through B's optional chain
+			List<BgNodeDef> directInputDeps = [.. node3.GetDirectInputDependencies()];
+			Assert.HasCount(1, directInputDeps);
+			Assert.AreEqual("Compile Shooter Game Win64", directInputDeps[0].Name);
 		}
 
 		[TestMethod]
@@ -135,5 +303,50 @@ namespace EpicGames.BuildGraph.Tests
 			Assert.IsTrue(label!.IncludedNodes.Contains(node1));
 			Assert.IsTrue(label!.IncludedNodes.Contains(node2));
 		}
+
+		[TestMethod]
+		public void TryResolveInputReferenceIncludesOptionalInputsTest()
+		{
+			BgAgent agent = new("test", "test");
+			BgNode nodeSpec1 = agent.AddNode(x => UpdateVersionFiles());
+			BgNode nodeSpec2 = agent.AddNode(x => CompileShooterGameWin64()).Requires(nodeSpec1);
+			BgNode nodeSpec3 = agent.AddNode(x => CookShooterGameWin64()).OptionalRequires(nodeSpec2);
+
+			BgGraph graphSpec = new(BgList.Create(nodeSpec1, nodeSpec2, nodeSpec3), BgList<BgAggregate>.Empty);
+			BgGraphDef graph = ((BgObjectDef)Evaluate(graphSpec)).Deserialize<BgGraphExpressionDef>().ToGraphDef();
+
+			// Resolving by node name should return outputs + inputs + optional inputs
+			bool resolved = graph.TryResolveInputReference("Cook Shooter Game Win64", out BgNodeOutput[]? outputs);
+			Assert.IsTrue(resolved);
+			Assert.IsNotNull(outputs);
+			Assert.HasCount(3, outputs);
+
+			string[] tagNames = [.. outputs.Select(x => x.TagName)];
+			Assert.Contains("#Cook Shooter Game Win64", tagNames);
+			Assert.Contains("#Compile Shooter Game Win64", tagNames);
+			Assert.Contains("#Update Version Files", tagNames);
+		}
+
+		[TestMethod]
+		public void TryResolveInputReferenceAggregateIncludesOptionalInputsTest()
+		{
+			BgAgent agent = new("test", "test");
+			BgNode nodeSpec1 = agent.AddNode(x => UpdateVersionFiles());
+			BgNode nodeSpec2 = agent.AddNode(x => CompileShooterGameWin64()).OptionalRequires(nodeSpec1);
+
+			BgAggregate aggregateSpec = new("All nodes", BgList.Create(nodeSpec2));
+			BgGraph graphSpec = new(BgList.Create(nodeSpec1, nodeSpec2), BgList.Create(aggregateSpec));
+			BgGraphDef graph = ((BgObjectDef)Evaluate(graphSpec)).Deserialize<BgGraphExpressionDef>().ToGraphDef();
+
+			// Resolving aggregate should include optional inputs from its required nodes
+			bool resolved = graph.TryResolveInputReference("All nodes", out BgNodeOutput[]? outputs);
+			Assert.IsTrue(resolved);
+			Assert.IsNotNull(outputs);
+			Assert.HasCount(2, outputs);
+
+			string[] tagNames = [.. outputs.Select(x => x.TagName)];
+			Assert.Contains("#Compile Shooter Game Win64", tagNames);
+			Assert.Contains("#Update Version Files", tagNames);
+		}
 	}
 }
diff --git a/Engine/Source/Programs/Shared/EpicGames.BuildGraph/BgGraphDef.cs b/Engine/Source/Programs/Shared/EpicGames.BuildGraph/BgGraphDef.cs
--- a/Engine/Source/Programs/Shared/EpicGames.BuildGraph/BgGraphDef.cs
+++ b/Engine/Source/Programs/Shared/EpicGames.BuildGraph/BgGraphDef.cs
@@ -198,8 +198,7 @@ namespace EpicGames.BuildGraph
 			if (name.StartsWith('#'))
 			{
 				// Check if it's a regular node or output name
-				BgNodeOutput? output;
-				if (TagNameToNodeOutput.TryGetValue(name, out output))
+				if (TagNameToNodeOutput.TryGetValue(name, out BgNodeOutput? output))
 				{
 					outOutputs = [output];
 					return true;
@@ -208,18 +207,16 @@ namespace EpicGames.BuildGraph
 			else
 			{
 				// Check if it's a regular node or output name
-				BgNodeDef? node;
-				if (NameToNode.TryGetValue(name, out node))
+				if (NameToNode.TryGetValue(name, out BgNodeDef? node))
 				{
-					outOutputs = [.. node.Outputs.Union(node.Inputs)];
+					outOutputs = [.. node.Outputs.Union(node.Inputs).Union(node.OptionalInputs)];
 					return true;
 				}
 
 				// Check if it's an aggregate name
-				BgAggregateDef? aggregate;
-				if (NameToAggregate.TryGetValue(name, out aggregate))
+				if (NameToAggregate.TryGetValue(name, out BgAggregateDef? aggregate))
 				{
-					outOutputs = [.. aggregate.RequiredNodes.SelectMany(x => x.Outputs.Union(x.Inputs)).Distinct()];
+					outOutputs = [.. aggregate.RequiredNodes.SelectMany(x => x.Outputs.Union(x.Inputs).Union(x.OptionalInputs)).Distinct()];
 					return true;
 				}
 			}
@@ -237,6 +234,10 @@ namespace EpicGames.BuildGraph
 				{
 					AddDependencies(inputDependency, retainNodes);
 				}
+				foreach (BgNodeDef optionalInputDependency in node.OptionalInputDependencies)
+				{
+					AddDependencies(optionalInputDependency, retainNodes);
+				}
 			}
 		}
 
@@ -521,6 +522,13 @@ namespace EpicGames.BuildGraph
 						}
 						jsonWriter.WriteArrayEnd();
 
+						jsonWriter.WriteArrayStart("OptionalInputs");
+						foreach (BgNodeOutput optionalInput in node.OptionalInputs)
+						{
+							jsonWriter.WriteValue(optionalInput.TagName);
+						}
+						jsonWriter.WriteArrayEnd();
+
 						jsonWriter.WriteArrayStart("Outputs");
 						foreach (BgNodeOutput output in node.Outputs)
 						{
@@ -535,6 +543,13 @@ namespace EpicGames.BuildGraph
 						}
 						jsonWriter.WriteArrayEnd();
 
+						jsonWriter.WriteArrayStart("OptionalInputDependencies");
+						foreach (string optionalInputDependency in node.GetDirectOptionalInputDependencies().Select(x => x.Name))
+						{
+							jsonWriter.WriteValue(optionalInputDependency);
+						}
+						jsonWriter.WriteArrayEnd();
+
 						jsonWriter.WriteArrayStart("OrderDependencies");
 						foreach (string orderDependency in node.GetDirectOrderDependencies().Select(x => x.Name))
 						{
@@ -719,8 +734,13 @@ namespace EpicGames.BuildGraph
 						{
 							logger.LogInformation("                input> {Name}", inputDependency.Name);
 						}
+						HashSet<BgNodeDef> optionalInputDependencies = [.. node.GetDirectOptionalInputDependencies()];
+						foreach (BgNodeDef optionalInputDependency in optionalInputDependencies)
+						{
+							logger.LogInformation("                optional> {Name}", optionalInputDependency.Name);
+						}
 						HashSet<BgNodeDef> orderDependencies = [.. node.GetDirectOrderDependencies()];
-						foreach (BgNodeDef orderDependency in orderDependencies.Except(inputDependencies))
+						foreach (BgNodeDef orderDependency in orderDependencies.Except(inputDependencies).Except(optionalInputDependencies))
 						{
 							logger.LogInformation("                after> {Name}", orderDependency.Name);
 						}
@@ -835,6 +855,10 @@ namespace EpicGames.BuildGraph
 				{
 					RegisterNode(graph, inputNode, uniqueNodes, uniqueAgents);
 				}
+				foreach (BgNodeExpressionDef optionalInputNode in node.OptionalInputDependencies.OfType<BgNodeExpressionDef>())
+				{
+					RegisterNode(graph, optionalInputNode, uniqueNodes, uniqueAgents);
+				}
 
 				BgAgentDef agent = node.Agent;
 				if (uniqueAgents.Add(agent))
diff --git a/Engine/Source/Programs/Shared/EpicGames.BuildGraph/BgNodeDef.cs b/Engine/Source/Programs/Shared/EpicGames.BuildGraph/BgNodeDef.cs
--- a/Engine/Source/Programs/Shared/EpicGames.BuildGraph/BgNodeDef.cs
+++ b/Engine/Source/Programs/Shared/EpicGames.BuildGraph/BgNodeDef.cs
@@ -79,7 +79,7 @@ namespace EpicGames.BuildGraph
 		{
 			if (OutputIndex == -1)
 			{
-				return [.. ProducingNode.InputDependencies.SelectMany(x => x.Outputs), .. ProducingNode.Outputs];
+				return [.. ProducingNode.InputDependencies.SelectMany(x => x.Outputs), .. ProducingNode.OptionalInputDependencies.SelectMany(x => x.Outputs), .. ProducingNode.Outputs];
 			}
 			else
 			{
@@ -117,6 +117,11 @@ namespace EpicGames.BuildGraph
 		/// </summary>
 		public List<BgNodeOutput> Inputs { get; } = [];
 
+		/// <summary>
+		/// Array of inputs which this node is not dependent on to run
+		/// </summary>
+		public List<BgNodeOutput> OptionalInputs { get; } = [];
+
 		/// <summary>
 		/// Array of outputs produced by this node
 		/// </summary>
@@ -127,6 +132,11 @@ namespace EpicGames.BuildGraph
 		/// </summary>
 		public List<BgNodeDef> InputDependencies { get; } = [];
 
+		/// <summary>
+		/// Nodes which this has optional input dependencies on
+		/// </summary>
+		public List<BgNodeDef> OptionalInputDependencies { get; } = [];
+
 		/// <summary>
 		/// Nodes which this node needs to run after
 		/// </summary>
@@ -198,17 +208,21 @@ namespace EpicGames.BuildGraph
 		/// </summary>
 		/// <param name="name">The name of this node</param>
 		/// <param name="inputs">Inputs that this node depends on</param>
+		/// <param name="optionalInputs">Inputs that this node can optionally depend on</param>
 		/// <param name="outputNames">Names of the outputs that this node produces</param>
 		/// <param name="inputDependencies">Nodes which this node is dependent on for its inputs</param>
+		/// <param name="optionalInputDependencies">Nodes which this node is optionally dependent on for its inputs</param>
 		/// <param name="orderDependencies">Nodes which this node needs to run after. Should include all input dependencies.</param>
 		/// <param name="requiredTokens">Optional tokens which must be required for this node to run</param>
 		/// <param name="ignoreModified">File patterns to ignore when checking for modified timestamps</param>
-		public BgNodeDef(string name, IReadOnlyList<BgNodeOutput> inputs, IReadOnlyList<string> outputNames, IReadOnlyList<BgNodeDef> inputDependencies, IReadOnlyList<BgNodeDef> orderDependencies, IReadOnlyList<FileReference> requiredTokens, IReadOnlyList<string> ignoreModified)
+		public BgNodeDef(string name, IReadOnlyList<BgNodeOutput> inputs, IReadOnlyList<BgNodeOutput> optionalInputs, IReadOnlyList<string> outputNames, IReadOnlyList<BgNodeDef> inputDependencies, IReadOnlyList<BgNodeDef> optionalInputDependencies, IReadOnlyList<BgNodeDef> orderDependencies, IReadOnlyList<FileReference> requiredTokens, IReadOnlyList<string> ignoreModified)
 			: this(name, null, outputNames)
 		{
 			Name = name;
 			Inputs.AddRange(inputs);
+			OptionalInputs.AddRange(optionalInputs);
 			InputDependencies.AddRange(inputDependencies);
+			OptionalInputDependencies.AddRange(optionalInputDependencies);
 			OrderDependencies.AddRange(orderDependencies);
 			RequiredTokens.AddRange(requiredTokens);
 			IgnoreModified.AddRange(ignoreModified);
@@ -229,6 +243,22 @@ namespace EpicGames.BuildGraph
 			foreach (BgNodeDef inputDependency in InputDependencies)
 			{
 				directDependencies.ExceptWith(inputDependency.InputDependencies);
+				directDependencies.ExceptWith(inputDependency.OptionalInputDependencies);
+			}
+			return directDependencies;
+		}
+
+		/// <summary>
+		/// Determines the minimal set of direct optional input dependencies for this node to run
+		/// </summary>
+		/// <returns>Sequence of nodes that are direct optional inputs to this node</returns>
+		public IEnumerable<BgNodeDef> GetDirectOptionalInputDependencies()
+		{
+			HashSet<BgNodeDef> directDependencies = [.. OptionalInputDependencies];
+			foreach (BgNodeDef optionalInputDependency in OptionalInputDependencies)
+			{
+				directDependencies.ExceptWith(optionalInputDependency.OptionalInputDependencies);
+				directDependencies.ExceptWith(optionalInputDependency.InputDependencies);
 			}
 			return directDependencies;
 		}
@@ -277,6 +307,11 @@ namespace EpicGames.BuildGraph
 		/// Input expressions
 		/// </summary>
 		public List<BgNodeOutputExprDef> InputExprs { get; } = [];
+		
+		/// <summary>
+		/// Optional input expressions
+		/// </summary>
+		public List<BgNodeOutputExprDef> OptionalInputExprs { get; } = [];
 
 		/// <summary>
 		/// Number of outputs from this node
@@ -308,6 +343,7 @@ namespace EpicGames.BuildGraph
 			obj.CopyTo(node);
 
 			node.Inputs.AddRange(node.InputExprs.SelectMany(x => x.Flatten()));
+			node.OptionalInputs.AddRange(node.OptionalInputExprs.SelectMany(x => x.Flatten()));
 
 			HashSet<BgNodeDef> inputDependencies = [];
 			foreach (BgNodeOutput input in node.Inputs)
@@ -319,6 +355,17 @@ namespace EpicGames.BuildGraph
 			inputDependencies.ExceptWith(node.InputDependencies);
 			node.InputDependencies.AddRange(inputDependencies);
 
+			HashSet<BgNodeDef> optionalInputDependencies = [];
+			foreach (BgNodeOutput optionalInput in node.OptionalInputs)
+			{
+				optionalInputDependencies.Add(optionalInput.ProducingNode);
+				optionalInputDependencies.UnionWith(optionalInput.ProducingNode.InputDependencies);
+				optionalInputDependencies.UnionWith(optionalInput.ProducingNode.OptionalInputDependencies);
+			}
+
+			optionalInputDependencies.ExceptWith(node.OptionalInputDependencies);
+			node.OptionalInputDependencies.AddRange(optionalInputDependencies);
+
 			return node;
 		}
 	}
diff --git a/Engine/Source/Programs/Shared/EpicGames.BuildGraph/Expressions/BgNode.cs b/Engine/Source/Programs/Shared/EpicGames.BuildGraph/Expressions/BgNode.cs
--- a/Engine/Source/Programs/Shared/EpicGames.BuildGraph/Expressions/BgNode.cs
+++ b/Engine/Source/Programs/Shared/EpicGames.BuildGraph/Expressions/BgNode.cs
@@ -81,6 +81,11 @@ namespace EpicGames.BuildGraph.Expressions
 		/// </summary>
 		public BgList<BgFileSet> Inputs { get; private set; } = BgList.Empty<BgFileSet>();
 
+		/// <summary>
+		/// Tokens for optional inputs of this node
+		/// </summary>
+		public BgList<BgFileSet> OptionalInputs { get; private set; } = BgList.Empty<BgFileSet>();
+
 		/// <summary>
 		/// Weak dependency on outputs that must be generated for the node to run, without making those dependencies inputs.
 		/// </summary>
@@ -120,6 +125,7 @@ namespace EpicGames.BuildGraph.Expressions
 			Agent = node.Agent;
 			DefaultOutput = new BgFileSetFromNodeOutputExpr(this, 0);
 			Inputs = node.Inputs;
+			OptionalInputs = node.OptionalInputs;
 			Fences = node.Fences;
 			RunEarly = node.RunEarly;
 			Labels = node.Labels;
@@ -134,6 +140,7 @@ namespace EpicGames.BuildGraph.Expressions
 			obj = obj.Set(x => x.Thunk, Thunk);
 			obj = obj.Set(x => x.OutputCount, (BgInt)OutputCount);
 			obj = obj.Set(x => x.InputExprs, Inputs);
+			obj = obj.Set(x => x.OptionalInputExprs, OptionalInputs);
 			obj = obj.Set(x => x.OrderDependencies, Fences);
 			obj = obj.Set(x => x.RunEarly, RunEarly);
 			obj = obj.Set(x => x.Labels, Labels);
@@ -144,17 +151,22 @@ namespace EpicGames.BuildGraph.Expressions
 		/// Creates a copy of this node and updates the given parameters
 		/// </summary>
 		/// <param name="inputs"></param>
+		/// <param name="optionalInputs"></param>
 		/// <param name="fences"></param>
 		/// <param name="runEarly"></param>
 		/// <param name="labels"></param>
 		/// <returns></returns>
-		internal BgNode Modify(BgList<BgFileSet>? inputs = null, BgList<BgNode>? fences = null, BgBool? runEarly = null, BgList<BgLabel>? labels = null)
+		internal BgNode Modify(BgList<BgFileSet>? inputs = null, BgList<BgFileSet>? optionalInputs = null, BgList<BgNode>? fences = null, BgBool? runEarly = null, BgList<BgLabel>? labels = null)
 		{
 			BgNode node = Clone();
 			if (inputs is not null)
 			{
 				node.Inputs = inputs;
 			}
+			if (optionalInputs is not null)
+			{
+				node.OptionalInputs = optionalInputs;
+			}
 			if (fences is not null)
 			{
 				node.Fences = fences;
@@ -483,6 +495,17 @@ namespace EpicGames.BuildGraph.Expressions
 			return (T)node.Modify(inputs: node.Inputs.Add(inputs.Select(x => (BgFileSet)x)));
 		}
 
+		/// <summary>
+		/// Add optional dependencies onto other nodes or outputs. Outputs from the given tokens will be copied to the current machine before execution of the node.
+		/// </summary>
+		/// <param name="node">The node to modify</param>
+		/// <param name="inputs">Files to add as optional inputs</param>
+		/// <returns>The current node spec, to allow chaining calls</returns>
+		public static T OptionalRequires<T>(this T node, params BgNode[] inputs) where T : BgNode
+		{
+			return (T)node.Modify(optionalInputs: node.OptionalInputs.Add(inputs.Select(x => (BgFileSet)x)));
+		}
+
 		/// <summary>
 		/// Add dependencies onto other nodes or outputs. Outputs from the given tokens will be copied to the current machine before execution of the node.
 		/// </summary>
@@ -494,6 +517,17 @@ namespace EpicGames.BuildGraph.Expressions
 			return (T)node.Modify(inputs: node.Inputs.Add(inputs));
 		}
 
+		/// <summary>
+		/// Add optional dependencies onto other nodes or outputs. Outputs from the given tokens will be copied to the current machine before execution of the node.
+		/// </summary>
+		/// <param name="node">The node to modify</param>
+		/// <param name="inputs">Files to add as optional inputs</param>
+		/// <returns>The current node spec, to allow chaining calls</returns>
+		public static T OptionalRequires<T>(this T node, params BgFileSet[] inputs) where T : BgNode
+		{
+			return (T)node.Modify(optionalInputs: node.OptionalInputs.Add(inputs));
+		}
+
 		/// <summary>
 		/// Add dependencies onto other nodes or outputs. Outputs from the given tokens will be copied to the current machine before execution of the node.
 		/// </summary>
@@ -505,6 +539,17 @@ namespace EpicGames.BuildGraph.Expressions
 			return (T)node.Modify(inputs: node.Inputs.Add(inputs));
 		}
 
+		/// <summary>
+		/// Add optional dependencies onto other nodes or outputs. Outputs from the given tokens will be copied to the current machine before execution of the node.
+		/// </summary>
+		/// <param name="node">The node to modify</param>
+		/// <param name="inputs">Files to add as optional inputs</param>
+		/// <returns>The current node spec, to allow chaining calls</returns>
+		public static T OptionalRequires<T>(this T node, BgList<BgFileSet> inputs) where T : BgNode
+		{
+			return (T)node.Modify(optionalInputs: node.OptionalInputs.Add(inputs));
+		}
+
 		/// <summary>
 		/// Add weak dependencies onto other nodes or outputs. The producing nodes must complete successfully if they are part of the graph, but outputs from them will not be 
 		/// transferred to the machine running this node.
