Logging: Added the frame counter (when present) to structured log output

Added parsing and storage of the frame to EpicGames.Core.LogEvent. It is only saved when non-zero and should only add storage/networking overhead when present.

Added writing of the frame to FOutputDeviceStdOut when running with -JsonStdOut or UE_LOG_JSON_TO_STDOUT=1.

#Synced-CL 50967057

diff --git a/Engine/Source/Programs/Shared/EpicGames.Core/JsonLogEvent.cs b/Engine/Source/Programs/Shared/EpicGames.Core/JsonLogEvent.cs
--- a/Engine/Source/Programs/Shared/EpicGames.Core/JsonLogEvent.cs
+++ b/Engine/Source/Programs/Shared/EpicGames.Core/JsonLogEvent.cs
@@ -256,6 +256,7 @@ namespace EpicGames.Core
 			if (reader.Read() && reader.TokenType == JsonTokenType.StartObject)
 			{
 				Utf8String? message = null;
+				ulong frame = 0;
 				DateTime? time = null;
 				while (reader.Read() && reader.TokenType == JsonTokenType.PropertyName)
 				{
@@ -268,6 +269,10 @@ namespace EpicGames.Core
 					{
 						time = reader.GetDateTime();
 					}
+					else if (propertyName.SequenceEqual(LogEventPropertyName.Frame) && reader.TokenType == JsonTokenType.Number)
+					{
+						frame = reader.GetUInt64();
+					}
 					else if (propertyName.SequenceEqual(LogEventPropertyName.Message) && reader.TokenType == JsonTokenType.String)
 					{
 						message = new Utf8String(reader.GetUtf8String().ToArray());
@@ -278,7 +283,7 @@ namespace EpicGames.Core
 				{
 					if (time is not null)
 					{
-						return $"[{time:yyyy.MM.dd-HH.mm.ss:fff}][  0]{message}"; // Structured logs currently don't contain frame number
+						return $"[{time:yyyy.MM.dd-HH.mm.ss:fff}][{frame % 1000,3}]{message}";
 					}
 					else
 					{
diff --git a/Engine/Source/Programs/Shared/EpicGames.Core/LogEvent.cs b/Engine/Source/Programs/Shared/EpicGames.Core/LogEvent.cs
--- a/Engine/Source/Programs/Shared/EpicGames.Core/LogEvent.cs
+++ b/Engine/Source/Programs/Shared/EpicGames.Core/LogEvent.cs
@@ -21,6 +21,7 @@ namespace EpicGames.Core
 	{
 #pragma warning disable CS1591 // Missing documentation
 		public static readonly Utf8String Time = new("time");
+		public static readonly Utf8String Frame = new("frame");
 		public static readonly Utf8String Level = new("level");
 		public static readonly Utf8String Id = new("id");
 		public static readonly Utf8String Line = new("line");
@@ -56,6 +57,11 @@ namespace EpicGames.Core
 		/// </summary>
 		public DateTime Time { get; set; }
 
+		/// <summary>
+		/// Frame that the event was emitted, or 0 when not applicable
+		/// </summary>
+		public ulong Frame { get; set; }
+
 		/// <summary>
 		/// The log level
 		/// </summary>
@@ -135,8 +141,17 @@ namespace EpicGames.Core
 		/// Constructor
 		/// </summary>
 		public LogEvent(DateTime time, LogLevel level, EventId eventId, int lineIndex, int lineCount, string message, string? format, IEnumerable<KeyValuePair<string, object?>>? properties, LogException? exception)
+			: this(time, 0, level, eventId, lineIndex, lineCount, message, format, properties, exception)
+		{
+		}
+
+		/// <summary>
+		/// Constructor
+		/// </summary>
+		public LogEvent(DateTime time, ulong frame, LogLevel level, EventId eventId, int lineIndex, int lineCount, string message, string? format, IEnumerable<KeyValuePair<string, object?>>? properties, LogException? exception)
 		{
 			Time = time;
+			Frame = frame;
 			Level = level;
 			Id = eventId;
 			LineIndex = lineIndex;
@@ -269,6 +284,7 @@ namespace EpicGames.Core
 #pragma warning restore CA1045 // Do not pass types by reference
 		{
 			DateTime time = new(0);
+			ulong frame = 0;
 			LogLevel level = LogLevel.None;
 			EventId eventId = new(0);
 			int line = 0;
@@ -285,6 +301,10 @@ namespace EpicGames.Core
 				{
 					time = reader.GetDateTime();
 				}
+				else if (Utf8StringComparer.OrdinalIgnoreCase.Equals(propertyName, LogEventPropertyName.Frame.Span))
+				{
+					frame = reader.GetUInt64();
+				}
 				else if (Utf8StringComparer.OrdinalIgnoreCase.Equals(propertyName, LogEventPropertyName.Level.Span))
 				{
 					level = Enum.Parse<LogLevel>(reader.GetString()!);
@@ -315,7 +335,7 @@ namespace EpicGames.Core
 				}
 			}
 
-			return new LogEvent(time, level, eventId, line, lineCount, message, format, properties, exception);
+			return new LogEvent(time, frame, level, eventId, line, lineCount, message, format, properties, exception);
 		}
 
 		static Dictionary<string, object?> ReadProperties(ref Utf8JsonReader reader)
@@ -414,6 +434,12 @@ namespace EpicGames.Core
 		{
 			writer.WriteStartObject();
 			writer.WriteString(LogEventPropertyName.Time.Span, Time.ToString("s", CultureInfo.InvariantCulture));
+
+			if (Frame != 0)
+			{
+				writer.WriteNumber(LogEventPropertyName.Frame.Span, Frame);
+			}
+
 			writer.WriteString(LogEventPropertyName.Level.Span, Level.ToString());
 			writer.WriteString(LogEventPropertyName.Message.Span, Message);
 
diff --git a/Engine/Source/Programs/Shared/EpicGames.Horde/Logs/ServerLogPacketBuilder.cs b/Engine/Source/Programs/Shared/EpicGames.Horde/Logs/ServerLogPacketBuilder.cs
--- a/Engine/Source/Programs/Shared/EpicGames.Horde/Logs/ServerLogPacketBuilder.cs
+++ b/Engine/Source/Programs/Shared/EpicGames.Horde/Logs/ServerLogPacketBuilder.cs
@@ -343,7 +343,7 @@ namespace EpicGames.Horde.Logs
 					message = message.Substring(0, maxMessageLength);
 				}
 
-				logEvent = new LogEvent(logEvent.Time, logEvent.Level, logEvent.Id, logEvent.LineIndex, logEvent.LineCount, $"{message} [...]", null, null, logEvent.Exception);
+				logEvent = new LogEvent(logEvent.Time, logEvent.Frame, logEvent.Level, logEvent.Id, logEvent.LineIndex, logEvent.LineCount, $"{message} [...]", null, null, logEvent.Exception);
 				jsonLogEvent = new JsonLogEvent(logEvent);
 			}
 
diff --git a/Engine/Source/Runtime/Core/Private/Misc/OutputDeviceStdOut.cpp b/Engine/Source/Runtime/Core/Private/Misc/OutputDeviceStdOut.cpp
--- a/Engine/Source/Runtime/Core/Private/Misc/OutputDeviceStdOut.cpp
+++ b/Engine/Source/Runtime/Core/Private/Misc/OutputDeviceStdOut.cpp
@@ -260,6 +260,10 @@ FORCENOINLINE void FOutputDeviceStdOutput::SerializeAsJson(const TCHAR* V, ELogV
 	TCbWriter<1024> Writer;
 	Writer.BeginObject();
 	Writer.AddDateTime(ANSITEXTVIEW("time"), FDateTime::UtcNow());
+	if (const uint64 Frame = GFrameCounter)
+	{
+		Writer.AddInteger(ANSITEXTVIEW("frame"), Frame);
+	}
 	Writer.AddString(ANSITEXTVIEW("level"), GetLevel(Verbosity));
 	AddMessage(Writer, V, Verbosity, Category, bShowCategory, bShowVerbosity);
 	if (bShowCategory || bShowVerbosity)
@@ -286,6 +290,10 @@ FORCENOINLINE void FOutputDeviceStdOutput::SerializeRecordAsJson(const UE::FLogR
 	TCbWriter<1024> Writer;
 	Writer.BeginObject();
 	Writer.AddDateTime(ANSITEXTVIEW("time"), Record.GetTime().GetUtcTime());
+	if (const uint64 Frame = GFrameCounter)
+	{
+		Writer.AddInteger(ANSITEXTVIEW("frame"), Frame);
+	}
 	Writer.AddString(ANSITEXTVIEW("level"), GetLevel(Record.GetVerbosity()));
 	AddMessage(Writer, Record, bShowCategory, bShowVerbosity);
 	if (bShowCategory || bShowVerbosity || Record.GetFields())
