EpicGames.Core: Add functions to get committed memory with tests

#rnx
#Synced-CL 52312713

diff --git a/Engine/Source/Programs/Shared/EpicGames.Core.Tests/SystemUtilsTests.cs b/Engine/Source/Programs/Shared/EpicGames.Core.Tests/SystemUtilsTests.cs
--- a/Engine/Source/Programs/Shared/EpicGames.Core.Tests/SystemUtilsTests.cs
+++ b/Engine/Source/Programs/Shared/EpicGames.Core.Tests/SystemUtilsTests.cs
@@ -126,6 +126,51 @@ public class SystemUtilsTests
 
 	#endregion
 
+	#region GetCurrentCommittedMemoryBytes
+
+	[TestMethod]
+	public void GetCurrentCommittedMemoryBytes_ReturnsPositiveValue()
+	{
+		long bytes = SystemUtils.GetCurrentCommittedMemoryBytes();
+		Assert.IsGreaterThan(0, bytes, $"Expected current committed memory > 0, got {bytes}");
+	}
+
+	[TestMethod]
+	public void GetCurrentCommittedMemoryBytes_NotGreaterThanMax()
+	{
+		long current = SystemUtils.GetCurrentCommittedMemoryBytes();
+		long max = SystemUtils.GetMaxCommittedMemoryBytes();
+		if (current >= 0 && max >= 0)
+		{
+			Assert.IsLessThanOrEqualTo(max, current, $"Current committed ({current}) must not exceed max committed ({max})");
+		}
+	}
+
+	#endregion
+
+	#region GetMaxCommittedMemoryBytes
+
+	[TestMethod]
+	public void GetMaxCommittedMemoryBytes_ReturnsPositiveValue()
+	{
+		long bytes = SystemUtils.GetMaxCommittedMemoryBytes();
+		Assert.IsGreaterThan(0, bytes, $"Expected max committed memory > 0, got {bytes}");
+	}
+
+	[TestMethod]
+	public void GetMaxCommittedMemoryBytes_AtLeastTotalSystemMemory()
+	{
+		long max = SystemUtils.GetMaxCommittedMemoryBytes();
+		long total = SystemUtils.GetTotalSystemMemoryBytes();
+		if (max >= 0 && total > 0)
+		{
+			// Commit limit should be at least as large as physical RAM (it includes the page file)
+			Assert.IsGreaterThanOrEqualTo(total, max, $"Max committed ({max}) should be >= total physical memory ({total})");
+		}
+	}
+
+	#endregion
+
 	#region GetTotalCpuUtilization
 
 	[TestMethod]
diff --git a/Engine/Source/Programs/Shared/EpicGames.Core/SystemUtils.cs b/Engine/Source/Programs/Shared/EpicGames.Core/SystemUtils.cs
--- a/Engine/Source/Programs/Shared/EpicGames.Core/SystemUtils.cs
+++ b/Engine/Source/Programs/Shared/EpicGames.Core/SystemUtils.cs
@@ -3,8 +3,10 @@
 using System;
 using System.Collections.Generic;
 using System.Diagnostics;
+using System.IO;
 using System.Linq;
 using System.Runtime.InteropServices;
+using System.Runtime.Versioning;
 using System.Text;
 
 namespace EpicGames.Core;
@@ -16,6 +18,7 @@ public static class SystemUtils
 {
 	#region Processor Count
 
+	[SupportedOSPlatform("windows")]
 	enum LogicalProcessorRelationship
 	{
 		RelationProcessorCore,
@@ -26,13 +29,19 @@ public static class SystemUtils
 		RelationAll = 0xffff
 	}
 
+	[SupportedOSPlatform("windows")]
 	[DllImport("kernel32.dll", SetLastError = true)]
 	static extern bool GetLogicalProcessorInformationEx(LogicalProcessorRelationship relationshipType, IntPtr buffer, ref uint returnedLength);
 
 	// int sysctlbyname(const char *name, void *oldp, size_t *oldlenp, void *newp, size_t newlen); // from man page
+	[SupportedOSPlatform("macos")]
 	[DllImport("libc")]
 	static extern int sysctlbyname(string name, out int oldp, ref ulong oldlenp, IntPtr newp, ulong newlen);
 
+	[SupportedOSPlatform("macos")]
+	[DllImport("libc")]
+	static extern int sysctlbyname(string name, out ulong oldp, ref ulong oldlenp, IntPtr newp, ulong newlen);
+
 	/// <summary>
 	/// Gets the number of logical cores. We use this rather than Environment.ProcessorCount when possible to handle machines with > 64 cores (the single group limit available to the .NET framework).
 	/// </summary>
@@ -197,6 +206,26 @@ public static class SystemUtils
 		return memoryInfo.TotalAvailableMemoryBytes != 0 ? memoryInfo.TotalAvailableMemoryBytes : -1;
 	}
 
+	// MEMORYSTATUSEX, from <sysinfoapi.h>
+	[SupportedOSPlatform("windows")]
+	[StructLayout(LayoutKind.Sequential)]
+	struct MemoryStatusEx
+	{
+		public uint _dwLength;
+		public uint _dwMemoryLoad;
+		public ulong _ullTotalPhys;
+		public ulong _ullAvailPhys;
+		public ulong _ullTotalPageFile;
+		public ulong _ullAvailPageFile;
+		public ulong _ullTotalVirtual;
+		public ulong _ullAvailVirtual;
+		public ulong _ullAvailExtendedVirtual;
+	}
+
+	[SupportedOSPlatform("windows")]
+	[DllImport("kernel32.dll", SetLastError = true)]
+	static extern bool GlobalMemoryStatusEx(ref MemoryStatusEx lpBuffer);
+
 	// vm_statistics64, based on the definition in <mach/vm_statistics.h>
 	[StructLayout(LayoutKind.Sequential)]
 	struct VmStatistics64
@@ -260,10 +289,12 @@ public static class SystemUtils
 	} // __attribute__((aligned(8)));
 
 	// kern_return_t host_statistics64(host_t host_priv, host_flavor_t flavor, host_info64_t host_info64_out, mach_msg_type_number_t *host_info64_outCnt); // from <mach/mach_host.h>
+	[SupportedOSPlatform("macos")]
 	[DllImport("libc")]
 	static extern int host_statistics64(IntPtr hostPriv, int flavor, out VmStatistics64 hostInfo64Out, ref uint hostInfoCount);
 
 	// mach_port_t mach_host_self() // from <mach/mach_init.h>
+	[SupportedOSPlatform("macos")]
 	[DllImport("libc")]
 	static extern IntPtr mach_host_self();
 
@@ -324,25 +355,135 @@ public static class SystemUtils
 		return freeMemoryBytes;
 	}
 
+	/// <summary>
+	/// Gets the current amount of committed (virtual) memory in bytes.
+	/// On Windows, this is the total page file usage (RAM + page file consumed by all processes).
+	/// On Linux, this is the Committed_AS value from /proc/meminfo.
+	/// On macOS, this is an approximation based on wired + active + inactive + compressor pages.
+	/// </summary>
+	/// <returns>Currently committed memory in bytes, or -1 if unavailable.</returns>
+	public static long GetCurrentCommittedMemoryBytes()
+	{
+		if (OperatingSystem.IsWindows())
+		{
+			MemoryStatusEx status = new() { _dwLength = (uint)Marshal.SizeOf<MemoryStatusEx>() };
+			if (GlobalMemoryStatusEx(ref status))
+			{
+				return (long)(status._ullTotalPageFile - status._ullAvailPageFile);
+			}
+		}
+		else if (OperatingSystem.IsLinux())
+		{
+			long kb = ReadProcMemInfoKilobytes("Committed_AS");
+			if (kb >= 0)
+			{
+				return kb * 1024L;
+			}
+		}
+		else if (OperatingSystem.IsMacOS())
+		{
+			const int HOST_VM_INFO64 = 4;
+			int hostVmInfo64Count = Marshal.SizeOf(typeof(VmStatistics64)) / 4;
+			uint structSize = (uint)hostVmInfo64Count;
+			IntPtr host = mach_host_self();
+			if (host_statistics64(host, HOST_VM_INFO64, out VmStatistics64 vmStats, ref structSize) == 0)
+			{
+				ulong outSize = 4;
+				if (0 != sysctlbyname("hw.pagesize", out int pageSize, ref outSize, IntPtr.Zero, 0))
+				{
+					pageSize = 4096;
+				}
+				long committed = (long)pageSize * ((long)vmStats._wireCount + vmStats._activeCount + vmStats._inactiveCount + vmStats._compressorPageCount);
+				return committed;
+			}
+		}
+		return -1;
+	}
+
+	/// <summary>
+	/// Gets the maximum amount of memory that can be committed in bytes (commit limit).
+	/// On Windows, this is the total page file size (RAM + page file capacity).
+	/// On Linux, this is the CommitLimit value from /proc/meminfo.
+	/// On macOS, this is the total physical RAM (hw.memsize).
+	/// </summary>
+	/// <returns>Commit limit in bytes, or -1 if unavailable.</returns>
+	public static long GetMaxCommittedMemoryBytes()
+	{
+		if (OperatingSystem.IsWindows())
+		{
+			MemoryStatusEx status = new() { _dwLength = (uint)Marshal.SizeOf<MemoryStatusEx>() };
+			if (GlobalMemoryStatusEx(ref status))
+			{
+				return (long)status._ullTotalPageFile;
+			}
+		}
+		else if (OperatingSystem.IsLinux())
+		{
+			long kb = ReadProcMemInfoKilobytes("CommitLimit");
+			if (kb >= 0)
+			{
+				return kb * 1024L;
+			}
+		}
+		else if (OperatingSystem.IsMacOS())
+		{
+			ulong outSize = sizeof(ulong);
+			if (0 == sysctlbyname("hw.memsize", out ulong memSize, ref outSize, IntPtr.Zero, 0))
+			{
+				return (long)memSize;
+			}
+		}
+		return -1;
+	}
+
+	// Parses a "<key>: <value> kB" line from /proc/meminfo and returns the value in kB, or -1.
+	[SupportedOSPlatform("linux")]
+	static long ReadProcMemInfoKilobytes(string key)
+	{
+		string prefix = key + ":";
+		foreach (string line in File.ReadLines("/proc/meminfo"))
+		{
+			if (line.StartsWith(prefix, StringComparison.Ordinal))
+			{
+				ReadOnlySpan<char> span = line.AsSpan(prefix.Length).Trim();
+				// Strip trailing " kB"
+				if (span.EndsWith(" kB", StringComparison.OrdinalIgnoreCase))
+				{
+					span = span[..^3];
+				}
+				if (long.TryParse(span, out long value))
+				{
+					return value;
+				}
+			}
+		}
+		return -1;
+	}
+
 	#endregion
 
 	#region CPU Utilization
 
+	[SupportedOSPlatform("windows")]
 	[DllImport("pdh.dll", SetLastError = true, CharSet = CharSet.Auto)]
 	static extern int PdhOpenQueryW([MarshalAs(UnmanagedType.LPWStr)] string? szDataSource, UIntPtr dwUserData, out IntPtr phQuery);
 
+	[SupportedOSPlatform("windows")]
 	[DllImport("pdh.dll", SetLastError = true, CharSet = CharSet.Auto)]
 	static extern uint PdhAddCounter(IntPtr hQuery, string szFullCounterPath, IntPtr dwUserData, out IntPtr phCounter);
 
+	[SupportedOSPlatform("windows")]
 	[DllImport("pdh.dll", SetLastError = true)]
 	static extern uint PdhCollectQueryData(IntPtr phQuery);
 
+	[SupportedOSPlatform("windows")]
 	struct PdhFmtCounterValue
 	{
 		public uint _cStatus;
 		public double _doubleValue;
 	};
 
+	[SupportedOSPlatform("windows")]
 	[DllImport("pdh.dll", SetLastError = true)]
 	static extern uint PdhGetFormattedCounterValue(IntPtr phCounter, uint dwFormat, IntPtr lpdwType, out PdhFmtCounterValue pValue);
 
