EpicGames.Core: Skip platform-specific SystemUtils tests on Linux

These two tests fail on Linux if running in cgroup limited containers because their invariants only hold on Windows.

diff --git a/Engine/Source/Programs/Shared/EpicGames.Core.Tests/SystemUtilsTests.cs b/Engine/Source/Programs/Shared/EpicGames.Core.Tests/SystemUtilsTests.cs
--- a/Engine/Source/Programs/Shared/EpicGames.Core.Tests/SystemUtilsTests.cs
+++ b/Engine/Source/Programs/Shared/EpicGames.Core.Tests/SystemUtilsTests.cs
@@ -20,6 +20,11 @@ public class SystemUtilsTests
 	[TestMethod]
 	public void GetLogicalProcessorCount_AtLeastEnvironmentProcessorCount()
 	{
+		if (!OperatingSystem.IsWindows())
+		{
+			return; // On Linux, lscpu respects cgroup CPU limits and may report fewer cores
+		}
+
 		// SystemUtils handles >64-core machines via group enumeration; it should never return
 		// fewer cores than the clamped Environment.ProcessorCount value.
 		int count = SystemUtils.GetLogicalProcessorCount();
@@ -160,6 +165,11 @@ public class SystemUtilsTests
 	[TestMethod]
 	public void GetMaxCommittedMemoryBytes_AtLeastTotalSystemMemory()
 	{
+		if (!OperatingSystem.IsWindows())
+		{
+			return; // On Linux, CommitLimit depends on overcommit_ratio and swap, may be < physical RAM
+		}
+
 		long max = SystemUtils.GetMaxCommittedMemoryBytes();
 		long total = SystemUtils.GetTotalSystemMemoryBytes();
 		if (max >= 0 && total > 0)
