EpicGames.Horde: Fix CommitIdWithOrder JSON serialization

CommitIdWithOrder extended CommitId but lacked a [JsonConverter] attribute causing System.Text.Json to fall back to default object serialization.

diff --git a/Engine/Source/Programs/Shared/EpicGames.Horde.Tests/Commits/CommitIdTests.cs b/Engine/Source/Programs/Shared/EpicGames.Horde.Tests/Commits/CommitIdTests.cs
new file mode 100644
--- /dev/null
+++ b/Engine/Source/Programs/Shared/EpicGames.Horde.Tests/Commits/CommitIdTests.cs
@@ -0,0 +1,87 @@
+// Copyright Epic Games, Inc. All Rights Reserved.
+
+using System.Text.Json;
+using EpicGames.Horde.Commits;
+using EpicGames.Horde.Jobs;
+using Microsoft.VisualStudio.TestTools.UnitTesting;
+
+namespace EpicGames.Horde.Tests.Commits;
+
+/// <summary>
+/// Verifies that <see cref="CommitIdWithOrder"/> serializes as a plain JSON string (not an object),
+/// matching the base <see cref="CommitId"/> behavior.
+/// </summary>
+[TestClass]
+public class CommitIdTests
+{
+	static readonly JsonSerializerOptions s_hordeOptions = CreateHordeOptions();
+
+	static JsonSerializerOptions CreateHordeOptions()
+	{
+		JsonSerializerOptions options = new();
+		HordeHttpRequest.ConfigureJsonSerializer(options);
+		return options;
+	}
+
+	[TestMethod]
+	public void CommitIdWithOrder_Serialize_MatchesCommitId()
+	{
+		CommitIdWithOrder commitId = new("12345", 42);
+		string json = JsonSerializer.Serialize(commitId, s_hordeOptions);
+
+		Assert.AreEqual("\"12345\"", json);
+		Assert.AreEqual(JsonSerializer.Serialize((CommitId)commitId, s_hordeOptions), json);
+	}
+
+	[TestMethod]
+	public void CommitIdWithOrder_InGetJobResponse_SerializesAsString()
+	{
+		GetJobResponse response = new() { CommitId = new CommitIdWithOrder("12345", 42) };
+		using JsonDocument doc = JsonDocument.Parse(JsonSerializer.Serialize(response, s_hordeOptions));
+
+		Assert.IsTrue(doc.RootElement.TryGetProperty("commitId", out JsonElement el));
+		Assert.AreEqual(JsonValueKind.String, el.ValueKind);
+	}
+
+	[TestMethod]
+	public void CommitIdWithOrder_Deserialize_FromString()
+	{
+		CommitIdWithOrder? commitId = JsonSerializer.Deserialize<CommitIdWithOrder>("\"12345\"", s_hordeOptions);
+
+		Assert.IsNotNull(commitId);
+		Assert.AreEqual("12345", commitId.Name);
+		Assert.AreEqual(12345, commitId.Order);
+	}
+
+	[TestMethod]
+	public void CommitIdWithOrder_Deserialize_FromNumber()
+	{
+		CommitIdWithOrder? commitId = JsonSerializer.Deserialize<CommitIdWithOrder>("12345", s_hordeOptions);
+
+		Assert.IsNotNull(commitId);
+		Assert.AreEqual("12345", commitId.Name);
+		Assert.AreEqual(12345, commitId.Order);
+	}
+
+	[TestMethod]
+	public void CommitIdWithOrder_Roundtrip_PreservesNumericOrder()
+	{
+		CommitIdWithOrder original = CommitIdWithOrder.FromPerforceChange(12345);
+		string json = JsonSerializer.Serialize(original, s_hordeOptions);
+		CommitIdWithOrder? result = JsonSerializer.Deserialize<CommitIdWithOrder>(json, s_hordeOptions);
+
+		Assert.IsNotNull(result);
+		Assert.AreEqual(original.Name, result.Name);
+		Assert.AreEqual(original.Order, result.Order);
+	}
+
+	[TestMethod]
+	public void CommitIdWithOrder_Deserialize_FromNonNumericString()
+	{
+		CommitIdWithOrder? commitId = JsonSerializer.Deserialize<CommitIdWithOrder>("\"abc123def\"", s_hordeOptions);
+
+		Assert.IsNotNull(commitId);
+		Assert.AreEqual("abc123def", commitId.Name);
+		Assert.AreEqual(0, commitId.Order);
+	}
+}
diff --git a/Engine/Source/Programs/Shared/EpicGames.Horde/Commits/CommitId.cs b/Engine/Source/Programs/Shared/EpicGames.Horde/Commits/CommitId.cs
--- a/Engine/Source/Programs/Shared/EpicGames.Horde/Commits/CommitId.cs
+++ b/Engine/Source/Programs/Shared/EpicGames.Horde/Commits/CommitId.cs
@@ -128,9 +128,29 @@ namespace EpicGames.Horde.Commits
 			=> value?.ToString();
 	}
 
+	class CommitIdWithOrderJsonConverter : JsonConverter<CommitIdWithOrder>
+	{
+		/// <inheritdoc/>
+		public override CommitIdWithOrder? Read(ref Utf8JsonReader reader, Type typeToConvert, JsonSerializerOptions options)
+			=> reader.TokenType switch
+			{
+				JsonTokenType.Number => CommitIdWithOrder.FromPerforceChange(reader.GetInt32()),
+				JsonTokenType.String => ParseString(reader.GetString()!),
+				_ => throw new JsonException("Invalid commit id")
+			};
+
+		static CommitIdWithOrder ParseString(string s) => Int32.TryParse(s, out int change)
+				? CommitIdWithOrder.FromPerforceChange(change)
+				: new CommitIdWithOrder(s, 0);
+
+		/// <inheritdoc/>
+		public override void Write(Utf8JsonWriter writer, CommitIdWithOrder value, JsonSerializerOptions options) => writer.WriteStringValue(value.Name);
+	}
+
 	/// <summary>
 	/// Variant of <see cref="CommitId"/> including a value that allows it to be used for ordering.
 	/// </summary>
+	[JsonConverter(typeof(CommitIdWithOrderJsonConverter))]
 	public class CommitIdWithOrder : CommitId, IEquatable<CommitIdWithOrder>, IComparable<CommitIdWithOrder>
 	{
 		/// <summary>
