[Backout] - CL52346907
#rnx
Original CL Desc
-----------------------------------------------------------------
EpicGames.Horde: Fix CommitIdWithOrder JSON serialization

CommitIdWithOrder extended CommitId but lacked a [JsonConverter] attribute causing System.Text.Json to fall back to default object serialization.

diff --git a/Engine/Source/Programs/Shared/EpicGames.Horde.Tests/Commits/CommitIdTests.cs b/Engine/Source/Programs/Shared/EpicGames.Horde.Tests/Commits/CommitIdTests.cs
deleted file mode 100644
diff --git a/Engine/Source/Programs/Shared/EpicGames.Horde/Commits/CommitId.cs b/Engine/Source/Programs/Shared/EpicGames.Horde/Commits/CommitId.cs
--- a/Engine/Source/Programs/Shared/EpicGames.Horde/Commits/CommitId.cs
+++ b/Engine/Source/Programs/Shared/EpicGames.Horde/Commits/CommitId.cs
@@ -128,29 +128,9 @@ namespace EpicGames.Horde.Commits
 			=> value?.ToString();
 	}
 
-	class CommitIdWithOrderJsonConverter : JsonConverter<CommitIdWithOrder>
-	{
-		/// <inheritdoc/>
-		public override CommitIdWithOrder? Read(ref Utf8JsonReader reader, Type typeToConvert, JsonSerializerOptions options)
-			=> reader.TokenType switch
-			{
-				JsonTokenType.Number => CommitIdWithOrder.FromPerforceChange(reader.GetInt32()),
-				JsonTokenType.String => ParseString(reader.GetString()!),
-				_ => throw new JsonException("Invalid commit id")
-			};
-
-		static CommitIdWithOrder ParseString(string s) => Int32.TryParse(s, out int change)
-				? CommitIdWithOrder.FromPerforceChange(change)
-				: new CommitIdWithOrder(s, 0);
-
-		/// <inheritdoc/>
-		public override void Write(Utf8JsonWriter writer, CommitIdWithOrder value, JsonSerializerOptions options) => writer.WriteStringValue(value.Name);
-	}
-
 	/// <summary>
 	/// Variant of <see cref="CommitId"/> including a value that allows it to be used for ordering.
 	/// </summary>
-	[JsonConverter(typeof(CommitIdWithOrderJsonConverter))]
 	public class CommitIdWithOrder : CommitId, IEquatable<CommitIdWithOrder>, IComparable<CommitIdWithOrder>
 	{
 		/// <summary>
