RigVMStructInputVariables: Add ExternalVariable UPROPERTY meta support for mutable input pins on RigVM structs

Adds support for UPROPERTY(meta=(Input, ExternalVariable)) on FRigUnit/FRigVMStruct
properties. This generates a non-const T& parameter in the static Execute signature,
matching what graph-based functions already support via bIsInputVariable.

Changes include:
- FRigVMStruct::ExternalVariableMetaName constant and GetPinDirectionFromProperty() update
- UHT InputVariable flag and IsConst() fix for non-const ref generation
- URigVMUnitNode::IsPure returns false for nodes with input variable pins
- URigVMPin::CanLink pure-function guard for input variable pins
- SetFunctionIsPure link cleanup on mutable-to-pure transition
- Engine tests for all of the above

#rnx

#Synced-CL 52296664

diff --git a/Engine/Plugins/Runtime/RigVM/Source/RigVM/Private/RigVMCore/RigVMStruct.cpp b/Engine/Plugins/Runtime/RigVM/Source/RigVM/Private/RigVMCore/RigVMStruct.cpp
--- a/Engine/Plugins/Runtime/RigVM/Source/RigVM/Private/RigVMCore/RigVMStruct.cpp
+++ b/Engine/Plugins/Runtime/RigVM/Source/RigVM/Private/RigVMCore/RigVMStruct.cpp
@@ -615,20 +615,62 @@ bool FRigVMStruct::ShouldCreatePinForProperty(const FProperty* InProperty) const
 
 ERigVMPinDirection FRigVMStruct::GetPinDirectionFromProperty(FProperty* InProperty)
 {
-	bool bIsInput = InProperty->HasMetaData(InputMetaName);
-	bool bIsOutput = InProperty->HasMetaData(OutputMetaName);
-	bool bIsVisible = InProperty->HasMetaData(VisibleMetaName);
+	const bool bIsInput = InProperty->HasMetaData(InputMetaName);
+	const bool bIsOutput = InProperty->HasMetaData(OutputMetaName);
+	const bool bIsIO = InProperty->HasMetaData(IOMetaName);
+	const bool bIsVisible = InProperty->HasMetaData(VisibleMetaName);
+	const bool bIsExternalVariable = InProperty->HasMetaData(ExternalVariableMetaName);
+
+	// Validate ExternalVariable usage: only valid as (Input, ExternalVariable) — no other combinations.
+	if (bIsExternalVariable)
+	{
+		if (bIsIO)
+		{
+			// IO already implies mutability; ExternalVariable is redundant and confusing.
+			UE_LOG(LogRigVM, Warning,
+				TEXT("Property '%s' has ExternalVariable meta combined with IO — IO already implies mutability. Treating as Hidden."),
+				*InProperty->GetName());
+			return ERigVMPinDirection::Hidden;
+		}
+		if (bIsOutput)
+		{
+			// Output (or Input+Output) is incompatible with ExternalVariable.
+			UE_LOG(LogRigVM, Warning,
+				TEXT("Property '%s' has ExternalVariable meta combined with Output — ExternalVariable is not compatible with Output. Treating as Hidden."),
+				*InProperty->GetName());
+			return ERigVMPinDirection::Hidden;
+		}
+		if (!bIsInput)
+		{
+			// ExternalVariable requires Input — standalone ExternalVariable is invalid.
+			UE_LOG(LogRigVM, Warning,
+				TEXT("Property '%s' has ExternalVariable meta without Input — ExternalVariable requires Input. Treating as Hidden."),
+				*InProperty->GetName());
+			return ERigVMPinDirection::Hidden;
+		}
+		if (InProperty->HasMetaData(ConstantMetaName))
+		{
+			// Constant forces immutability, which contradicts ExternalVariable's mutable-reference purpose.
+			UE_LOG(LogRigVM, Warning,
+				TEXT("Property '%s' has ExternalVariable meta combined with Constant — these are contradictory. Treating as Hidden."),
+				*InProperty->GetName());
+			return ERigVMPinDirection::Hidden;
+		}
+
+		// Valid case: (Input, ExternalVariable) — mutable input passed by non-const reference.
+		return ERigVMPinDirection::Input;
+	}
 
 	if (bIsVisible)
 	{
 		return ERigVMPinDirection::Visible;
 	}
-	
+
 	if (bIsInput)
 	{
 		return bIsOutput ? ERigVMPinDirection::IO : ERigVMPinDirection::Input;
-	} 
-	
+	}
+
 	if(bIsOutput)
 	{
 		return ERigVMPinDirection::Output;
diff --git a/Engine/Plugins/Runtime/RigVM/Source/RigVM/Public/RigVMCore/RigVMStruct.h b/Engine/Plugins/Runtime/RigVM/Source/RigVM/Public/RigVMCore/RigVMStruct.h
--- a/Engine/Plugins/Runtime/RigVM/Source/RigVM/Public/RigVMCore/RigVMStruct.h
+++ b/Engine/Plugins/Runtime/RigVM/Source/RigVM/Public/RigVMCore/RigVMStruct.h
@@ -261,6 +261,8 @@ public:
 	static inline const FLazyName IOMetaName = FLazyName(TEXT("IO"));
 	static inline const FLazyName HiddenMetaName = FLazyName(TEXT("Hidden"));
 	static inline const FLazyName VisibleMetaName = FLazyName(TEXT("Visible"));
+	/** Meta tag for input pins that are mutable (passed by non-const reference). Must be combined with Input. */
+	static inline const FLazyName ExternalVariableMetaName = FLazyName(TEXT("ExternalVariable"));
 	static inline const FLazyName DetailsOnlyMetaName = FLazyName(TEXT("DetailsOnly"));
 	static inline const FLazyName AbstractMetaName = FLazyName(TEXT("Abstract"));
 	static inline const FLazyName CategoryMetaName = FLazyName(TEXT("Category"));
diff --git a/Engine/Plugins/Runtime/RigVM/Source/RigVMDeveloper/Private/RigVMCompiler/RigVMCompiler.cpp b/Engine/Plugins/Runtime/RigVM/Source/RigVMDeveloper/Private/RigVMCompiler/RigVMCompiler.cpp
--- a/Engine/Plugins/Runtime/RigVM/Source/RigVMDeveloper/Private/RigVMCompiler/RigVMCompiler.cpp
+++ b/Engine/Plugins/Runtime/RigVM/Source/RigVMDeveloper/Private/RigVMCompiler/RigVMCompiler.cpp
@@ -1386,6 +1386,27 @@ bool URigVMCompiler::Compile(const FRigVMCompileSettings& InSettings, TArray<URi
 				}
 			}
 
+			// Validate that ExternalVariable (InputVariable) pins on unit nodes are bound.
+			// These pins are mutable references to caller-owned data — if neither a link
+			// nor a variable binding provides the reference, the node cannot execute correctly.
+			if(URigVMUnitNode* UnitNode = Cast<URigVMUnitNode>(ModelNode))
+			{
+				for(URigVMPin* Pin : UnitNode->GetPins())
+				{
+					if(Pin->IsDefinedAsInputVariable())
+					{
+						const bool bHasLink = Pin->GetSourceLinks().Num() > 0;
+						const bool bHasBinding = Pin->IsBoundToVariable();
+						if(!bHasLink && !bHasBinding)
+						{
+							static const FString UnboundInputVariableMessage = TEXT("Node @@ has an ExternalVariable pin that is not linked or bound to a variable. ExternalVariable pins must be resolved at runtime by the caller.");
+							Settings.ASTSettings.Report(EMessageSeverity::Error, ModelNode, UnboundInputVariableMessage);
+							bEncounteredGraphError = true;
+						}
+					}
+				}
+			}
+
 			// for variable let's validate ill formed variable nodes
 			if(URigVMVariableNode* VariableNode = Cast<URigVMVariableNode>(ModelNode))
 			{
diff --git a/Engine/Plugins/Runtime/RigVM/Source/RigVMDeveloper/Private/RigVMModel/Nodes/RigVMUnitNode.cpp b/Engine/Plugins/Runtime/RigVM/Source/RigVMDeveloper/Private/RigVMModel/Nodes/RigVMUnitNode.cpp
--- a/Engine/Plugins/Runtime/RigVM/Source/RigVMDeveloper/Private/RigVMModel/Nodes/RigVMUnitNode.cpp
+++ b/Engine/Plugins/Runtime/RigVM/Source/RigVMDeveloper/Private/RigVMModel/Nodes/RigVMUnitNode.cpp
@@ -90,6 +90,16 @@ bool URigVMUnitNode::IsDefinedAsVarying() const
 
 bool URigVMUnitNode::IsPure() const
 {
+	// A node with an input variable pin holds a mutable reference to caller-owned data.
+	// Treating it as pure would allow mutation in a pure context, so it must be non-pure.
+	for (const URigVMPin* Pin : GetPins())
+	{
+		if (Pin->IsDefinedAsInputVariable())
+		{
+			return false;
+		}
+	}
+
 	if (UScriptStruct* Struct = GetScriptStruct())
 	{
 		// Pure metadata explicitly marks nodes as pure, overriding hidden pin checks.
diff --git a/Engine/Plugins/Runtime/RigVM/Source/RigVMDeveloper/Private/RigVMModel/RigVMController.cpp b/Engine/Plugins/Runtime/RigVM/Source/RigVMDeveloper/Private/RigVMModel/RigVMController.cpp
--- a/Engine/Plugins/Runtime/RigVM/Source/RigVMDeveloper/Private/RigVMModel/RigVMController.cpp
+++ b/Engine/Plugins/Runtime/RigVM/Source/RigVMDeveloper/Private/RigVMModel/RigVMController.cpp
@@ -15473,12 +15473,75 @@ bool URigVMController::SetFunctionIsPure(const FName& InFunctionName, bool bIsPu
 
 	if (bIsPure)
 	{
+		// Group the execute-pin removal and any subsequent link breaks into a single
+		// undo action so that a single "Undo" reverses the entire operation.
+		if (bSetupUndoRedo)
+		{
+			FunctionController->OpenUndoBracket(TEXT("Set Function Pure"));
+		}
+
 		// Remove the execute pin to make the function pure
 		if (!FunctionController->RemoveExposedPin_Internal(FRigVMStruct::ExecuteContextName, bSetupUndoRedo, bPrintPythonCommand, /*bAllowExecutePinManagement=*/ true))
 		{
+			if (bSetupUndoRedo)
+			{
+				FunctionController->CancelUndoBracket();
+			}
 			ReportError(TEXT("Failed to remove execute pin from function."));
 			return false;
 		}
+
+		// After removing the execute pin, break any links to InputVariable pins that are
+		// now illegal. A non-local variable getter or entry node pin feeding an InputVariable
+		// pin would let the pure function mutate external state — those links must be broken.
+		// Local variable getters remain valid (function-scoped, not external state).
+		{
+			URigVMGraph* ContainedGraph = FunctionNode->GetContainedGraph();
+			const TArray<URigVMLink*>& Links = ContainedGraph->GetLinks();
+
+			TArray<TPair<FString, FString>> LinkPathsToBreak;
+			for (URigVMLink* Link : Links)
+			{
+				URigVMPin* SourcePin = Link->GetSourcePin();
+				URigVMPin* TargetPin = Link->GetTargetPin();
+				if (SourcePin == nullptr || TargetPin == nullptr)
+				{
+					continue;
+				}
+
+				if (!TargetPin->IsDefinedAsInputVariable())
+				{
+					continue;
+				}
+
+				// Reject if the source is a non-local variable getter
+				if (const URigVMVariableNode* VarNode = Cast<URigVMVariableNode>(SourcePin->GetNode()))
+				{
+					if (VarNode->IsGetter() && !VarNode->IsLocalVariable())
+					{
+						LinkPathsToBreak.Emplace(SourcePin->GetPinPath(), TargetPin->GetPinPath());
+						continue;
+					}
+				}
+
+				// Reject if the source is an entry node pin
+				if (SourcePin->GetNode()->IsA<URigVMFunctionEntryNode>())
+				{
+					LinkPathsToBreak.Emplace(SourcePin->GetPinPath(), TargetPin->GetPinPath());
+				}
+			}
+
+			for (const TPair<FString, FString>& PathPair : LinkPathsToBreak)
+			{
+				// ignore the return - non broken links can be discarded
+				(void)FunctionController->BreakLink(PathPair.Key, PathPair.Value, bSetupUndoRedo);
+			}
+		}
+
+		if (bSetupUndoRedo)
+		{
+			FunctionController->CloseUndoBracket();
+		}
 	}
 	else
 	{
diff --git a/Engine/Plugins/Runtime/RigVM/Source/RigVMDeveloper/Private/RigVMModel/RigVMPin.cpp b/Engine/Plugins/Runtime/RigVM/Source/RigVMDeveloper/Private/RigVMModel/RigVMPin.cpp
--- a/Engine/Plugins/Runtime/RigVM/Source/RigVMDeveloper/Private/RigVMModel/RigVMPin.cpp
+++ b/Engine/Plugins/Runtime/RigVM/Source/RigVMDeveloper/Private/RigVMModel/RigVMPin.cpp
@@ -2909,9 +2909,19 @@ bool URigVMPin::CanLink(const URigVMPin* InSourcePin, const URigVMPin* InTargetP
 			for (const URigVMPin* TargetPin : TargetPins)
 			{
 				TargetNode = TargetPin->GetNode();
-				if (TargetPin->GetDirection() == ERigVMPinDirection::IO)
+
+				// Whether this target is an IO pin or an input variable pin — both represent
+				// mutable references that a pure function must not be allowed to write through.
+				const bool bIsIOPin = (TargetPin->GetDirection() == ERigVMPinDirection::IO);
+				const bool bIsInputVariablePin = TargetPin->IsDefinedAsInputVariable();
+
+				if (bIsIOPin || bIsInputVariablePin)
 				{
-					// Local variable getters are allowed — they're function-scoped
+					// Build a human-readable label for the target pin kind used in failure messages
+					const TCHAR* TargetPinKind = bIsInputVariablePin ? TEXT("input variable") : TEXT("IO");
+
+					// Local variable getters are allowed — they're function-scoped and
+					// do not expose external state, so mutation is safe.
 					if (const URigVMVariableNode* VarNode = Cast<URigVMVariableNode>(SourcePin->GetNode()))
 					{
 						if (VarNode->IsGetter() && !VarNode->IsLocalVariable())
@@ -2919,22 +2929,23 @@ bool URigVMPin::CanLink(const URigVMPin* InSourcePin, const URigVMPin* InTargetP
 							if (OutFailureReason)
 							{
 								*OutFailureReason = FString::Printf(
-									TEXT("Only local variables can connect to IO pin '%s' in a pure function %s. Non-local data cannot be modified."),
-									*TargetPin->GetPinPath(), *OwningFunction->GetName());
+									TEXT("Only local variables can connect to %s pin '%s' in a pure function %s. Non-local data cannot be modified."),
+									TargetPinKind, *TargetPin->GetPinPath(), *OwningFunction->GetName());
 							}
-				
+
 							return false;
 						}
 					}
 
-					// Entry node pins (input parameters) connected to IO pins would allow mutation of caller data
+					// Entry node pins (input parameters) connected to IO / input variable pins
+					// would allow the pure function to mutate caller-owned data.
 					if (SourceNode->IsA<URigVMFunctionEntryNode>())
 					{
 						if (OutFailureReason)
 						{
 							*OutFailureReason = FString::Printf(
-								TEXT("Input parameter '%s' cannot connect to IO pin '%s' in a pure function %s. This would allow modification of the input."),
-								*SourcePin->GetName(), *TargetPin->GetPinPath(), *OwningFunction->GetName());
+								TEXT("Input parameter '%s' cannot connect to %s pin '%s' in a pure function %s. This would allow modification of the input."),
+								*SourcePin->GetName(), TargetPinKind, *TargetPin->GetPinPath(), *OwningFunction->GetName());
 						}
 						return false;
 					}
diff --git a/Engine/Source/Programs/Shared/EpicGames.UHT/Types/UhtScriptStruct.cs b/Engine/Source/Programs/Shared/EpicGames.UHT/Types/UhtScriptStruct.cs
--- a/Engine/Source/Programs/Shared/EpicGames.UHT/Types/UhtScriptStruct.cs
+++ b/Engine/Source/Programs/Shared/EpicGames.UHT/Types/UhtScriptStruct.cs
@@ -49,6 +49,12 @@ namespace EpicGames.UHT.Types
 		/// </summary>
 		EditorOnly = 0x00000010,
 
+		/// <summary>
+		/// "ExternalVariable" metadata specified — marks an Input pin as a mutable input variable.
+		/// The parameter is directionally an input but generates a non-const reference (T&amp; instead of const T&amp;).
+		/// </summary>
+		InputVariable = 0x00000020,
+
 		/// <summary>
 		/// Set if the property is an enum
 		/// </summary>
@@ -178,6 +184,12 @@ namespace EpicGames.UHT.Types
 		[JsonIgnore]
 		public bool Singleton => ParameterFlags.HasAnyFlags(UhtRigVMParameterFlags.Singleton);
 
+		/// <summary>
+		/// True if the parameter is marked as "ExternalVariable" — a mutable input variable
+		/// </summary>
+		[JsonIgnore]
+		public bool InputVariable => ParameterFlags.HasAnyFlags(UhtRigVMParameterFlags.InputVariable);
+
 		/// <summary>
 		/// True if the parameter is editor only
 		/// </summary>
@@ -224,6 +236,40 @@ namespace EpicGames.UHT.Types
 			ParameterFlags |= property.IsEditorOnlyProperty ? UhtRigVMParameterFlags.EditorOnly : UhtRigVMParameterFlags.None;
 			ParameterFlags |= property.MetaData.ContainsKey("Singleton") ? UhtRigVMParameterFlags.Singleton : UhtRigVMParameterFlags.None;
 			ParameterFlags |= property.MetaData.ContainsKey("Lazy") ? UhtRigVMParameterFlags.IsLazy : UhtRigVMParameterFlags.None;
+			ParameterFlags |= property.MetaData.ContainsKey("ExternalVariable") ? UhtRigVMParameterFlags.InputVariable : UhtRigVMParameterFlags.None;
+
+			// Validate ExternalVariable meta combinations before flag injection from Visible/Constant
+			// so that error messages reflect what the user actually wrote, not injected flags.
+			// Note: the Lazy + InputVariable check also lives here for the same reason (self-contained
+			// validation), even though the legacy Lazy + Output check remains in CollectRigVMMembers().
+			if (InputVariable)
+			{
+				if (!Input)
+				{
+					property.LogError($"RigVM Struct '{Property.Outer?.SourceName}' - Member '{Property.SourceName}' has 'ExternalVariable' metadata without 'Input'. ExternalVariable requires Input.");
+				}
+				if (Output)
+				{
+					property.LogError($"RigVM Struct '{Property.Outer?.SourceName}' - Member '{Property.SourceName}' has both 'Output' and 'ExternalVariable'. ExternalVariable is not compatible with Output.");
+				}
+				if (Constant)
+				{
+					property.LogError($"RigVM Struct '{Property.Outer?.SourceName}' - Member '{Property.SourceName}' has both 'Constant' and 'ExternalVariable'. ExternalVariable is not compatible with Constant.");
+				}
+				if (property.MetaData.ContainsKey("Visible"))
+				{
+					// 'Visible' means read-only display pin; ExternalVariable means mutable input variable.
+					// These are semantically incompatible — check here rather than relying on Visible's
+					// side-effect of injecting the Constant flag, which would produce a confusing error.
+					property.LogError($"RigVM Struct '{Property.Outer?.SourceName}' - Member '{Property.SourceName}' has both 'Visible' and 'ExternalVariable'. ExternalVariable is not compatible with Visible.");
+				}
+				// TRigVMLazyValue<const T&> wraps the value for deferred evaluation — it is incompatible
+				// with the mutable T& that InputVariable requires.
+				if (IsLazy)
+				{
+					property.LogError($"RigVM Struct '{Property.Outer?.SourceName}' - Member '{Property.SourceName}' has both 'Lazy' and 'ExternalVariable'. ExternalVariable is not compatible with Lazy.");
+				}
+			}
 
 			if (property.MetaData.ContainsKey("Visible"))
 			{
@@ -460,12 +506,14 @@ namespace EpicGames.UHT.Types
 		}
 
 		/// <summary>
-		/// True if the parameter is constant
+		/// True if the parameter is constant.
+		/// InputVariable params are mutable inputs — they emit T&amp; not const T&amp;, so they are not const
+		/// even though they have the Input flag set.
 		/// </summary>
 		/// <returns>True if the parameter is constant</returns>
 		public bool IsConst()
 		{
-			return Constant || (Input && !Output);
+			return Constant || (Input && !Output && !InputVariable);
 		}
 
 		/// <summary>
@@ -1201,6 +1249,8 @@ namespace EpicGames.UHT.Types
 							{
 								this.LogError($"RigVM Struct {SourceName} - Member {parameter.Name} is both an output and a lazy input.");
 							}
+							// Note: Lazy + InputVariable is validated in UhtRigVMParameter constructor
+							// alongside all other ExternalVariable combination checks.
 							RigVMStructInfo.Members.Add(parameter);
 						}
 					}
@@ -1210,6 +1260,21 @@ namespace EpicGames.UHT.Types
 				{
 					this.LogError($"RigVM Struct '{SourceName}' - has {RigVMStructInfo.Members.Count} members (64 is the limit).");
 				}
+
+				// ExternalVariable pins are mutable references — the node must have an execute
+				// context pin so it can be scheduled as a mutable (non-pure) operation.
+				// A struct with ExternalVariable but no execute pin would be neither meaningfully
+				// pure nor mutable, which is an invalid state.
+				bool bHasInputVariableMember = false;
+				foreach (UhtRigVMParameter m in RigVMStructInfo.Members)
+				{
+					if (m.InputVariable) { bHasInputVariableMember = true; break; }
+				}
+
+				if (!RigVMStructInfo.HasAnyExecuteContextMember && bHasInputVariableMember)
+				{
+					this.LogError($"RigVM Struct '{SourceName}' has one or more ExternalVariable properties but no execute context pin. ExternalVariable requires a mutable node — either declare an FRigVMExecuteContext (or subclass) member with no metadata on this struct, or inherit one from a base struct.");
+				}
 			}
 		}
 	}
