EpicGames.Core: Fix unapplied single quote replace

#rnx
#Synced-CL 52329845

diff --git a/Engine/Source/Programs/Shared/EpicGames.Core/SystemUtils.cs b/Engine/Source/Programs/Shared/EpicGames.Core/SystemUtils.cs
--- a/Engine/Source/Programs/Shared/EpicGames.Core/SystemUtils.cs
+++ b/Engine/Source/Programs/Shared/EpicGames.Core/SystemUtils.cs
@@ -548,6 +548,13 @@ public static class SystemUtils
 
 	static string RunProcessAndReturnStdOut(string command, string args)
 	{
+		// Process Arguments follow windows conventions in .NET Core
+		// Which means single quotes ' are not considered quotes.
+		// see [REDACTED_PROFILE_URL]/runtime/issues/29857
+		// also see UE-102580
+		// for rules see https://docs.microsoft.com/en-us/cpp/cpp/main-function-command-line-args
+		args = args.Replace('\'', '\"') ?? String.Empty;
+
 		ProcessStartInfo startInfo = new(command, args)
 		{
 			UseShellExecute = false,
@@ -558,7 +565,17 @@ public static class SystemUtils
 		};
 		using Process localProcess = Process.Start(startInfo)!;
 		string output = localProcess.StandardOutput.ReadToEnd().Trim();
+		string error = localProcess.StandardError.ReadToEnd().Trim();
 		localProcess.WaitForExit();
+
+		if (error.Length > 0)
+		{
+			if (output.Length > 0)
+			{
+				output += Environment.NewLine;
+			}
+			output += error;
+		}
 		return output;
 	}
 
