Shared: Handle projects with TargetFrameworks

#rnx
#Synced-CL 52364399

diff --git a/Engine/Source/Programs/Shared/EpicGames.Build/DotnetProjectInfo.cs b/Engine/Source/Programs/Shared/EpicGames.Build/DotnetProjectInfo.cs
--- a/Engine/Source/Programs/Shared/EpicGames.Build/DotnetProjectInfo.cs
+++ b/Engine/Source/Programs/Shared/EpicGames.Build/DotnetProjectInfo.cs
@@ -136,6 +136,36 @@ public sealed record class DotnetProjectInfo : IDotnetProjectInfo
 		};
 		Project project = Project.FromFile(file.FullName, projectOptions);
 
+		HashSet<string> targetFrameworks = [.. (project.GetProperty("TargetFrameworks")?.EvaluatedValue ?? String.Empty)
+			.Split(';', StringSplitOptions.TrimEntries | StringSplitOptions.RemoveEmptyEntries)];
+		if (targetFrameworks.Count > 0)
+		{
+			Dictionary<string, string> innerProps = project.GlobalProperties.ToDictionary(kvp => kvp.Key, kvp => kvp.Value);
+#if NET10_0_OR_GREATER
+			if (!innerProps.ContainsKey("TargetFramework") && targetFrameworks.Contains("net10.0"))
+			{
+				innerProps["TargetFramework"] = "net10.0";
+			}
+#endif
+#if NET9_0_OR_GREATER
+			if (!innerProps.ContainsKey("TargetFramework") && targetFrameworks.Contains("net9.0"))
+			{
+				innerProps["TargetFramework"] = "net9.0";
+			}
+#endif
+#if NET8_0_OR_GREATER
+			if (!innerProps.ContainsKey("TargetFramework") && targetFrameworks.Contains("net8.0"))
+			{
+				innerProps["TargetFramework"] = "net8.0";
+			}
+#endif
+			if (!innerProps.ContainsKey("TargetFramework"))
+			{
+				innerProps["TargetFramework"] = targetFrameworks.Order().Last();
+			}
+			project = new Project(project.FullPath, innerProps, project.ToolsVersion, project.ProjectCollection);
+		}
+
 		Dictionary<string, string> properties = project.Properties.ToDictionary(p => p.Name, p => p.EvaluatedValue);
 
 		DirectoryReference projectDirectory = file.Directory;
diff --git a/Engine/Source/Programs/Shared/EpicGames.MsBuild/CsProjBuilder.cs b/Engine/Source/Programs/Shared/EpicGames.MsBuild/CsProjBuilder.cs
--- a/Engine/Source/Programs/Shared/EpicGames.MsBuild/CsProjBuilder.cs
+++ b/Engine/Source/Programs/Shared/EpicGames.MsBuild/CsProjBuilder.cs
@@ -128,16 +128,47 @@ namespace EpicGames.MsBuild
 					try
 					{
 						project = new Project(projectPath, globalProperties, toolsVersion: null, projectCollection: projectCollection);
+
+						HashSet<string> targetFrameworks = [.. (project.GetProperty("TargetFrameworks")?.EvaluatedValue ?? String.Empty)
+							.Split(';', StringSplitOptions.TrimEntries | StringSplitOptions.RemoveEmptyEntries)];
+						if (targetFrameworks.Count > 0)
+						{
+							Dictionary<string, string> innerProps = project.GlobalProperties.ToDictionary(kvp => kvp.Key, kvp => kvp.Value);
+#if NET10_0_OR_GREATER
+							if (!innerProps.ContainsKey("TargetFramework") && targetFrameworks.Contains("net10.0"))
+							{
+								innerProps["TargetFramework"] = "net10.0";
+							}
+#endif
+#if NET9_0_OR_GREATER
+							if (!innerProps.ContainsKey("TargetFramework") && targetFrameworks.Contains("net9.0"))
+							{
+								innerProps["TargetFramework"] = "net9.0";
+							}
+#endif
+#if NET8_0_OR_GREATER
+							if (!innerProps.ContainsKey("TargetFramework") && targetFrameworks.Contains("net8.0"))
+							{
+								innerProps["TargetFramework"] = "net8.0";
+							}
+#endif
+							if (!innerProps.ContainsKey("TargetFramework"))
+							{
+								innerProps["TargetFramework"] = targetFrameworks.Order().Last();
+							}
+							project = new Project(project.FullPath, innerProps, project.ToolsVersion, project.ProjectCollection);
+						}
 					}
 					catch (Microsoft.Build.Exceptions.InvalidProjectFileException iPFEx)
 					{
-						logger.LogWarning("Could not load project file {ProjectPath}: {Message}", projectPath, iPFEx.BaseMessage);
+						logger.LogWarning(iPFEx, "Could not load project file {ProjectPath}: {Message}", projectPath, iPFEx.BaseMessage);
 						return;
 					}
 
 					// check the TargetFramework for unsupported versions. see https://learn.microsoft.com/en-us/dotnet/standard/frameworks
-					string targetFramework = project.GetProperty("TargetFramework").EvaluatedValue;
-					bool skipFramework = targetFramework.Contains("netstandard", StringComparison.Ordinal)
+					string targetFramework = project.GetProperty("TargetFramework")?.EvaluatedValue ?? String.Empty;
+					bool skipFramework = String.IsNullOrEmpty(targetFramework)
+						| targetFramework.Contains("netstandard", StringComparison.Ordinal)
 						| targetFramework.Contains("netcoreapp", StringComparison.Ordinal)
 						| targetFramework.Contains("uap", StringComparison.Ordinal)
 						| targetFramework.Contains("netnano", StringComparison.Ordinal)
