Shared: Move logging related files into a subdirectory.

#Synced-CL 52377985

diff --git a/Engine/Source/Programs/Shared/EpicGames.Core/DefaultConsoleLogger.cs b/Engine/Source/Programs/Shared/EpicGames.Core/DefaultConsoleLogger.cs
deleted file mode 100644
diff --git a/Engine/Source/Programs/Shared/EpicGames.Core/ForwardingLogger.cs b/Engine/Source/Programs/Shared/EpicGames.Core/ForwardingLogger.cs
deleted file mode 100644
diff --git a/Engine/Source/Programs/Shared/EpicGames.Core/ILogCursor.cs b/Engine/Source/Programs/Shared/EpicGames.Core/ILogCursor.cs
deleted file mode 100644
diff --git a/Engine/Source/Programs/Shared/EpicGames.Core/ILoggerIndent.cs b/Engine/Source/Programs/Shared/EpicGames.Core/ILoggerIndent.cs
deleted file mode 100644
diff --git a/Engine/Source/Programs/Shared/EpicGames.Core/ILoggerProgress.cs b/Engine/Source/Programs/Shared/EpicGames.Core/ILoggerProgress.cs
deleted file mode 100644
diff --git a/Engine/Source/Programs/Shared/EpicGames.Core/ILoggerProperty.cs b/Engine/Source/Programs/Shared/EpicGames.Core/ILoggerProperty.cs
deleted file mode 100644
diff --git a/Engine/Source/Programs/Shared/EpicGames.Core/JsonLogEvent.cs b/Engine/Source/Programs/Shared/EpicGames.Core/JsonLogEvent.cs
deleted file mode 100644
diff --git a/Engine/Source/Programs/Shared/EpicGames.Core/KnownLogEvents.cs b/Engine/Source/Programs/Shared/EpicGames.Core/KnownLogEvents.cs
deleted file mode 100644
diff --git a/Engine/Source/Programs/Shared/EpicGames.Core/LegacyEventLogger.cs b/Engine/Source/Programs/Shared/EpicGames.Core/LegacyEventLogger.cs
deleted file mode 100644
diff --git a/Engine/Source/Programs/Shared/EpicGames.Core/Log.cs b/Engine/Source/Programs/Shared/EpicGames.Core/Log.cs
deleted file mode 100644
diff --git a/Engine/Source/Programs/Shared/EpicGames.Core/LogBuffer.cs b/Engine/Source/Programs/Shared/EpicGames.Core/LogBuffer.cs
deleted file mode 100644
diff --git a/Engine/Source/Programs/Shared/EpicGames.Core/LogEvent.cs b/Engine/Source/Programs/Shared/EpicGames.Core/LogEvent.cs
deleted file mode 100644
diff --git a/Engine/Source/Programs/Shared/EpicGames.Core/LogEventBuilder.cs b/Engine/Source/Programs/Shared/EpicGames.Core/LogEventBuilder.cs
deleted file mode 100644
diff --git a/Engine/Source/Programs/Shared/EpicGames.Core/LogEventMarkup.cs b/Engine/Source/Programs/Shared/EpicGames.Core/LogEventMarkup.cs
deleted file mode 100644
diff --git a/Engine/Source/Programs/Shared/EpicGames.Core/LogEventParser.cs b/Engine/Source/Programs/Shared/EpicGames.Core/LogEventParser.cs
deleted file mode 100644
diff --git a/Engine/Source/Programs/Shared/EpicGames.Core/LogIndent.cs b/Engine/Source/Programs/Shared/EpicGames.Core/LogIndent.cs
deleted file mode 100644
diff --git a/Engine/Source/Programs/Shared/EpicGames.Core/LogIndentScope.cs b/Engine/Source/Programs/Shared/EpicGames.Core/LogIndentScope.cs
deleted file mode 100644
diff --git a/Engine/Source/Programs/Shared/EpicGames.Core/LogStatusScope.cs b/Engine/Source/Programs/Shared/EpicGames.Core/LogStatusScope.cs
deleted file mode 100644
diff --git a/Engine/Source/Programs/Shared/EpicGames.Core/LogValue.cs b/Engine/Source/Programs/Shared/EpicGames.Core/LogValue.cs
deleted file mode 100644
diff --git a/Engine/Source/Programs/Shared/EpicGames.Core/LogValueFormatter.cs b/Engine/Source/Programs/Shared/EpicGames.Core/LogValueFormatter.cs
deleted file mode 100644
diff --git a/Engine/Source/Programs/Shared/EpicGames.Core/Logging/DefaultConsoleLogger.cs b/Engine/Source/Programs/Shared/EpicGames.Core/Logging/DefaultConsoleLogger.cs
new file mode 100644
--- /dev/null
+++ b/Engine/Source/Programs/Shared/EpicGames.Core/Logging/DefaultConsoleLogger.cs
@@ -0,0 +1,206 @@
+// Copyright Epic Games, Inc. All Rights Reserved.
+
+using System;
+using System.Collections.Generic;
+using System.IO;
+using System.Linq;
+using Microsoft.Extensions.Logging;
+
+namespace EpicGames.Core
+{
+	/// <summary>
+	/// Simple implementation of <see cref="ILogger"/> which formats output to the console similar to Serilog's 'Code' style.
+	/// </summary>
+	public class DefaultConsoleLogger : ILogger
+	{
+		sealed class Disposable : IDisposable
+		{
+			public void Dispose() { }
+		}
+
+		const string AnsiReset = "\x1b[0m";
+
+		const string AnsiTime = "\x1b[38;5;0246m";
+
+		const string AnsiTrace = "\x1b[37m";
+		const string AnsiDebug = "\x1b[37m";
+		const string AnsiInformation = "\x1b[37;1m";
+		const string AnsiWarning = "\x1b[38;5;0229m";
+		const string AnsiError = "\x1b[38;5;0197m\x1b[48;5;0238m";
+
+		const string AnsiNull = "\x1b[38;5;0038m";
+		const string AnsiBool = "\x1b[38;5;0038m";
+		const string AnsiNumber = "\x1b[38;5;151m";
+		const string AnsiString = "\x1b[38;5;0216m";
+
+		readonly LogLevel _minLevel;
+
+		/// <summary>
+		/// Constructor
+		/// </summary>
+		/// <param name="minLevel">Minimum level for messages to be output</param>
+		public DefaultConsoleLogger(LogLevel minLevel = LogLevel.Debug)
+		{
+			_minLevel = minLevel;
+		}
+
+		/// <inheritdoc/>
+		public bool IsEnabled(LogLevel logLevel) => logLevel >= _minLevel;
+
+		/// <inheritdoc/>
+		public IDisposable? BeginScope<TState>(TState state) where TState : notnull => new Disposable();
+
+		/// <inheritdoc/>
+		public void Log<TState>(LogLevel logLevel, EventId eventId, TState state, Exception? exception, Func<TState, Exception?, string> formatter)
+		{
+			using StringWriter writer = new();
+			Write<TState>(logLevel, state, exception, formatter, writer);
+			Console.WriteLine(writer.ToString());
+		}
+
+		static void Write<TState>(LogLevel logLevel, TState state, Exception? exception, Func<TState, Exception?, string> formatter, TextWriter textWriter)
+		{
+			DateTime time = DateTime.Now;
+
+			textWriter.Write(AnsiTime);
+			textWriter.Write($"[{time:HH:mm:ss} ");
+			switch (logLevel)
+			{
+				case LogLevel.Trace:
+					textWriter.Write(AnsiTrace);
+					textWriter.Write("trc");
+					break;
+				case LogLevel.Debug:
+					textWriter.Write(AnsiDebug);
+					textWriter.Write("dbg");
+					break;
+				case LogLevel.Information:
+					textWriter.Write(AnsiInformation);
+					textWriter.Write("inf" + AnsiTime);
+					break;
+				case LogLevel.Warning:
+					textWriter.Write(AnsiWarning);
+					textWriter.Write("wrn" + AnsiTime);
+					break;
+				case LogLevel.Error:
+					textWriter.Write(AnsiError);
+					textWriter.Write("err");
+					break;
+				case LogLevel.Critical:
+					textWriter.Write(AnsiError);
+					textWriter.Write("ctl");
+					break;
+			}
+			textWriter.Write(AnsiReset);
+			textWriter.Write(AnsiTime);
+			textWriter.Write("] ");
+			textWriter.Write(AnsiReset);
+
+			const string OriginalFormatKeyName = "{OriginalFormat}";
+			if (state is IEnumerable<KeyValuePair<string, object>> enumerable)
+			{
+				string? originalFormat = enumerable.FirstOrDefault(x => x.Key.Equals(OriginalFormatKeyName, StringComparison.OrdinalIgnoreCase)).Value as string;
+				if (originalFormat != null)
+				{
+					ReadOnlySpan<char> fmt = originalFormat.AsSpan();
+					while (fmt.Length > 0)
+					{
+						fmt = WriteToken(fmt, enumerable, textWriter);
+					}
+					return;
+				}
+			}
+
+			textWriter.Write(formatter(state, exception));
+		}
+
+		static ReadOnlySpan<char> WriteToken(ReadOnlySpan<char> fmt, IEnumerable<KeyValuePair<string, object>> properties, TextWriter writer)
+		{
+			int idx = fmt.IndexOf('{');
+			if (idx == -1 || idx + 1 == fmt.Length)
+			{
+				writer.Write(fmt);
+				return [];
+			}
+			if (fmt[idx + 1] == '{')
+			{
+				writer.Write(fmt.Slice(0, idx + 1));
+				return fmt.Slice(idx + 2);
+			}
+
+			writer.Write(fmt.Slice(0, idx));
+			fmt = fmt.Slice(idx);
+
+			int endIdx = fmt.IndexOf('}');
+			if (endIdx == -1)
+			{
+				writer.Write(fmt);
+				return [];
+			}
+
+			ReadOnlySpan<char> name = fmt.Slice(1, endIdx - 1);
+			ReadOnlySpan<char> spec = [];
+
+			int specIdx = name.IndexOfAny(':', ',');
+			if (specIdx != -1)
+			{
+				spec = name.Slice(specIdx);
+				name = name.Slice(0, specIdx);
+			}
+
+			object? obj = GetPropertyValue(properties, name);
+			string? color = obj switch
+			{
+				null => AnsiNull,
+
+				bool _ => AnsiBool,
+
+				byte _ => AnsiNumber,
+				sbyte _ => AnsiNumber,
+				short _ => AnsiNumber,
+				ushort _ => AnsiNumber,
+				int _ => AnsiNumber,
+				ulong _ => AnsiNumber,
+				float _ => AnsiNumber,
+				double _ => AnsiNumber,
+
+				string _ => AnsiString,
+
+				_ => null
+			};
+
+			if (color != null)
+			{
+				writer.Write(color);
+			}
+
+			if (spec.Length > 0)
+			{
+				writer.Write($"{{0{spec}}}", obj);
+			}
+			else
+			{
+				writer.Write(obj);
+			}
+
+			if (color != null)
+			{
+				writer.Write(AnsiReset);
+			}
+
+			return fmt.Slice(endIdx + 1);
+		}
+
+		static object? GetPropertyValue(IEnumerable<KeyValuePair<string, object>> enumerable, ReadOnlySpan<char> name)
+		{
+			foreach (KeyValuePair<string, object> item in enumerable)
+			{
+				if (name.Equals(item.Key, StringComparison.Ordinal))
+				{
+					return item.Value;
+				}
+			}
+			return null;
+		}
+	}
+}
diff --git a/Engine/Source/Programs/Shared/EpicGames.Core/Logging/ForwardingLogger.cs b/Engine/Source/Programs/Shared/EpicGames.Core/Logging/ForwardingLogger.cs
new file mode 100644
--- /dev/null
+++ b/Engine/Source/Programs/Shared/EpicGames.Core/Logging/ForwardingLogger.cs
@@ -0,0 +1,83 @@
+// Copyright Epic Games, Inc. All Rights Reserved.
+
+using System;
+using System.Collections.Generic;
+using Microsoft.Extensions.Logging;
+
+namespace EpicGames.Core
+{
+	/// <summary>
+	/// Logger which forwards output to several other loggers
+	/// </summary>
+	public class ForwardingLogger : ILogger
+	{
+		/// <summary>
+		/// List of IDisposable instances
+		/// </summary>
+		class DisposableList : List<IDisposable>, IDisposable
+		{
+			/// <inheritdoc/>
+			public void Dispose()
+			{
+				foreach (IDisposable child in this)
+				{
+					child.Dispose();
+				}
+			}
+		}
+
+		/// <summary>
+		/// The loggers that are being written to
+		/// </summary>
+		public IReadOnlyList<ILogger> Loggers
+		{
+			get;
+		}
+
+		/// <summary>
+		/// Constructor
+		/// </summary>
+		/// <param name="loggers">Array of loggers to forward to</param>
+		public ForwardingLogger(params ILogger[] loggers)
+		{
+			Loggers = loggers;
+		}
+
+		/// <inheritdoc/>
+		public IDisposable? BeginScope<TState>(TState state) where TState : notnull
+		{
+			DisposableList scope = new();
+			foreach (ILogger logger in Loggers)
+			{
+				IDisposable? next = logger.BeginScope(state);
+				if (next != null)
+				{
+					scope.Add(next);
+				}
+			}
+			return scope;
+		}
+
+		/// <inheritdoc/>
+		public bool IsEnabled(LogLevel logLevel)
+		{
+			foreach (ILogger logger in Loggers)
+			{
+				if (logger.IsEnabled(logLevel))
+				{
+					return true;
+				}
+			}
+			return false;
+		}
+
+		/// <inheritdoc/>
+		public void Log<TState>(LogLevel logLevel, EventId eventId, TState state, Exception? exception, Func<TState, Exception?, string> formatter)
+		{
+			foreach (ILogger logger in Loggers)
+			{
+				logger.Log(logLevel, eventId, state, exception, formatter);
+			}
+		}
+	}
+}
diff --git a/Engine/Source/Programs/Shared/EpicGames.Core/Logging/ILogCursor.cs b/Engine/Source/Programs/Shared/EpicGames.Core/Logging/ILogCursor.cs
new file mode 100644
--- /dev/null
+++ b/Engine/Source/Programs/Shared/EpicGames.Core/Logging/ILogCursor.cs
@@ -0,0 +1,351 @@
+// Copyright Epic Games, Inc. All Rights Reserved.
+
+using System;
+using System.Diagnostics.CodeAnalysis;
+using System.Text.RegularExpressions;
+
+namespace EpicGames.Core
+{
+	/// <summary>
+	/// Allows querying the input text from the current cursor position
+	/// </summary>
+	public interface ILogCursor
+	{
+		/// <summary>
+		/// Text for the current line
+		/// </summary>
+		string? CurrentLine
+		{
+			get;
+		}
+
+		/// <summary>
+		/// The current line number
+		/// </summary>
+		int CurrentLineNumber
+		{
+			get;
+		}
+
+		/// <summary>
+		/// Returns the line count
+		/// </summary>
+		int Length
+		{
+			get;
+		}
+
+		/// <summary>
+		/// Index to find the string at the given offset
+		/// </summary>
+		/// <param name="offset"></param>
+		/// <returns></returns>
+		string? this[int offset]
+		{
+			get;
+		}
+	}
+
+	/// <summary>
+	/// Extension methods for log cursors
+	/// </summary>
+	public static partial class LogCursorExtensions
+	{
+		/// <summary>
+		/// Implementation of ILogCursor which positions the cursor at a fixed offset from the inner cursor
+		/// </summary>
+		class RebasedLogCursor : ILogCursor
+		{
+			readonly ILogCursor _inner;
+			readonly int _baseLineNumber;
+
+			public RebasedLogCursor(ILogCursor inner, int baseLineNumber)
+			{
+				_inner = inner;
+				_baseLineNumber = baseLineNumber;
+			}
+
+			public string? this[int index] => _inner[(_baseLineNumber + index) - _inner.CurrentLineNumber];
+			public string? CurrentLine => _inner[_baseLineNumber - _inner.CurrentLineNumber];
+			public int CurrentLineNumber => _baseLineNumber;
+			public int Length => _inner.Length - (_baseLineNumber - _inner.CurrentLineNumber);
+		}
+
+		/// <summary>
+		/// Creates a new log cursor based at an offset from the current line
+		/// </summary>
+		/// <param name="cursor">The current log cursor instance</param>
+		/// <param name="offset">Line number offset from the current</param>
+		/// <returns>New log cursor instance</returns>
+		public static ILogCursor Rebase(this ILogCursor cursor, int offset)
+		{
+			return new RebasedLogCursor(cursor, cursor.CurrentLineNumber + offset);
+		}
+
+		/// <summary>
+		/// Attempts to get a line at the given offset
+		/// </summary>
+		/// <param name="cursor">The log cursor instance</param>
+		/// <param name="offset">Offset of the line to retrieve</param>
+		/// <param name="nextLine">On success, receives the matched line</param>
+		/// <returns>True if the line was retrieved</returns>
+		public static bool TryGetLine(this ILogCursor cursor, int offset, [NotNullWhen(true)] out string? nextLine)
+		{
+			nextLine = cursor[offset];
+			return nextLine != null;
+		}
+
+		/// <summary>
+		/// Tests whether the current line contains the given text
+		/// </summary>
+		/// <param name="cursor">The log cursor instance</param>
+		/// <param name="text">The text to look for</param>
+		/// <returns>True if the line contains the given string</returns>
+		public static bool Contains(this ILogCursor cursor, string text) => Contains(cursor, 0, text);
+
+		/// <summary>
+		/// Tests whether the current line contains the given text
+		/// </summary>
+		/// <param name="cursor">The log cursor instance</param>
+		/// <param name="offset">Offset of the line to match</param>
+		/// <param name="text">The text to look for</param>
+		/// <returns>True if the line contains the given string</returns>
+		public static bool Contains(this ILogCursor cursor, int offset, string text)
+		{
+			string? line;
+			return cursor.TryGetLine(offset, out line) && line.Contains(text, StringComparison.Ordinal);
+		}
+
+		/// <summary>
+		/// Tests whether the current line starts with the given text
+		/// </summary>
+		/// <param name="cursor">The log cursor instance</param>
+		/// <param name="text">The text to look for</param>
+		/// <returns>True if the line starts with the given string</returns>
+		public static bool StartsWith(this ILogCursor cursor, string text) => StartsWith(cursor, 0, text);
+
+		/// <summary>
+		/// Tests whether the current line starts with the given text
+		/// </summary>
+		/// <param name="cursor">The log cursor instance</param>
+		/// <param name="offset">Offset of the line to match</param>
+		/// <param name="text">The text to look for</param>
+		/// <returns>True if the line starts with the given string</returns>
+		public static bool StartsWith(this ILogCursor cursor, int offset, string text)
+		{
+			string? line;
+			return cursor.TryGetLine(offset, out line) && line.StartsWith(text, StringComparison.Ordinal);
+		}
+
+		/// <summary>
+		/// Determines if the current line matches the given regex
+		/// </summary>
+		/// <param name="cursor">The log cursor instance</param>
+		/// <param name="pattern">The regex pattern to match</param>
+		/// <returns>True if the current line matches the given patter</returns>
+		public static bool IsMatch(this ILogCursor cursor, Regex pattern)
+		{
+			return IsMatch(cursor, 0, pattern);
+		}
+
+		/// <summary>
+		/// Determines if the line at the given offset matches the given regex
+		/// </summary>
+		/// <param name="cursor">The log cursor instance</param>
+		/// <param name="offset">Offset of the line to match</param>
+		/// <param name="pattern">The regex pattern to match</param>
+		/// <returns>True if the requested line matches the given patter</returns>
+		public static bool IsMatch(this ILogCursor cursor, int offset, Regex pattern)
+		{
+			string? line;
+			return cursor.TryGetLine(offset, out line) && pattern.IsMatch(line!);
+		}
+
+		/// <summary>
+		/// Determines if the current line matches the given regex
+		/// </summary>
+		/// <param name="cursor">The log cursor instance</param>
+		/// <param name="pattern">The regex pattern to match</param>
+		/// <param name="outMatch">On success, receives the match result</param>
+		/// <returns>True if the current line matches the given patter</returns>
+		public static bool TryMatch(this ILogCursor cursor, Regex pattern, [NotNullWhen(true)] out Match? outMatch)
+		{
+			return TryMatch(cursor, 0, pattern, out outMatch);
+		}
+
+		/// <summary>
+		/// Determines if the line at the given offset matches the given regex
+		/// </summary>
+		/// <param name="cursor">The log cursor instance</param>
+		/// <param name="offset">The line offset to check</param>
+		/// <param name="pattern">The regex pattern to match</param>
+		/// <param name="outMatch">On success, receives the match result</param>
+		/// <returns>True if the current line matches the given patter</returns>
+		public static bool TryMatch(this ILogCursor cursor, int offset, Regex pattern, [NotNullWhen(true)] out Match? outMatch)
+		{
+			string? line;
+			if (!cursor.TryGetLine(offset, out line))
+			{
+				outMatch = null;
+				return false;
+			}
+
+			Match match = pattern.Match(line);
+			if (!match.Success)
+			{
+				outMatch = null;
+				return false;
+			}
+
+			outMatch = match;
+			return true;
+		}
+
+		/// <summary>
+		/// Matches lines forward from the given offset while the given pattern matches
+		/// </summary>
+		/// <param name="cursor">The log cursor instance</param>
+		/// <param name="offset">Initial offset</param>
+		/// <param name="pattern">Pattern to match</param>
+		/// <returns>Offset of the last line that still matches the pattern (inclusive)</returns>
+		public static int MatchForwards(this ILogCursor cursor, int offset, Regex pattern)
+		{
+			while (IsMatch(cursor, offset + 1, pattern))
+			{
+				offset++;
+			}
+			return offset;
+		}
+
+		/// <summary>
+		/// Matches lines forwards from the given offset until the given pattern matches
+		/// </summary>
+		/// <param name="cursor">The log cursor</param>
+		/// <param name="offset">Initial offset</param>
+		/// <param name="pattern">Pattern to match</param>
+		/// <returns>Offset of the line that matches the pattern (inclusive), or EOF is encountered</returns>
+		public static int MatchForwardsUntil(this ILogCursor cursor, int offset, Regex pattern)
+		{
+			string? nextLine;
+			for (int nextOffset = offset + 1; cursor.TryGetLine(nextOffset, out nextLine); nextOffset++)
+			{
+				if (pattern.IsMatch(nextLine))
+				{
+					return nextOffset;
+				}
+			}
+			return offset;
+		}
+
+		/// <summary>
+		/// Matches lines forwards from the given offset until the given pattern matches
+		/// </summary>
+		/// <param name="cursor">The log cursor</param>
+		/// <param name="offset">Initial offset</param>
+		/// <param name="pattern">Pattern to match</param>
+		/// <param name="maxLines">Max number of lines to look forward</param>
+		/// <returns>Offset of the line that matches the pattern (inclusive), or EOF is encountered</returns>
+		public static int MatchForwardsLimited(this ILogCursor cursor, int offset, Regex pattern, int maxLines)
+		{
+			string? nextLine;
+			int numLines = 0;
+			int nextOffset = offset + 1;
+			cursor.TryGetLine(nextOffset, out nextLine);
+
+			while (numLines < maxLines && nextLine != null)
+			{
+
+				if (pattern.IsMatch(nextLine))
+				{
+					return nextOffset;
+				}
+				cursor.TryGetLine(nextOffset, out nextLine);
+				nextOffset++;
+				numLines++;
+			}
+			return offset;
+		}
+
+		/// <summary>
+		/// Tests if a line consists only of whitespace
+		/// </summary>
+		/// <param name="cursor">The log cursor</param>
+		/// <param name="offset">Offset of the line to check</param>
+		/// <returns></returns>
+		public static bool IsBlank(this ILogCursor cursor, int offset)
+		{
+			if (!cursor.TryGetLine(offset, out string? line))
+			{
+				return false;
+			}
+
+			for (int idx = 0; idx < line.Length; idx++)
+			{
+				if (!Char.IsWhiteSpace(line[idx]))
+				{
+					return false;
+				}
+			}
+
+			return true;
+		}
+
+		/// <summary>
+		/// Check if a line at a given offset is left-aligned with at least this indent of another line
+		/// </summary>
+		/// <param name="cursor">The log cursor</param>
+		/// <param name="offset">Offset of the line to check</param>
+		/// <param name="firstLine">The first line to compare with</param>
+		/// <returns>True if the line at the given offset is aligned at least as much as the current lien</returns>
+		public static bool IsAligned(this ILogCursor cursor, int offset, string? firstLine)
+		{
+			if (firstLine == null || !cursor.TryGetLine(offset, out string? line))
+			{
+				return false;
+			}
+
+			for (int idx = 0; idx < line.Length; idx++)
+			{
+				if (idx == firstLine.Length || !Char.IsWhiteSpace(firstLine[idx]))
+				{
+					return true;
+				}
+				if (line[idx] != firstLine[idx])
+				{
+					return false;
+				}
+			}
+
+			return true;
+		}
+
+		/// <summary>
+		/// Check if a line at a given offset is left-aligned with at least this indent of another line
+		/// </summary>
+		/// <param name="cursor">The log cursor</param>
+		/// <param name="offset">Offset of the line to check</param>
+		/// <param name="firstLine">The first line to compare with</param>
+		/// <returns>True if the line at the given offset is aligned at least as much as the current lien</returns>
+		public static bool IsHanging(this ILogCursor cursor, int offset, string? firstLine)
+		{
+			if (firstLine == null || !cursor.TryGetLine(offset, out string? line))
+			{
+				return false;
+			}
+
+			for (int idx = 0; idx < line.Length; idx++)
+			{
+				if (idx == firstLine.Length || !Char.IsWhiteSpace(firstLine[idx]))
+				{
+					return Char.IsWhiteSpace(line[idx]);
+				}
+				if (line[idx] != firstLine[idx])
+				{
+					return false;
+				}
+			}
+
+			return true;
+		}
+	}
+}
diff --git a/Engine/Source/Programs/Shared/EpicGames.Core/Logging/ILoggerIndent.cs b/Engine/Source/Programs/Shared/EpicGames.Core/Logging/ILoggerIndent.cs
new file mode 100644
--- /dev/null
+++ b/Engine/Source/Programs/Shared/EpicGames.Core/Logging/ILoggerIndent.cs
@@ -0,0 +1,242 @@
+// Copyright Epic Games, Inc. All Rights Reserved.
+
+using System;
+using System.Collections;
+using System.Collections.Generic;
+using System.Linq;
+using Microsoft.Extensions.Logging;
+
+namespace EpicGames.Core
+{
+	/// <summary>
+	/// Class to apply a log indent for the lifetime of an object 
+	/// </summary>
+	public interface ILoggerIndent
+	{
+		/// <summary>
+		/// The indent to apply
+		/// </summary>
+		string Indent { get; }
+	}
+
+	/// <summary>
+	/// Wrapper class for ILogger classes which supports LoggerStatusScope
+	/// </summary>
+	public class DefaultLoggerIndentHandler : ILogger
+	{
+		/// <summary>
+		/// Scoped indent message
+		/// </summary>
+		class Scope : IDisposable
+		{
+			/// <summary>
+			/// Owning object
+			/// </summary>
+			readonly DefaultLoggerIndentHandler _owner;
+
+			/// <summary>
+			/// The indent scope object
+			/// </summary>
+			public ILoggerIndent Indent { get; }
+
+			/// <summary>
+			/// Constructor
+			/// </summary>
+			public Scope(DefaultLoggerIndentHandler owner, ILoggerIndent indent)
+			{
+				_owner = owner;
+				Indent = indent;
+
+				lock (owner._scopes)
+				{
+					owner._scopes.Add(this);
+				}
+			}
+
+			/// <summary>
+			/// Remove this indent from the list
+			/// </summary>
+			public void Dispose()
+			{
+				lock (_owner._scopes)
+				{
+					_owner._scopes.Remove(this);
+				}
+			}
+		}
+
+		/// <summary>
+		/// Struct to wrap a formatted set of log values with applied indent
+		/// </summary>
+		/// <typeparam name="TState">Arbitrary type parameter</typeparam>
+		readonly struct FormattedLogValues<TState> : IEnumerable<KeyValuePair<string, object>>
+		{
+			/// <summary>
+			/// The indent to apply
+			/// </summary>
+			readonly string _indent;
+
+			/// <summary>
+			/// The inner state
+			/// </summary>
+			readonly TState _state;
+
+			/// <summary>
+			/// Formatter for the inner state
+			/// </summary>
+			readonly Func<TState, Exception?, string> _formatter;
+
+			/// <summary>
+			/// Constructor
+			/// </summary>
+			/// <param name="indent">The indent to apply</param>
+			/// <param name="state">The inner state</param>
+			/// <param name="formatter">Formatter for the inner state</param>
+			public FormattedLogValues(string indent, TState state, Func<TState, Exception?, string> formatter)
+			{
+				_indent = indent;
+				_state = state;
+				_formatter = formatter;
+			}
+
+			/// <inheritdoc/>
+			public IEnumerator<KeyValuePair<string, object>> GetEnumerator()
+			{
+				IEnumerable<KeyValuePair<string, object>>? innerEnumerable = _state as IEnumerable<KeyValuePair<string, object>>;
+				if (innerEnumerable != null)
+				{
+					foreach (KeyValuePair<string, object> pair in innerEnumerable)
+					{
+						if (pair.Key.Equals("{OriginalFormat}", StringComparison.Ordinal))
+						{
+							yield return new KeyValuePair<string, object>(pair.Key, _indent + pair.Value.ToString());
+						}
+						else
+						{
+							yield return pair;
+						}
+					}
+				}
+			}
+
+			/// <inheritdoc/>
+			IEnumerator IEnumerable.GetEnumerator()
+			{
+				throw new NotImplementedException();
+			}
+
+			/// <summary>
+			/// Formats an instance of this object
+			/// </summary>
+			/// <param name="values">The object instance</param>
+			/// <param name="exception">The exception to format</param>
+			/// <returns>The formatted string</returns>
+			public static string Format(FormattedLogValues<TState> values, Exception? exception)
+			{
+				return values._indent + values._formatter(values._state, exception);
+			}
+		}
+
+		/// <summary>
+		/// The internal logger
+		/// </summary>
+		readonly ILogger _inner;
+
+		/// <summary>
+		/// Current list of indents
+		/// </summary>
+		readonly List<Scope> _scopes = [];
+
+		/// <summary>
+		/// The inner logger
+		/// </summary>
+		public ILogger Inner => _inner;
+
+		/// <summary>
+		/// The current indent text
+		/// </summary>
+		public string Indent
+		{
+			get;
+			private set;
+		}
+
+		/// <summary>
+		/// Constructor
+		/// </summary>
+		/// <param name="inner">The logger to wrap</param>
+		public DefaultLoggerIndentHandler(ILogger inner)
+		{
+			_inner = inner;
+			Indent = "";
+		}
+
+		/// <inheritdoc/>
+		public IDisposable? BeginScope<TState>(TState state) where TState : notnull
+		{
+			ILoggerIndent? indent = state as ILoggerIndent;
+			if (indent != null)
+			{
+				return new Scope(this, indent);
+			}
+
+			return _inner.BeginScope(state);
+		}
+
+		/// <inheritdoc/>
+		public bool IsEnabled(LogLevel logLevel)
+		{
+			return _inner.IsEnabled(logLevel);
+		}
+
+		/// <inheritdoc/>
+		public void Log<TState>(LogLevel logLevel, EventId eventId, TState state, Exception? exception, Func<TState, Exception?, string> formatter)
+		{
+			if (_scopes.Count > 0)
+			{
+				string indent = String.Join("", _scopes.Select(x => x.Indent.Indent));
+				_inner.Log(logLevel, eventId, new FormattedLogValues<TState>(indent, state, formatter), exception, FormattedLogValues<TState>.Format);
+				return;
+			}
+
+			_inner.Log(logLevel, eventId, state, exception, formatter);
+		}
+	}
+
+	/// <summary>
+	/// Extension methods for creating an indent
+	/// </summary>
+	public static class LoggerIndentExtensions
+	{
+		/// <summary>
+		/// Class to apply a log indent for the lifetime of an object 
+		/// </summary>
+		class LoggerIndent : ILoggerIndent
+		{
+			/// <summary>
+			/// The previous indent
+			/// </summary>
+			public string Indent { get; }
+
+			/// <summary>
+			/// Constructor
+			/// </summary>
+			/// <param name="indent">Indent to append to the existing indent</param>
+			public LoggerIndent(string indent)
+			{
+				Indent = indent;
+			}
+		}
+
+		/// <summary>
+		/// Create an indent
+		/// </summary>
+		/// <param name="logger">Logger interface</param>
+		/// <param name="indent">The indent to apply</param>
+		/// <returns>Disposable object</returns>
+		public static IDisposable? BeginIndentScope(this ILogger logger, string indent)
+		{
+			return logger.BeginScope(new LoggerIndent(indent));
+		}
+	}
+}
diff --git a/Engine/Source/Programs/Shared/EpicGames.Core/Logging/ILoggerProgress.cs b/Engine/Source/Programs/Shared/EpicGames.Core/Logging/ILoggerProgress.cs
new file mode 100644
--- /dev/null
+++ b/Engine/Source/Programs/Shared/EpicGames.Core/Logging/ILoggerProgress.cs
@@ -0,0 +1,139 @@
+// Copyright Epic Games, Inc. All Rights Reserved.
+
+using System;
+using System.Diagnostics;
+using Microsoft.Extensions.Logging;
+
+namespace EpicGames.Core
+{
+	/// <summary>
+	/// Manages a status message for a long running operation, which can be updated with progress. Typically transient on consoles, and written to the same line.
+	/// </summary>
+	public interface ILoggerProgress : IDisposable
+	{
+		/// <summary>
+		/// Prefix message for the status
+		/// </summary>
+		string Message
+		{
+			get;
+		}
+
+		/// <summary>
+		/// The current 
+		/// </summary>
+		string Progress
+		{
+			get; set;
+		}
+	}
+
+	/// <summary>
+	/// Extension methods for status objects
+	/// </summary>
+	public static class LoggerProgressExtensions
+	{
+		/// <summary>
+		/// Concrete implementation of <see cref="ILoggerProgress"/>
+		/// </summary>
+		class LoggerProgress : ILoggerProgress
+		{
+			/// <summary>
+			/// The logger to output to
+			/// </summary>
+			private readonly ILogger _logger;
+
+			/// <summary>
+			/// Prefix message for the status
+			/// </summary>
+			public string Message
+			{
+				get;
+			}
+
+			/// <summary>
+			/// The current 
+			/// </summary>
+			public string Progress
+			{
+				get => _progressInternal;
+				set
+				{
+					_progressInternal = value;
+
+					if (_timer.Elapsed > TimeSpan.FromSeconds(3.0))
+					{
+						_lastOutput = String.Empty;
+						Flush();
+						_timer.Restart();
+					}
+				}
+			}
+
+			/// <summary>
+			/// The last string that was output
+			/// </summary>
+			string _lastOutput = String.Empty;
+
+			/// <summary>
+			/// Backing storage for the Progress string
+			/// </summary>
+			string _progressInternal = String.Empty;
+
+			/// <summary>
+			/// Timer since the last update
+			/// </summary>
+			readonly Stopwatch _timer = Stopwatch.StartNew();
+
+			/// <summary>
+			/// Constructor
+			/// </summary>
+			/// <param name="logger">The logger to write to</param>
+			/// <param name="message">The base message to display</param>
+			public LoggerProgress(ILogger logger, string message)
+			{
+				_logger = logger;
+				Message = message;
+				_progressInternal = String.Empty;
+
+				logger.LogInformation("{Progress}", message);
+			}
+
+			/// <summary>
+			/// Dispose of this object
+			/// </summary>
+			public void Dispose()
+			{
+				Flush();
+			}
+
+			/// <summary>
+			/// Flushes the current output to the log
+			/// </summary>
+			void Flush()
+			{
+				string output = Message;
+				if (!String.IsNullOrEmpty(Progress))
+				{
+					output += $" {Progress}";
+				}
+				if (!String.Equals(output, _lastOutput, StringComparison.Ordinal))
+				{
+					_logger.LogInformation("{Progress}", output);
+					_lastOutput = output;
+				}
+			}
+		}
+
+		/// <summary>
+		/// Begins a new progress scope
+		/// </summary>
+		/// <param name="logger">The logger being written to</param>
+		/// <param name="message">The message prefix</param>
+		/// <returns>Scope object, which should be disposed when finished</returns>
+		public static ILoggerProgress BeginProgressScope(this ILogger logger, string message)
+		{
+			return new LoggerProgress(logger, message);
+		}
+	}
+}
diff --git a/Engine/Source/Programs/Shared/EpicGames.Core/Logging/ILoggerProperty.cs b/Engine/Source/Programs/Shared/EpicGames.Core/Logging/ILoggerProperty.cs
new file mode 100644
--- /dev/null
+++ b/Engine/Source/Programs/Shared/EpicGames.Core/Logging/ILoggerProperty.cs
@@ -0,0 +1,79 @@
+// Copyright Epic Games, Inc. All Rights Reserved.
+
+using System;
+using System.Collections.Generic;
+using Microsoft.Extensions.Logging;
+
+namespace EpicGames.Core
+{
+	/// <summary>
+	/// Extension methods for more convenient attaching of properties to log messages
+	/// </summary>
+	public static class LoggerPropertyExtensions
+	{
+		/// <summary>
+		/// Builder struct for attaching properties to log messages
+		/// Wraps and disposes log scope.
+		/// </summary>
+		public readonly struct LogScopedPropertyBuilder
+		{
+			private readonly Dictionary<string, object> _properties;
+			private readonly ILogger _logger;
+
+			/// <summary>
+			/// Constructor
+			/// </summary>
+			/// <param name="logger">Logger to wrap</param>
+			public LogScopedPropertyBuilder(ILogger logger)
+			{
+				_properties = [];
+				_logger = logger;
+			}
+
+			/// <summary>
+			/// Set a key/value property for current log scope
+			/// </summary>
+			/// <param name="key">Arbitrary key</param>
+			/// <param name="value">Arbitrary value</param>
+			/// <returns></returns>
+			public LogScopedPropertyBuilder WithProperty(string key, object value)
+			{
+				_properties[key] = value;
+				return this;
+			}
+
+			/// <summary>
+			/// Starts new scope on underlying logger, passing previously set properties
+			/// </summary>
+			/// <returns>Original IDisposable from ILogger.BeginScope()</returns>
+			public IDisposable? BeginScope()
+			{
+				return _logger.BeginScope(_properties);
+			}
+		}
+
+		/// <summary>
+		/// Creates a new log scope and attaches a key/value property
+		/// Additional properties can be added with <see cref="LogScopedPropertyBuilder.WithProperty" />.
+		/// </summary>
+		/// <param name="logger">Logger to use</param>
+		/// <param name="key">Arbitrary key</param>
+		/// <param name="value">Arbitrary value</param>
+		/// <returns>A disposable log property builder that wraps the original ILogger scope</returns>
+		/// <example>
+		/// Simple example of attaching two properties
+		/// <code>
+		/// using (_logger.WithProperty("stream", "ue5").WithProperty("os", "windows").BuildScope())
+		/// {
+		///     _logger.LogInformation("A log message with properties attached");
+		/// }
+		/// </code>
+		/// </example>
+		public static LogScopedPropertyBuilder WithProperty(this ILogger logger, string key, object value)
+		{
+			LogScopedPropertyBuilder builder = new(logger);
+			builder.WithProperty(key, value);
+			return builder;
+		}
+	}
+}
diff --git a/Engine/Source/Programs/Shared/EpicGames.Core/Logging/JsonLogEvent.cs b/Engine/Source/Programs/Shared/EpicGames.Core/Logging/JsonLogEvent.cs
new file mode 100644
--- /dev/null
+++ b/Engine/Source/Programs/Shared/EpicGames.Core/Logging/JsonLogEvent.cs
@@ -0,0 +1,429 @@
+// Copyright Epic Games, Inc. All Rights Reserved.
+
+using System;
+using System.Collections;
+using System.Collections.Generic;
+using System.Text;
+using System.Text.Json;
+using Microsoft.Extensions.Logging;
+
+namespace EpicGames.Core
+{
+	/// <summary>
+	/// Defines a preformatted Json log event, which can pass through raw Json data directly or format it as a regular string
+	/// </summary>
+	public readonly struct JsonLogEvent : IEnumerable<KeyValuePair<string, object?>>
+	{
+		/// <summary>
+		/// The log level
+		/// </summary>
+		public LogLevel Level { get; }
+
+		/// <summary>
+		/// The event id, if set
+		/// </summary>
+		public EventId EventId { get; }
+
+		/// <summary>
+		/// Index of this line
+		/// </summary>
+		public int LineIndex { get; }
+
+		/// <summary>
+		/// Number of lines in a multi-line message
+		/// </summary>
+		public int LineCount { get; }
+
+		/// <summary>
+		/// The utf-8 encoded JSON event
+		/// </summary>
+		public ReadOnlyMemory<byte> Data { get; }
+
+		/// <summary>
+		/// Constructor
+		/// </summary>
+		public JsonLogEvent(LogEvent logEvent)
+			: this(logEvent.Level, logEvent.Id, logEvent.LineIndex, logEvent.LineCount, logEvent.ToJsonBytes())
+		{
+		}
+
+		/// <summary>
+		/// Constructor
+		/// </summary>
+		public JsonLogEvent(LogLevel level, EventId eventId, int lineIndex, int lineCount, ReadOnlyMemory<byte> data)
+		{
+			Level = level;
+			EventId = eventId;
+			LineIndex = lineIndex;
+			LineCount = lineCount;
+			Data = data;
+		}
+
+		/// <summary>
+		/// Creates a json log event from the given logger parameters
+		/// </summary>
+		/// <inheritdoc cref="ILogger.Log{TState}(LogLevel, EventId, TState, Exception, Func{TState, Exception, String})"/>
+		public static JsonLogEvent FromLoggerState<TState>(LogLevel logLevel, EventId eventId, TState state, Exception? exception, Func<TState, Exception?, string> formatter)
+		{
+			if (state is JsonLogEvent jsonLogEvent)
+			{
+				return jsonLogEvent;
+			}
+
+			LogEvent? logEvent = state as LogEvent ?? LogEvent.FromState(logLevel, eventId, state, exception, formatter);
+			return new JsonLogEvent(logLevel, eventId, 0, 1, logEvent.ToJsonBytes());
+		}
+
+		/// <summary>
+		/// Parse an event from the given data
+		/// </summary>
+		/// <param name="data"></param>
+		/// <returns></returns>
+		public static JsonLogEvent Parse(ReadOnlyMemory<byte> data)
+		{
+			JsonLogEvent logEvent;
+			if (!TryParseInternal(data, out logEvent))
+			{
+				throw new InvalidOperationException("Cannot parse string");
+			}
+			return logEvent;
+		}
+
+		/// <summary>
+		/// Tries to parse a Json log event from the given string
+		/// </summary>
+		/// <param name="data"></param>
+		/// <param name="logEvent"></param>
+		/// <returns></returns>
+		public static bool TryParse(ReadOnlyMemory<byte> data, out JsonLogEvent logEvent)
+		{
+			try
+			{
+				return TryParseInternal(data, out logEvent);
+			}
+			catch
+			{
+				logEvent = default;
+				return false;
+			}
+		}
+
+		/// <summary>
+		/// Tries to parse a Json log event from the given string
+		/// </summary>
+		/// <param name="text">Text to parse</param>
+		/// <param name="logEvent"></param>
+		/// <returns></returns>
+		public static bool TryParse(string text, out JsonLogEvent logEvent)
+		{
+			byte[] data = Encoding.UTF8.GetBytes(text);
+			return TryParse(data, out logEvent);
+		}
+
+		static bool TryParseInternal(ReadOnlyMemory<byte> data, out JsonLogEvent logEvent)
+		{
+			LogLevel level = LogLevel.None;
+			int eventId = 0;
+			int lineIndex = 0;
+			int lineCount = 1;
+
+			Utf8JsonReader reader = new(data.Span);
+			if (reader.Read() && reader.TokenType == JsonTokenType.StartObject)
+			{
+				while (reader.Read() && reader.TokenType == JsonTokenType.PropertyName)
+				{
+					ReadOnlySpan<byte> propertyName = reader.ValueSpan;
+
+					if (!reader.Read())
+					{
+						break;
+					}
+					else if (propertyName.SequenceEqual(LogEventPropertyName.Level) && reader.TokenType == JsonTokenType.String)
+					{
+						level = ParseLevel(reader.ValueSpan);
+					}
+					else if (propertyName.SequenceEqual(LogEventPropertyName.Id) && reader.TokenType == JsonTokenType.Number)
+					{
+						eventId = reader.GetInt32();
+					}
+					else if (propertyName.SequenceEqual(LogEventPropertyName.Line) && reader.TokenType == JsonTokenType.Number)
+					{
+						reader.TryGetInt32(out lineIndex);
+					}
+					else if (propertyName.SequenceEqual(LogEventPropertyName.LineCount) && reader.TokenType == JsonTokenType.Number)
+					{
+						reader.TryGetInt32(out lineCount);
+					}
+					else if (propertyName.SequenceEqual(LogEventPropertyName.Properties) && reader.TokenType == JsonTokenType.StartObject)
+					{
+						while (reader.Read() && reader.TokenType == JsonTokenType.PropertyName)
+						{
+							ReadOnlySpan<byte> innerPropertyName = reader.ValueSpan;
+
+							if (!reader.Read())
+							{
+								break;
+							}
+
+							// In the case of UE_LOGFMT the event id can be nested within the properties instead
+							if (innerPropertyName.SequenceEqual(LogEventPropertyName.Id) && reader.TokenType == JsonTokenType.Number)
+							{
+								eventId = reader.GetInt32();
+							}
+							reader.Skip();
+						}
+					}
+					reader.Skip();
+				}
+			}
+
+			if (reader.TokenType == JsonTokenType.EndObject && level != LogLevel.None && reader.BytesConsumed == data.Length)
+			{
+				logEvent = new JsonLogEvent(level, new EventId(eventId), lineIndex, lineCount, data.ToArray());
+				return true;
+			}
+			else
+			{
+				logEvent = default;
+				return false;
+			}
+		}
+
+		static readonly sbyte[] s_firstCharToLogLevel;
+		static readonly byte[][] s_logLevelNames;
+
+#pragma warning disable CA2207 // Initialize value type static fields inline
+		static JsonLogEvent()
+#pragma warning restore CA2207 // Initialize value type static fields inline
+		{
+			const int LogLevelCount = (int)LogLevel.None;
+
+			s_firstCharToLogLevel = new sbyte[256];
+			Array.Fill(s_firstCharToLogLevel, (sbyte)-1);
+
+			s_logLevelNames = new byte[LogLevelCount][];
+			for (int idx = 0; idx < (int)LogLevel.None; idx++)
+			{
+				byte[] name = Encoding.UTF8.GetBytes(Enum.GetName(typeof(LogLevel), (LogLevel)idx)!);
+				s_logLevelNames[idx] = name;
+				s_firstCharToLogLevel[name[0]] = (sbyte)idx;
+			}
+		}
+
+		static LogLevel ParseLevel(ReadOnlySpan<byte> level)
+		{
+			int result = s_firstCharToLogLevel[level[0]];
+			if (!level.SequenceEqual(s_logLevelNames[result]))
+			{
+				throw new InvalidOperationException();
+			}
+			return (LogLevel)result;
+		}
+
+		static readonly Utf8String s_newlineEscaped = new("\\n");
+
+		/// <summary>
+		/// Gets the rendered message from the event data
+		/// </summary>
+		public Utf8String GetRenderedMessage()
+		{
+			Utf8JsonReader reader = new(Data.Span);
+			if (reader.Read() && reader.TokenType == JsonTokenType.StartObject)
+			{
+				while (reader.Read() && reader.TokenType == JsonTokenType.PropertyName)
+				{
+					ReadOnlySpan<byte> propertyName = reader.ValueSpan;
+					if (!reader.Read())
+					{
+						break;
+					}
+					else if (propertyName.SequenceEqual(LogEventPropertyName.Message) && reader.TokenType == JsonTokenType.String)
+					{
+						return new Utf8String(reader.GetUtf8String().ToArray());
+					}
+				}
+			}
+			return Utf8String.Empty;
+		}
+
+		/// <summary>
+		/// Gets the event data rendered as a legacy unreal log line of the format:
+		/// [timestamp][frame number]LogChannel: LogVerbosity: Message
+		/// </summary>
+		public string GetLegacyLogLine()
+		{
+			Utf8JsonReader reader = new(Data.Span);
+			if (reader.Read() && reader.TokenType == JsonTokenType.StartObject)
+			{
+				Utf8String? message = null;
+				ulong frame = 0;
+				DateTime? time = null;
+				while (reader.Read() && reader.TokenType == JsonTokenType.PropertyName)
+				{
+					ReadOnlySpan<byte> propertyName = reader.ValueSpan;
+					if (!reader.Read())
+					{
+						break;
+					}
+					else if (propertyName.SequenceEqual(LogEventPropertyName.Time) && reader.TokenType == JsonTokenType.String)
+					{
+						time = reader.GetDateTime();
+					}
+					else if (propertyName.SequenceEqual(LogEventPropertyName.Frame) && reader.TokenType == JsonTokenType.Number)
+					{
+						frame = reader.GetUInt64();
+					}
+					else if (propertyName.SequenceEqual(LogEventPropertyName.Message) && reader.TokenType == JsonTokenType.String)
+					{
+						message = new Utf8String(reader.GetUtf8String().ToArray());
+					}
+				}
+
+				if (message is not null)
+				{
+					if (time is not null)
+					{
+						return $"[{time:yyyy.MM.dd-HH.mm.ss:fff}][{frame % 1000,3}]{message}";
+					}
+					else
+					{
+						return $"{message}";
+					}
+				}
+			}
+			return String.Empty;
+		}
+
+		/// <summary>
+		/// Count the number of lines in the message field of a log event
+		/// </summary>
+		/// <returns>Number of lines in the message</returns>
+		public int GetMessageLineCount()
+		{
+			if (Data.Span.IndexOf(s_newlineEscaped.Span) != -1)
+			{
+				Utf8JsonReader reader = new(Data.Span);
+				if (reader.Read() && reader.TokenType == JsonTokenType.StartObject)
+				{
+					while (reader.Read() && reader.TokenType == JsonTokenType.PropertyName)
+					{
+						ReadOnlySpan<byte> propertyName = reader.ValueSpan;
+						if (!reader.Read())
+						{
+							break;
+						}
+						else if (propertyName.SequenceEqual(LogEventPropertyName.Message) && reader.TokenType == JsonTokenType.String)
+						{
+							return CountLines(reader.ValueSpan);
+						}
+					}
+				}
+			}
+			return 1;
+		}
+
+		/// <summary>
+		/// Counts the number of newlines in an escaped JSON string
+		/// </summary>
+		/// <param name="str">The escaped string</param>
+		/// <returns></returns>
+		static int CountLines(ReadOnlySpan<byte> str)
+		{
+			int lines = 1;
+			for (int idx = 0; idx < str.Length - 1; idx++)
+			{
+				if (str[idx] == '\\')
+				{
+					if (str[idx + 1] == 'n')
+					{
+						lines++;
+					}
+					idx++;
+				}
+			}
+			return lines;
+		}
+
+		/// <summary>
+		/// Formats an event as a string
+		/// </summary>
+		/// <param name="state"></param>
+		/// <param name="ex"></param>
+		/// <returns></returns>
+		public static string Format(JsonLogEvent state, Exception? ex)
+		{
+			_ = ex;
+			return LogEvent.Read(state.Data.Span).ToString();
+		}
+
+		/// <summary>
+		/// Find all properties of the given type in a particular log line
+		/// </summary>
+		/// <param name="type">Type of property to return</param>
+		/// <returns></returns>
+		public IEnumerable<JsonProperty> FindPropertiesOfType(Utf8String type)
+		{
+			JsonDocument document = JsonDocument.Parse(Data);
+			return FindPropertiesOfType(document.RootElement, type);
+		}
+
+		/// <summary>
+		/// Find all properties of the given type in a particular log line
+		/// </summary>
+		/// <param name="line">Line data</param>
+		/// <param name="type">Type of property to return</param>
+		/// <returns></returns>
+		public static IEnumerable<JsonProperty> FindPropertiesOfType(JsonElement line, Utf8String type)
+		{
+			JsonElement properties;
+			if (line.TryGetProperty("properties", out properties) && properties.ValueKind == JsonValueKind.Object)
+			{
+				foreach (JsonProperty property in properties.EnumerateObject())
+				{
+					if (property.Value.ValueKind == JsonValueKind.Object)
+					{
+						foreach (JsonProperty subProperty in property.Value.EnumerateObject())
+						{
+							if (subProperty.NameEquals(LogEventPropertyName.Type.Span))
+							{
+								if (subProperty.Value.ValueKind == JsonValueKind.String && subProperty.Value.ValueEquals(type.Span))
+								{
+									yield return property;
+								}
+								else
+								{
+									break;
+								}
+							}
+						}
+					}
+				}
+			}
+		}
+
+		/// <inheritdoc/>
+		public override string ToString() => Encoding.UTF8.GetString(Data.ToArray());
+
+		/// <inheritdoc/>
+		public IEnumerator<KeyValuePair<string, object?>> GetEnumerator() => LogEvent.Read(Data.Span).GetEnumerator();
+
+		/// <inheritdoc/>
+		IEnumerator IEnumerable.GetEnumerator() => GetEnumerator();
+	}
+
+	/// <summary>
+	/// Extension methods for <see cref="JsonLogEvent"/>
+	/// </summary>
+	public static class JsonLogEventExtensions
+	{
+		/// <summary>
+		/// Logs a <see cref="JsonLogEvent"/> to the given logger
+		/// </summary>
+		/// <param name="logger">Logger to write to</param>
+		/// <param name="jsonLogEvent">Json log event to write</param>
+		public static void LogJsonLogEvent(this ILogger logger, JsonLogEvent jsonLogEvent)
+		{
+			logger.Log(jsonLogEvent.Level, jsonLogEvent.EventId, jsonLogEvent, null, JsonLogEvent.Format);
+		}
+	}
+}
diff --git a/Engine/Source/Programs/Shared/EpicGames.Core/Logging/KnownLogEvents.cs b/Engine/Source/Programs/Shared/EpicGames.Core/Logging/KnownLogEvents.cs
new file mode 100644
--- /dev/null
+++ b/Engine/Source/Programs/Shared/EpicGames.Core/Logging/KnownLogEvents.cs
@@ -0,0 +1,434 @@
+// Copyright Epic Games, Inc. All Rights Reserved.
+
+using Microsoft.Extensions.Logging;
+
+namespace EpicGames.Core
+{
+	/// <summary>
+	/// Well known log events
+	/// </summary>
+	[System.Diagnostics.CodeAnalysis.SuppressMessage("Naming", "CA1707:Identifiers should not contain underscores", Justification = "Intentional")]
+	public static class KnownLogEvents
+	{
+		/// <summary>
+		/// Unset
+		/// </summary>
+		public static EventId None { get; } = new EventId(0);
+
+		/// <summary>
+		/// Generic error
+		/// </summary>
+		public static EventId Generic { get; } = new EventId(1);
+
+		/// <summary>
+		/// Generic exception
+		/// </summary>
+		public static EventId Exception { get; } = new EventId(2);
+
+		/// <summary>
+		/// Exit code error
+		/// </summary>
+		public static EventId ExitCode { get; } = new EventId(3);
+
+		/// <summary>
+		/// Compiler error
+		/// </summary>
+		public static EventId Compiler { get; } = new EventId(100);
+
+		/// <summary>
+		/// Compiler summary line (eg. 'XYZ failed to compile')
+		/// </summary>
+		public static EventId Compiler_Summary { get; } = new EventId(101);
+
+		/// <summary>
+		/// TargetRules or ModuleRules error
+		/// </summary>
+		public static EventId RulesAssembly { get; } = new EventId(102);
+
+		/// <summary>
+		/// Validation messages
+		/// </summary>
+		public static EventId Validation { get; } = new EventId(103);
+
+		/// <summary>
+		/// Circular dependency messages
+		/// </summary>
+		public static EventId CircularDependency { get; } = new EventId(104);
+
+		/// <summary>
+		/// UHT compiler message
+		/// </summary>
+		public static EventId UHT { get; } = new EventId(110);
+
+		/// <summary>
+		/// UBA message
+		/// </summary>
+		public static EventId UBA { get; } = new EventId(111);
+
+		/// <summary>
+		/// Linker error
+		/// </summary>
+		public static EventId Linker { get; } = new EventId(200);
+
+		/// <summary>
+		/// Linker: Undefined symbol
+		/// </summary>
+		public static EventId Linker_UndefinedSymbol { get; } = new EventId(201);
+
+		/// <summary>
+		/// Linker: Multiply defined symbol
+		/// </summary>
+		public static EventId Linker_DuplicateSymbol { get; } = new EventId(202);
+
+		/// <summary>
+		/// Engine error
+		/// </summary>
+		public static EventId Engine { get; } = new EventId(300);
+
+		/// <summary>
+		/// Engine log channel
+		/// </summary>
+		public static EventId Engine_LogChannel { get; } = new EventId(301);
+
+		/// <summary>
+		/// Engine: Crash dump
+		/// </summary>
+		public static EventId Engine_Crash { get; } = new EventId(302);
+
+		/// <summary>
+		/// Engine: UE_ASSET_LOG output
+		/// </summary>
+		public static EventId Engine_AssetLog { get; } = new EventId(303);
+
+		/// <summary>
+		/// Engine: Localization commandlet output
+		/// </summary>
+		public static EventId Engine_Localization { get; } = new EventId(304);
+
+		/// <summary>
+		/// Engine: appErrorf called
+		/// </summary>
+		public static EventId Engine_AppError { get; } = new EventId(305);
+
+		/// <summary>
+		/// Engine: Assertion failed
+		/// </summary>
+		public static EventId Engine_AssertionFailed { get; } = new EventId(306);
+
+		/// <summary>
+		/// Engine: Shader compiler output
+		/// </summary>
+		public static EventId Engine_ShaderCompiler { get; } = new EventId(310);
+
+		/// <summary>
+		/// UAT error
+		/// </summary>
+		public static EventId AutomationTool { get; } = new EventId(400);
+
+		/// <summary>
+		/// UAT: Crash dump
+		/// </summary>
+		public static EventId AutomationTool_Crash { get; } = new EventId(401);
+
+		/// <summary>
+		/// UAT: Exit code indicating a crash
+		/// </summary>
+		public static EventId AutomationTool_CrashExitCode { get; } = new EventId(402);
+
+		/// <summary>
+		/// UAT: Source file with line number
+		/// </summary>
+		public static EventId AutomationTool_SourceFileLine { get; } = new EventId(403);
+
+		/// <summary>
+		/// UAT: Missing copyright notice
+		/// </summary>
+		public static EventId AutomationTool_MissingCopyright { get; } = new EventId(404);
+
+		/// <summary>
+		/// UAT: Mismatched Perforce case
+		/// </summary>
+		public static EventId AutomationTool_PerforceCase { get; } = new EventId(405);
+
+		/// <summary>
+		/// UAT: Unacceptable words
+		/// </summary>
+		public static EventId AutomationTool_UnacceptableWords { get; } = new EventId(406);
+
+		/// <summary>
+		/// UAT: BuildGraph script parsing error
+		/// </summary>
+		public static EventId AutomationTool_BuildGraphScript { get; } = new EventId(407);
+
+		/// <summary>
+		/// UAT: CodeSign message
+		/// </summary>
+		public static EventId AutomationTool_CodeSign { get; } = new EventId(408);
+
+		/// <summary>
+		/// MSBuild: Generic error
+		/// </summary>
+		public static EventId MSBuild { get; } = new EventId(500);
+
+		/// <summary>
+		/// Microsoft: Generic MSTest error
+		/// </summary>
+		public static EventId MSTest { get; } = new EventId(540);
+
+		/// <summary>
+		/// Microsoft: Generic Visual Studio error
+		/// </summary>
+		public static EventId Microsoft { get; } = new EventId(550);
+
+		/// <summary>
+		/// Error message from Gauntlet
+		/// </summary>
+		public static EventId Gauntlet { get; } = new EventId(600);
+
+		/// <summary>
+		/// Error message from Gauntlet test events
+		/// </summary>
+		public static EventId Gauntlet_TestEvent { get; } = new EventId(601);
+
+		/// <summary>
+		/// Error message from Gauntlet device events
+		/// </summary>
+		public static EventId Gauntlet_DeviceEvent { get; } = new EventId(602);
+
+		/// <summary>
+		/// Error message from Gauntlet Unreal Engine test events
+		/// </summary>
+		public static EventId Gauntlet_UnrealEngineTestEvent { get; } = new EventId(603);
+
+		/// <summary>
+		/// Error message from Gauntlet build drop events
+		/// </summary>
+		public static EventId Gauntlet_BuildDropEvent { get; } = new EventId(604);
+
+		/// <summary>
+		/// Fatal Error message from Gauntlet events
+		/// </summary>
+		public static EventId Gauntlet_FatalEvent { get; } = new EventId(605);
+
+		/// <summary>
+		/// A systemic event, relating to the health of the farm
+		/// </summary>
+		public static EventId Systemic { get; } = new EventId(700);
+
+		/// <summary>
+		/// A systemic event relating to Perforce
+		/// </summary>
+		public static EventId Systemic_Perforce { get; } = new EventId(701);
+
+		/// <summary>
+		/// A systemic event from XGE
+		/// </summary>
+		public static EventId Systemic_Xge { get; } = new EventId(710);
+
+		/// <summary>
+		/// Builds will run in standalone mode
+		/// </summary>
+		public static EventId Systemic_Xge_Standalone { get; } = new EventId(711);
+
+		/// <summary>
+		/// BuildService.exe is not running
+		/// </summary>
+		public static EventId Systemic_Xge_ServiceNotRunning { get; } = new EventId(712);
+
+		/// <summary>
+		/// General build failed error
+		/// </summary>
+		public static EventId Systemic_Xge_BuildFailed { get; } = new EventId(713);
+
+		/// <summary>
+		/// Cache size reached
+		/// </summary>
+		public static EventId Systemic_Xge_CacheLimit { get; } = new EventId(714);
+
+		/// <summary>
+		/// Current logging level may impact performance
+		/// </summary>
+		public static EventId Systemic_Xge_DetailedLogging { get; } = new EventId(715);
+
+		/// <summary>
+		/// Metadata about an XGE task (local or remote execution, start/end time etc)
+		/// </summary>
+		public static EventId Systemic_Xge_TaskMetadata { get; } = new EventId(716);
+
+		/// <summary>
+		/// DDC is slow
+		/// </summary>
+		public static EventId Systemic_SlowDDC { get; } = new EventId(720);
+
+		/// <summary>
+		/// Internal Horde error
+		/// </summary>
+		public static EventId Systemic_Horde { get; } = new EventId(730);
+
+		/// <summary>
+		/// Artifact upload failed
+		/// </summary>
+		public static EventId Systemic_Horde_ArtifactUpload { get; } = new EventId(731);
+
+		/// <summary>
+		/// Internal Horde error
+		/// </summary>
+		public static EventId Systemic_Horde_Compute { get; } = new EventId(732);
+
+		/// <summary>
+		/// HTTP error
+		/// </summary>
+		public static EventId Systemic_Horde_Http { get; } = new EventId(733);
+
+		/// <summary>
+		/// Harmless pdbutil error
+		/// </summary>
+		public static EventId Systemic_PdbUtil { get; } = new EventId(740);
+
+		/// <summary>
+		/// A systemic event from MSBuild
+		/// </summary>
+		public static EventId Systemic_MSBuild { get; } = new EventId(750);
+
+		/// <summary>
+		/// A systemic event from XCode
+		/// </summary>
+		public static EventId Systemic_XCode { get; } = new EventId(755);
+
+		/// <summary>
+		/// Robomerge gate is locked
+		/// </summary>
+		public static EventId Systemic_RoboMergeGateLocked { get; } = new EventId(760);
+
+		/// <summary>
+		/// XGEControlWorker is missing
+		/// </summary>
+		public static EventId Systemic_MissingXgeControlWorker { get; } = new EventId(761);
+
+		/// <summary>
+		/// Log parser is taking a significant amount of time; may be bottleneck.
+		/// </summary>
+		public static EventId Systemic_LogParserBottleneck { get; } = new EventId(762);
+
+		/// <summary>
+		/// Exception parsing a log event.
+		/// </summary>
+		public static EventId Systemic_LogEventMatcher { get; } = new EventId(763);
+
+		/// <summary>
+		/// Host is down during I/O operation, usually meaning a share isn't accessible.
+		/// </summary>
+		public static EventId Systemic_HostDownIOException { get; } = new EventId(764);
+
+		/// <summary>
+		/// Microsoft SignTool cannot reach specified timestamp server 
+		/// 
+		/// Signing performs a looping retry using different servers and throws when all retries attempts have been exhausted.
+		/// </summary>
+		public static EventId Systemic_SignToolTimeStampServer { get; } = new EventId(765);
+
+		/// <summary>
+		/// Microsoft SignTool generic error
+		/// 
+		/// Usually preceded by another type of error.
+		/// </summary>
+		public static EventId Systemic_SignTool { get; } = new EventId(766);
+
+		/// <summary>
+		/// Missing file list error
+		/// 
+		/// Usually due to a volume that is not mounted or some of other IO error
+		/// </summary>
+		public static EventId Systemic_MissingFileList { get; } = new EventId(767);
+
+		/// <summary>
+		/// Out of disk space
+		/// </summary>
+		public static EventId Systemic_OutOfDiskSpace { get; } = new EventId(768);
+
+		/// <summary>
+		/// Error moving files to cache
+		/// </summary>
+		public static EventId Systemic_ManagedWorkspace { get; } = new EventId(769);
+
+		/// <summary>
+		/// Generic timeout message emmitted by UBA
+		/// </summary>
+		public static EventId Systemic_UbaCacheResponseTimeout { get; } = new EventId(770);
+
+		/// <summary>
+		/// Docker messages that are unactionable and can be ignored
+		/// </summary>
+		public static EventId Systemic_Docker { get; } = new EventId(771);
+
+		/// <summary>
+		/// Maximum code for systemic events. Add new events in the 700-799 range.
+		/// </summary>
+		public static EventId Systemic_Max { get; } = new EventId(799);
+
+		/// <summary>
+		/// Generic Sanitizer Error
+		/// </summary>
+		public static EventId Sanitizer { get; } = new EventId(800);
+
+		/// <summary>
+		/// Thread Sanitizer Error
+		/// </summary>
+		public static EventId Sanitizer_Thread { get; } = new EventId(801);
+
+		/// <summary>
+		/// Address Sanitizer Error
+		/// </summary>
+		public static EventId Sanitizer_Address { get; } = new EventId(802);
+
+		/// <summary>
+		/// RaceDetector Sanitizer Error
+		/// </summary>
+		public static EventId Sanitizer_RaceDetector { get; } = new EventId(803);
+
+		/// <summary>
+		/// World Leak Error
+		/// </summary>
+		public static EventId WorldLeak { get; } = new EventId(900);
+
+		/// <summary>
+		/// Horde error codes
+		/// </summary>
+		public static EventId Horde { get; } = new EventId(1000);
+
+		/// <summary>
+		/// Invalid preflight change
+		/// </summary>
+		public static EventId Horde_InvalidPreflight { get; } = new EventId(1001);
+
+		/// <summary>
+		/// Information about blobs being read
+		/// </summary>
+		public static EventId Horde_BlobRead { get; } = new EventId(1002);
+
+		/// <summary>
+		/// Event used to supporess the logging from the Horde build health service, regardless of log level
+		/// </summary>
+		public static EventId Horde_BuildHealth_Ignore { get; } = new EventId(1003);
+
+		/// <summary>
+		/// Maximum code for Horde events. Add new events in the 1000-1099 range.
+		/// </summary>
+		public static EventId Horde_Max { get; } = new EventId(1099);
+
+		/// <summary>
+		/// An event while running zen cli
+		/// </summary>
+		public static EventId ZenCli { get; } = new EventId(1100);
+
+		/// <summary>
+		/// Strongly typed zen cli event
+		/// </summary>
+		public static EventId ZenCli_Event { get; } = new EventId(1101);
+
+		/// <summary>
+		/// Maximum code for zencli events. Add new events in the 1100-1200 range.
+		/// </summary>
+		public static EventId ZenCli_Max { get; } = new EventId(1199);
+	}
+}
diff --git a/Engine/Source/Programs/Shared/EpicGames.Core/Logging/LegacyEventLogger.cs b/Engine/Source/Programs/Shared/EpicGames.Core/Logging/LegacyEventLogger.cs
new file mode 100644
--- /dev/null
+++ b/Engine/Source/Programs/Shared/EpicGames.Core/Logging/LegacyEventLogger.cs
@@ -0,0 +1,168 @@
+// Copyright Epic Games, Inc. All Rights Reserved.
+
+using System;
+using System.Threading;
+using System.Threading.Channels;
+using Microsoft.Extensions.Logging;
+
+namespace EpicGames.Core
+{
+	/// <summary>
+	/// Wrapper around a custom logger interface which flushes the event parser when switching between legacy
+	/// and native structured logging. Also support threading the log messages to minimize performance overhead
+	/// </summary>
+	sealed class LegacyEventLogger : ILogger, IDisposable
+	{
+		private volatile ILogger _inner;
+		private readonly LogEventParser _parser;
+		private ThreadingState? _state;
+
+		public LogEventParser Parser => _parser;
+
+		public LegacyEventLogger(ILogger inner)
+		{
+			_inner = inner;
+			_parser = new LogEventParser(inner);
+		}
+
+		public void Dispose()
+		{
+			// Drain and stop the worker if enabled, then dispose parser.
+			DisableThreaded();
+			_parser.Dispose();
+		}
+
+		/// <summary>
+		/// Set inner logger to transfer log calls to
+		/// </summary>
+		/// <param name="inner"></param>
+		public void SetInnerLogger(ILogger inner)
+		{
+			lock (_parser)
+			{
+				_parser.Flush();
+				_inner = inner;
+				_parser.Logger = inner;
+			}
+		}
+
+		/// <summary>
+		/// Enable background logging (idempotent).
+		/// </summary>
+		public void EnableThreaded()
+		{
+			ThreadingState created = new();
+			Interlocked.CompareExchange(ref _state, created, null);
+		}
+
+		/// <summary>
+		/// Disable background logging and drain queued items (idempotent).
+		/// </summary>
+		public void DisableThreaded()
+		{
+			ThreadingState? state = Interlocked.Exchange(ref _state, null);
+			state?.Stop(); // completes channel, drains, joins worker thread
+		}
+
+		// ---- ILogger ----
+		public IDisposable? BeginScope<TState>(TState state) where TState : notnull => _inner.BeginScope(state);
+
+		public bool IsEnabled(LogLevel logLevel) => _inner.IsEnabled(logLevel);
+
+		public void Log<TState>(LogLevel logLevel, EventId eventId, TState state, Exception? exception, Func<TState, Exception?, string> formatter)
+		{
+			void DoLog()
+			{
+				ILogger logger = _inner;
+				lock (_parser)
+				{
+					_parser.Flush();
+				}
+				logger.Log(logLevel, eventId, state, exception, formatter);
+			}
+
+			ThreadingState? s = Volatile.Read(ref _state);
+			if (s == null)
+			{
+				DoLog();
+				return;
+			}
+
+			if (logLevel < LogLevel.Error)
+			{
+				// Non-blocking enqueue; unbounded channel TryWrite is always true.
+				s.Writer.TryWrite(DoLog);
+				return;
+			}
+
+			// Error+ logs: preserve synchronous semantics.
+			using ManualResetEventSlim done = new(false, spinCount: 1);
+			s.Writer.TryWrite(() =>
+			{
+				try
+				{
+					DoLog();
+				}
+				finally
+				{
+					done.Set();
+				}
+			});
+			done.Wait();
+		}
+
+		private sealed class ThreadingState
+		{
+			public readonly ChannelWriter<Action> Writer;
+			private readonly ChannelReader<Action> _reader;
+			private readonly Thread _thread;
+
+			public ThreadingState()
+			{
+				// Single reader = true => fewer fences and better throughput.
+				Channel<Action> ch = Channel.CreateUnbounded<Action>(new UnboundedChannelOptions
+				{
+					SingleReader = true,
+					SingleWriter = false,
+					AllowSynchronousContinuations = false
+				});
+				Writer = ch.Writer;
+				_reader = ch.Reader;
+
+				_thread = new Thread(static o =>
+				{
+					ThreadingState self = (ThreadingState)o!;
+					ChannelReader<Action> reader = self._reader;
+
+					while (reader.WaitToReadAsync().AsTask().GetAwaiter().GetResult())
+					{
+						while (reader.TryRead(out Action? work))
+						{
+							try
+							{
+								work();
+							}
+							catch
+							{
+								// Swallow exceptions to keep the logging loop healthy.
+								// Consider publishing to a fallback sink if needed.
+							}
+						}
+					}
+				})
+				{
+					IsBackground = true,
+					Name = "LegacyEventLogger",
+				};
+				_thread.Start(this);
+			}
+
+			public void Stop()
+			{
+				// Complete the writer so the reader drains and then exits cleanly.
+				Writer.TryComplete();
+				_thread.Join();
+			}
+		}
+	}
+}
\ No newline at end of file
diff --git a/Engine/Source/Programs/Shared/EpicGames.Core/Logging/Log.cs b/Engine/Source/Programs/Shared/EpicGames.Core/Logging/Log.cs
new file mode 100644
--- /dev/null
+++ b/Engine/Source/Programs/Shared/EpicGames.Core/Logging/Log.cs
@@ -0,0 +1,1522 @@
+// Copyright Epic Games, Inc. All Rights Reserved.
+
+using System;
+using System.Collections.Concurrent;
+using System.Collections.Generic;
+using System.Diagnostics;
+using System.IO;
+using System.Linq;
+using System.Reflection;
+using System.Text;
+using System.Text.RegularExpressions;
+using System.Threading.Channels;
+using System.Threading.Tasks;
+using JetBrains.Annotations;
+using Microsoft.Extensions.DependencyInjection;
+using Microsoft.Extensions.Logging;
+
+namespace EpicGames.Core
+{
+	/// <summary>
+	/// Log Event Type
+	/// </summary>
+#pragma warning disable CA1027 // Mark enums with FlagsAttribute
+	public enum LogEventType
+#pragma warning restore CA1027 // Mark enums with FlagsAttribute
+	{
+		/// <summary>
+		/// The log event is a fatal error
+		/// </summary>
+		Fatal = LogLevel.Critical,
+
+		/// <summary>
+		/// The log event is an error
+		/// </summary>
+		Error = LogLevel.Error,
+
+		/// <summary>
+		/// The log event is a warning
+		/// </summary>
+		Warning = LogLevel.Warning,
+
+		/// <summary>
+		/// Output the log event to the console
+		/// </summary>
+		Console = LogLevel.Information,
+
+		/// <summary>
+		/// Output the event to the on-disk log
+		/// </summary>
+		Log = LogLevel.Debug,
+
+		/// <summary>
+		/// The log event should only be displayed if verbose logging is enabled
+		/// </summary>
+		Verbose = LogLevel.Trace,
+
+		/// <summary>
+		/// The log event should only be displayed if very verbose logging is enabled
+		/// </summary>
+#pragma warning disable CA1069 // Enums values should not be duplicated
+		VeryVerbose = LogLevel.Trace
+#pragma warning restore CA1069 // Enums values should not be duplicated
+	}
+
+	/// <summary>
+	/// Options for formatting messages
+	/// </summary>
+	[Flags]
+	public enum LogFormatOptions
+	{
+		/// <summary>
+		/// Format normally
+		/// </summary>
+		None = 0,
+
+		/// <summary>
+		/// Never write a severity prefix. Useful for pre-formatted messages that need to be in a particular format for, eg. the Visual Studio output window
+		/// </summary>
+		NoSeverityPrefix = 1,
+
+		/// <summary>
+		/// Do not output text to the console
+		/// </summary>
+		NoConsoleOutput = 2,
+	}
+
+	/// <summary>
+	/// UAT/UBT Custom log system.
+	/// 
+	/// This lets you use any TraceListeners you want, but you should only call the static 
+	/// methods below, not call Trace.XXX directly, as the static methods
+	/// This allows the system to enforce the formatting and filtering conventions we desire.
+	///
+	/// For posterity, we cannot use the Trace or TraceSource class directly because of our special log requirements:
+	///   1. We possibly capture the method name of the logging event. This cannot be done as a macro, so must be done at the top level so we know how many layers of the stack to peel off to get the real function.
+	///   2. We have a verbose filter we would like to apply to all logs without having to have each listener filter individually, which would require our string formatting code to run every time.
+	///   3. We possibly want to ensure severity prefixes are logged, but Trace.WriteXXX does not allow any severity info to be passed down.
+	/// </summary>
+	public static class Log
+	{
+		/// <summary>
+		/// Singleton instance of the default output logger
+		/// </summary>
+		private static readonly DefaultLogger s_defaultLogger = new();
+
+		/// <summary>
+		/// Logger instance which parses events and forwards them to the main logger.
+		/// </summary>
+		private static readonly LegacyEventLogger s_legacyLogger = new(s_defaultLogger);
+
+		/// <summary>
+		/// Accessor for the global event parser from legacy events
+		/// </summary>
+		public static LogEventParser EventParser => s_legacyLogger.Parser;
+
+		/// <summary>
+		/// Logger instance
+		/// </summary>
+		public static ILogger Logger => s_legacyLogger;
+
+		/// <summary>
+		/// When true, verbose logging is enabled.
+		/// </summary>
+		public static LogEventType OutputLevel
+		{
+			get => (LogEventType)s_defaultLogger.OutputLevel;
+			set => s_defaultLogger.OutputLevel = (LogLevel)value;
+		}
+
+		/// <summary>
+		/// Whether to include timestamps on each line of log output
+		/// </summary>
+		public static bool IncludeTimestamps
+		{
+			get => s_defaultLogger.IncludeTimestamps;
+			set => s_defaultLogger.IncludeTimestamps = value;
+		}
+
+		/// <summary>
+		/// When true, warnings and errors will have a WARNING: or ERROR: prefix, respectively.
+		/// </summary>
+		public static bool IncludeSeverityPrefix { get; set; } = true;
+
+		/// <summary>
+		/// When true, warnings and errors will have a prefix suitable for display by MSBuild (avoiding error messages showing as (EXEC : Error : ")
+		/// </summary>
+		public static bool IncludeProgramNameWithSeverityPrefix { get; set; }
+
+		/// <summary>
+		/// When true, will detect warnings and errors and set the console output color to yellow and red.
+		/// </summary>
+		public static bool ColorConsoleOutput
+		{
+			get => s_defaultLogger.ColorConsoleOutput;
+			set => s_defaultLogger.ColorConsoleOutput = value;
+		}
+
+		/// <summary>
+		/// When true, a timestamp will be written to the log file when the first listener is added
+		/// </summary>
+		public static bool IncludeStartingTimestamp
+		{
+			get => s_defaultLogger.IncludeStartingTimestamp;
+			set => s_defaultLogger.IncludeStartingTimestamp = value;
+		}
+
+		/// <summary>
+		/// When true, create a backup of any log file that would be overwritten by a new log
+		/// Log.txt will be backed up with its UTC creation time in the name e.g.
+		/// Log-backup-2021.10.29-19.53.17.txt
+		/// </summary>
+		public static bool BackupLogFiles { get; set; } = true;
+
+		/// <summary>
+		/// The number of backups to be preserved - when there are more than this, the oldest backups will be deleted.
+		/// Backups will not be deleted if BackupLogFiles is false.
+		/// </summary>
+		public static int LogFileBackupCount { get; set; } = 10;
+
+		/// <summary>
+		/// Path to the log file being written to. May be null.
+		/// </summary>
+		public static FileReference? OutputFile => s_defaultLogger?.OutputFile;
+
+		/// <summary>
+		/// A collection of strings that have been already written once
+		/// </summary>
+		private static readonly ConcurrentDictionary<string, bool> s_writeOnceSet = new();
+
+		/// <summary>
+		/// Overrides the logger used for formatting output, after event parsing
+		/// </summary>
+		/// <param name="logger"></param>
+		public static void SetInnerLogger(ILogger logger)
+		{
+			s_legacyLogger.SetInnerLogger(logger);
+		}
+
+		/// <summary>
+		/// Flush the current log output
+		/// </summary>
+		/// <returns></returns>
+		public static async Task FlushAsync()
+		{
+			await s_defaultLogger.FlushAsync();
+		}
+
+		/// <summary>
+		/// Enable threaded logging. This can be used for performance reasons to prevent stalls when lots of logging is made from multiple threads
+		/// </summary>
+		public static void EnableThreaded()
+		{
+			s_legacyLogger.EnableThreaded();
+		}
+
+		/// <summary>
+		/// Disable threaded logging. will flush queued log entries and get back to synchronous logging
+		/// </summary>
+		public static void DisableThreaded()
+		{
+			s_legacyLogger.DisableThreaded();
+		}
+
+		/// <summary>
+		/// Backup an existing log file if it already exists at the outputpath
+		/// </summary>
+		/// <param name="outputFile">The file to back up</param>
+		public static void BackupLogFile(FileReference outputFile)
+		{
+			if (!BackupLogFiles || !FileReference.Exists(outputFile))
+			{
+				return;
+			}
+
+			// before creating a new backup, cap the number of existing files
+			string filenameWithoutExtension = outputFile.GetFileNameWithoutExtension();
+			string extension = outputFile.GetExtension();
+
+			Regex backupForm =
+				new(filenameWithoutExtension + @"-backup-\d\d\d\d\.\d\d\.\d\d-\d\d\.\d\d\.\d\d" + extension);
+
+			foreach (FileReference oldBackup in DirectoryReference
+				.EnumerateFiles(outputFile.Directory)
+				// find files that match the way that we name backup files
+				.Where(x => backupForm.IsMatch(x.GetFileName()))
+				// sort them from newest to oldest
+				.OrderByDescending(x => x.GetFileName())
+				// skip the newest ones that are to be kept; -1 because we're about to create another backup.
+				.Skip(LogFileBackupCount - 1))
+			{
+				Logger.LogDebug("Deleting old log file: {File}", oldBackup);
+				FileReference.Delete(oldBackup);
+			}
+
+			// Ensure that the backup gets a unique name, in the extremely unlikely case that UBT was run twice during
+			// the same second.
+			DateTime fileTime = File.GetCreationTimeUtc(outputFile.FullName);
+
+			FileReference backupFile;
+			for (; ; )
+			{
+				string timestamp = $"{fileTime:yyyy.MM.dd-HH.mm.ss}";
+				backupFile = FileReference.Combine(outputFile.Directory,
+					$"{filenameWithoutExtension}-backup-{timestamp}{extension}");
+				if (!FileReference.Exists(backupFile))
+				{
+					break;
+				}
+
+				fileTime = fileTime.AddSeconds(1);
+			}
+
+			FileReference.Move(outputFile, backupFile);
+		}
+
+		/// <summary>
+		/// Adds a trace listener that writes to a log file.
+		/// If a StartupTraceListener was in use, this function will copy its captured data to the log file(s)
+		/// and remove the startup listener from the list of registered listeners.
+		/// </summary>
+		/// <param name="name">Identifier for the writer</param>
+		/// <param name="outputFile">The file to write to</param>
+		public static void AddFileWriter(string name, FileReference outputFile)
+		{
+			Logger.LogInformation("Log file: {OutputFile}", outputFile);
+
+			BackupLogFile(outputFile);
+			AddFileWriterWithoutBackup(name, outputFile);
+		}
+
+		/// <summary>
+		/// Adds a trace listener that writes to a log file.
+		/// If a StartupTraceListener was in use, this function will copy its captured data to the log file(s)
+		/// and remove the startup listener from the list of registered listeners.
+		/// </summary>
+		/// <param name="name">Identifier for the writer</param>
+		/// <param name="outputFile">The file to write to</param>
+		public static void AddFileWriterWithoutBackup(string name, FileReference outputFile)
+		{
+			s_defaultLogger.AddFileWriterWithoutBackup(name, outputFile);
+		}
+
+		/// <summary>
+		/// Adds a <see cref="TraceListener"/> to the collection in a safe manner.
+		/// </summary>
+		/// <param name="traceListener">The <see cref="TraceListener"/> to add.</param>
+		public static void AddTraceListener(TraceListener traceListener)
+		{
+			s_defaultLogger.AddTraceListener(traceListener);
+		}
+
+		/// <summary>
+		/// Removes a <see cref="TraceListener"/> from the collection in a safe manner.
+		/// </summary>
+		/// <param name="traceListener">The <see cref="TraceListener"/> to remove.</param>
+		public static void RemoveTraceListener(TraceListener traceListener)
+		{
+			s_defaultLogger.RemoveTraceListener(traceListener);
+		}
+
+		/// <summary>
+		/// Removes a <see cref="StartupTraceListener"/> from the collection in a safe manner.
+		/// </summary>
+		public static void RemoveStartupTraceListener()
+		{
+			s_defaultLogger.RemoveStartupTraceListener();
+		}
+
+		/// <summary>
+		/// Determines if a TextWriterTraceListener has been added to the list of trace listeners
+		/// </summary>
+		/// <returns>True if a TextWriterTraceListener has been added</returns>
+		public static bool HasFileWriter()
+		{
+			return s_defaultLogger.HasFileWriter();
+		}
+
+		/// <summary>
+		/// Converts a LogEventType into a log prefix. Only used when bLogSeverity is true.
+		/// </summary>
+		/// <param name="severity"></param>
+		/// <returns></returns>
+		private static string GetSeverityPrefix(LogEventType severity)
+		{
+			switch (severity)
+			{
+				case LogEventType.Fatal:
+					return "FATAL ERROR: ";
+				case LogEventType.Error:
+					return "ERROR: ";
+				case LogEventType.Warning:
+					return "WARNING: ";
+				case LogEventType.Console:
+					return "";
+				case LogEventType.Verbose:
+					return "VERBOSE: ";
+				default:
+					return "";
+			}
+		}
+
+		/// <summary>
+		/// Writes a formatted message to the console. All other functions should boil down to calling this method.
+		/// </summary>
+		/// <param name="bWriteOnce">If true, this message will be written only once</param>
+		/// <param name="verbosity">Message verbosity level. We only meaningfully use values up to Verbose</param>
+		/// <param name="formatOptions">Options for formatting messages</param>
+		/// <param name="format">Message format string.</param>
+		/// <param name="args">Optional arguments</param>
+		[StringFormatMethod("Format")]
+		private static void WriteLinePrivate(bool bWriteOnce, LogEventType verbosity, LogFormatOptions formatOptions, string format, params object?[] args)
+		{
+			if (Logger.IsEnabled((LogLevel)verbosity))
+			{
+				StringBuilder message = new();
+
+				// Get the severity prefix for this message
+				if (IncludeSeverityPrefix && ((formatOptions & LogFormatOptions.NoSeverityPrefix) == 0))
+				{
+					message.Append(GetSeverityPrefix(verbosity));
+					if (message.Length > 0 && IncludeProgramNameWithSeverityPrefix)
+					{
+						// Include the executable name when running inside MSBuild. If unspecified, MSBuild re-formats them with an "EXEC :" prefix.
+						message.Insert(0, $"{Path.GetFileNameWithoutExtension(Assembly.GetEntryAssembly()!.Location)}: ");
+					}
+				}
+
+				// Append the formatted string
+				int indentLen = message.Length;
+				if (args.Length == 0)
+				{
+					message.Append(format);
+				}
+				else
+				{
+					message.AppendFormat(format, args);
+				}
+
+				// Replace any Windows \r\n sequences with \n
+				message.Replace("\r\n", "\n");
+
+				// Remove any trailing whitespace
+				int trimLen = message.Length;
+				while (trimLen > 0 && " \t\r\n".Contains(message[trimLen - 1], StringComparison.Ordinal))
+				{
+					trimLen--;
+				}
+				message.Remove(trimLen, message.Length - trimLen);
+
+				// Update the indent length to include any whitespace at the start of the message
+				while (indentLen < message.Length && message[indentLen] == ' ')
+				{
+					indentLen++;
+				}
+
+				// If there are multiple lines, insert a prefix at the start of each one
+				for (int idx = 0; idx < message.Length; idx++)
+				{
+					if (message[idx] == '\n')
+					{
+						message.Insert(idx + 1, " ", indentLen);
+						idx += indentLen;
+					}
+				}
+
+				// if we want this message only written one time, check if it was already written out
+				if (bWriteOnce && !s_writeOnceSet.TryAdd(message.ToString(), true))
+				{
+					return;
+				}
+
+				// Forward it on to the internal logger
+				if (verbosity < LogEventType.Console)
+				{
+					Logger.Log((LogLevel)verbosity, "{Message}", message.ToString());
+				}
+				else
+				{
+					lock (EventParser)
+					{
+						int baseIdx = 0;
+						for (int idx = 0; idx < message.Length; idx++)
+						{
+							if (message[idx] == '\n')
+							{
+								EventParser.WriteLine(message.ToString(baseIdx, idx - baseIdx));
+								baseIdx = idx + 1;
+							}
+						}
+						EventParser.WriteLine(message.ToString(baseIdx, message.Length - baseIdx));
+					}
+				}
+			}
+		}
+
+		/// <summary>
+		/// Similar to Trace.WriteLineIf
+		/// </summary>
+		/// <param name="condition"></param>
+		/// <param name="verbosity"></param>
+		/// <param name="format"></param>
+		/// <param name="args"></param>
+		[StringFormatMethod("Format")]
+		public static void WriteLineIf(bool condition, LogEventType verbosity, string format, params object?[] args)
+		{
+			if (condition)
+			{
+				WriteLinePrivate(false, verbosity, LogFormatOptions.None, format, args);
+			}
+		}
+
+		/// <summary>
+		/// Similar to Trace.WriteLine
+		/// </summary>
+		/// <param name="verbosity"></param>
+		/// <param name="format"></param>
+		/// <param name="args"></param>
+		[StringFormatMethod("Format")]
+		public static void WriteLine(LogEventType verbosity, string format, params object?[] args)
+		{
+			WriteLinePrivate(false, verbosity, LogFormatOptions.None, format, args);
+		}
+
+		/// <summary>
+		/// Similar to Trace.WriteLine
+		/// </summary>
+		/// <param name="verbosity"></param>
+		/// <param name="formatOptions"></param>
+		/// <param name="format"></param>
+		/// <param name="args"></param>
+		[StringFormatMethod("Format")]
+		public static void WriteLine(LogEventType verbosity, LogFormatOptions formatOptions, string format, params object?[] args)
+		{
+			WriteLinePrivate(false, verbosity, formatOptions, format, args);
+		}
+
+		/// <summary>
+		/// Formats an exception for display in the log. The exception message is shown as an error, and the stack trace is included in the log.
+		/// </summary>
+		/// <param name="ex">The exception to display</param>
+		/// <param name="logFileName">The log filename to display, if any</param>
+		public static void WriteException(Exception ex, FileReference? logFileName)
+		{
+			string logSuffix = (logFileName == null) ? "" : String.Format("\n(see {0} for full exception trace)", logFileName);
+			Logger.LogDebug("==============================================================================");
+			Logger.LogError(ex, "{Message}{Suffix}", ExceptionUtils.FormatException(ex), logSuffix);
+			Logger.LogDebug("");
+			Logger.LogDebug("{Details}", ExceptionUtils.FormatExceptionDetails(ex));
+			Logger.LogDebug("==============================================================================");
+		}
+
+		/// <summary>
+		/// Writes an error message to the console.
+		/// </summary>
+		/// <param name="format">Message format string</param>
+		/// <param name="args">Optional arguments</param>
+		[StringFormatMethod("Format")]
+		[Obsolete("Use Logger.LogError with a message template instead; see https://tinyurl.com/bp96bk2r.", false)]
+		public static void TraceError(string format, params object?[] args)
+		{
+			WriteLinePrivate(false, LogEventType.Error, LogFormatOptions.None, format, args);
+		}
+
+		/// <summary>
+		/// Writes an error message to the console, in a format suitable for Visual Studio to parse.
+		/// </summary>
+		/// <param name="file">The file containing the error</param>
+		/// <param name="format">Message format string</param>
+		/// <param name="args">Optional arguments</param>
+		[StringFormatMethod("Format")]
+		public static void TraceErrorTask(FileReference file, string format, params object?[] args)
+		{
+			WriteLinePrivate(false, LogEventType.Error, LogFormatOptions.NoSeverityPrefix, "{0}: error: {1}", file, String.Format(format, args));
+		}
+
+		/// <summary>
+		/// Writes an error message to the console, in a format suitable for Visual Studio to parse.
+		/// </summary>
+		/// <param name="file">The file containing the error</param>
+		/// <param name="line">Line number of the error</param>
+		/// <param name="format">Message format string</param>
+		/// <param name="args">Optional arguments</param>
+		[StringFormatMethod("Format")]
+		public static void TraceErrorTask(FileReference file, int line, string format, params object?[] args)
+		{
+			WriteLinePrivate(false, LogEventType.Error, LogFormatOptions.NoSeverityPrefix, "{0}({1}): error: {2}", file, line, String.Format(format, args));
+		}
+
+		/// <summary>
+		/// Writes a verbose message to the console.
+		/// </summary>
+		/// <param name="format">Message format string</param>
+		/// <param name="args">Optional arguments</param>
+		[Conditional("TRACE")]
+		[StringFormatMethod("Format")]
+		[Obsolete("Use Logger.LogDebug with a message template instead; see https://tinyurl.com/bp96bk2r.", false)]
+		public static void TraceVerbose(string format, params object?[] args)
+		{
+			WriteLinePrivate(false, LogEventType.Verbose, LogFormatOptions.None, format, args);
+		}
+
+		/// <summary>
+		/// Writes a message to the console.
+		/// </summary>
+		/// <param name="format">Message format string</param>
+		/// <param name="args">Optional arguments</param>
+		[StringFormatMethod("Format")]
+		[Obsolete("Use Logger.LogInformation with a message template instead; see https://tinyurl.com/bp96bk2r.", false)]
+		public static void TraceInformation(string format, params object?[] args)
+		{
+			WriteLinePrivate(false, LogEventType.Console, LogFormatOptions.None, format, args);
+		}
+
+		/// <summary>
+		/// Writes a warning message to the console.
+		/// </summary>
+		/// <param name="format">Message format string</param>
+		/// <param name="args">Optional arguments</param>
+		[StringFormatMethod("Format")]
+		[Obsolete("Use Logger.LogWarning with a message template instead; see https://tinyurl.com/bp96bk2r.", false)]
+		public static void TraceWarning(string format, params object?[] args)
+		{
+			WriteLinePrivate(false, LogEventType.Warning, LogFormatOptions.None, format, args);
+		}
+
+		/// <summary>
+		/// Writes a warning message to the console, in a format suitable for Visual Studio to parse.
+		/// </summary>
+		/// <param name="file">The file containing the warning</param>
+		/// <param name="format">Message format string</param>
+		/// <param name="args">Optional arguments</param>
+		[StringFormatMethod("Format")]
+		public static void TraceWarningTask(FileReference file, string format, params object?[] args)
+		{
+			WriteLinePrivate(false, LogEventType.Warning, LogFormatOptions.NoSeverityPrefix, "{0}: warning: {1}", file, String.Format(format, args));
+		}
+
+		/// <summary>
+		/// Writes a warning message to the console, in a format suitable for Visual Studio to parse.
+		/// </summary>
+		/// <param name="file">The file containing the warning</param>
+		/// <param name="line">Line number of the warning</param>
+		/// <param name="format">Message format string</param>
+		/// <param name="args">Optional arguments</param>
+		[StringFormatMethod("Format")]
+		public static void TraceWarningTask(FileReference file, int line, string format, params object?[] args)
+		{
+			WriteLinePrivate(false, LogEventType.Warning, LogFormatOptions.NoSeverityPrefix, "{0}({1}): warning: {2}", file, line, String.Format(format, args));
+		}
+
+		/// <summary>
+		/// Writes a message to the console.
+		/// </summary>
+		/// <param name="file">The file containing the message</param>
+		/// <param name="line">Line number of the message</param>
+		/// <param name="format">Message format string</param>
+		/// <param name="args">Optional arguments</param>
+		public static void TraceConsoleTask(FileReference file, int line, string format, params object?[] args)
+		{
+			WriteLinePrivate(false, LogEventType.Console, LogFormatOptions.NoSeverityPrefix, "{0}({1}): {2}", file, line, String.Format(format, args));
+		}
+
+		/// <summary>
+		/// Writes a very verbose message to the console.
+		/// </summary>
+		/// <param name="format">Message format string</param>
+		/// <param name="args">Optional arguments</param>
+		[Conditional("TRACE")]
+		[StringFormatMethod("Format")]
+		[Obsolete("Use Logger.LogTrace with a message template instead; see https://tinyurl.com/bp96bk2r.", false)]
+		public static void TraceVeryVerbose(string format, params object?[] args)
+		{
+			WriteLinePrivate(false, LogEventType.VeryVerbose, LogFormatOptions.None, format, args);
+		}
+
+		/// <summary>
+		/// Writes a message to the log only.
+		/// </summary>
+		/// <param name="format">Message format string</param>
+		/// <param name="args">Optional arguments</param>
+		[Conditional("TRACE")]
+		[StringFormatMethod("Format")]
+		[Obsolete("Use Logger.LogDebug with a message template instead; see https://tinyurl.com/bp96bk2r.", false)]
+		public static void TraceLog(string format, params object?[] args)
+		{
+			WriteLinePrivate(false, LogEventType.Log, LogFormatOptions.None, format, args);
+		}
+
+		/// <summary>
+		/// Similar to Trace.WriteLine
+		/// </summary>
+		/// <param name="verbosity"></param>
+		/// <param name="format"></param>
+		/// <param name="args"></param>
+		[StringFormatMethod("Format")]
+		public static void WriteLineOnce(LogEventType verbosity, string format, params object?[] args)
+		{
+			WriteLinePrivate(true, verbosity, LogFormatOptions.None, format, args);
+		}
+
+		/// <summary>
+		/// Similar to Trace.WriteLine
+		/// </summary>
+		/// <param name="verbosity"></param>
+		/// <param name="options"></param>
+		/// <param name="format"></param>
+		/// <param name="args"></param>
+		[StringFormatMethod("Format")]
+		public static void WriteLineOnce(LogEventType verbosity, LogFormatOptions options, string format, params object?[] args)
+		{
+			WriteLinePrivate(true, verbosity, options, format, args);
+		}
+
+		/// <summary>
+		/// Writes an error message to the console.
+		/// </summary>
+		/// <param name="format">Message format string</param>
+		/// <param name="args">Optional arguments</param>
+		[StringFormatMethod("Format")]
+		public static void TraceErrorOnce(string format, params object?[] args)
+		{
+			WriteLinePrivate(true, LogEventType.Error, LogFormatOptions.None, format, args);
+		}
+
+		/// <summary>
+		/// Writes a verbose message to the console.
+		/// </summary>
+		/// <param name="format">Message format string</param>
+		/// <param name="args">Optional arguments</param>
+		[Conditional("TRACE")]
+		[StringFormatMethod("Format")]
+		public static void TraceVerboseOnce(string format, params object?[] args)
+		{
+			WriteLinePrivate(true, LogEventType.Verbose, LogFormatOptions.None, format, args);
+		}
+
+		/// <summary>
+		/// Writes a message to the console.
+		/// </summary>
+		/// <param name="format">Message format string</param>
+		/// <param name="args">Optional arguments</param>
+		[StringFormatMethod("Format")]
+		public static void TraceInformationOnce(string format, params object?[] args)
+		{
+			WriteLinePrivate(true, LogEventType.Console, LogFormatOptions.None, format, args);
+		}
+
+		/// <summary>
+		/// Writes a warning message to the console.
+		/// </summary>
+		/// <param name="format">Message format string</param>
+		/// <param name="args">Optional arguments</param>
+		[StringFormatMethod("Format")]
+		public static void TraceWarningOnce(string format, params object?[] args)
+		{
+			WriteLinePrivate(true, LogEventType.Warning, LogFormatOptions.None, format, args);
+		}
+
+		/// <summary>
+		/// Writes a warning message to the console.
+		/// </summary>
+		/// <param name="file">The file containing the error</param>
+		/// <param name="format">Message format string</param>
+		/// <param name="args">Optional arguments</param>
+		[StringFormatMethod("Format")]
+		public static void TraceWarningOnce(FileReference file, string format, params object?[] args)
+		{
+			WriteLinePrivate(true, LogEventType.Warning, LogFormatOptions.NoSeverityPrefix, "{0}: warning: {1}", file, String.Format(format, args));
+		}
+
+		/// <summary>
+		/// Writes a warning message to the console.
+		/// </summary>
+		/// <param name="file">The file containing the error</param>
+		/// <param name="line">Line number of the error</param>
+		/// <param name="format">Message format string</param>
+		/// <param name="args">Optional arguments</param>
+		[StringFormatMethod("Format")]
+		public static void TraceWarningOnce(FileReference file, int line, string format, params object?[] args)
+		{
+			WriteLinePrivate(true, LogEventType.Warning, LogFormatOptions.NoSeverityPrefix, "{0}({1}): warning: {2}", file, line, String.Format(format, args));
+		}
+
+		/// <summary>
+		/// Writes a very verbose message to the console.
+		/// </summary>
+		/// <param name="format">Message format string</param>
+		/// <param name="args">Optional arguments</param>
+		[Conditional("TRACE")]
+		[StringFormatMethod("Format")]
+		public static void TraceVeryVerboseOnce(string format, params object?[] args)
+		{
+			WriteLinePrivate(true, LogEventType.VeryVerbose, LogFormatOptions.None, format, args);
+		}
+
+		/// <summary>
+		/// Writes a message to the log only.
+		/// </summary>
+		/// <param name="format">Message format string</param>
+		/// <param name="args">Optional arguments</param>
+		[Conditional("TRACE")]
+		[StringFormatMethod("Format")]
+		public static void TraceLogOnce(string format, params object?[] args)
+		{
+			WriteLinePrivate(true, LogEventType.Log, LogFormatOptions.None, format, args);
+		}
+
+		/// <summary>
+		/// Enter a scope with the given status message. The message will be written to the console without a newline, allowing it to be updated through subsequent calls to UpdateStatus().
+		/// The message will be written to the log immediately. If another line is written while in a status scope, the initial status message is flushed to the console first.
+		/// </summary>
+		/// <param name="message">The status message</param>
+		[Conditional("TRACE")]
+		public static void PushStatus(string message)
+		{
+			s_defaultLogger.PushStatus(message);
+		}
+
+		/// <summary>
+		/// Updates the current status message. This will overwrite the previous status line.
+		/// </summary>
+		/// <param name="message">The status message</param>
+		[Conditional("TRACE")]
+		public static void UpdateStatus(string message)
+		{
+			s_defaultLogger.UpdateStatus(message);
+		}
+
+		/// <summary>
+		/// Updates the Pops the top status message from the stack. The mess
+		/// </summary>
+		[Conditional("TRACE")]
+		public static void PopStatus()
+		{
+			s_defaultLogger.PopStatus();
+		}
+	}
+
+	/// <summary>
+	/// Base implementation for loggers which handles creation of LogEvents from log commands
+	/// </summary>
+	public sealed class LoggerScopeCollection
+	{
+		class Scope : IDisposable
+		{
+			readonly LoggerScopeCollection _outer;
+			readonly IEnumerable<KeyValuePair<string, object?>>? _properties;
+
+			public IEnumerable<KeyValuePair<string, object?>>? Properties => _properties;
+
+			public Scope(LoggerScopeCollection outer, object? value)
+			{
+				_outer = outer;
+				_properties = value switch
+				{
+					IEnumerable<KeyValuePair<string, object?>> props => props,
+					_ => null
+				};
+
+				_outer._scopes.Add(this);
+			}
+
+			public void Dispose()
+				=> _outer._scopes.Remove(this);
+		}
+
+		readonly List<Scope> _scopes = [];
+
+		/// <summary>
+		/// Tests whether the collection is empty
+		/// </summary>
+		public bool IsEmpty()
+			=> _scopes.Count == 0;
+
+		/// <summary>
+		/// Creates a scope with the given state
+		/// </summary>
+		public IDisposable BeginScope<TState>(TState state)
+			=> new Scope(this, state);
+
+		/// <summary>
+		/// Enumerates all current properties
+		/// </summary>
+		public IEnumerable<KeyValuePair<string, object?>> GetProperties()
+		{
+			foreach (Scope scope in _scopes)
+			{
+				if (scope.Properties != null)
+				{
+					foreach (KeyValuePair<string, object?> property in scope.Properties)
+					{
+						yield return property;
+					}
+				}
+			}
+		}
+	}
+
+	/// <summary>
+	/// Logger which captures the output for rendering later
+	/// </summary>
+	public class CaptureLogger : ILogger
+	{
+		readonly LoggerScopeCollection _scopeCollection = new();
+
+		/// <summary>
+		/// List of captured events
+		/// </summary>
+		public List<LogEvent> Events { get; } = [];
+
+		/// <summary>
+		/// Renders the captured events as a single string
+		/// </summary>
+		/// <returns>Rendered log text</returns>
+		public string Render() => Render("\n");
+
+		/// <summary>
+		/// Renders the captured events as a single string
+		/// </summary>
+		/// <returns>Rendered log text</returns>
+		public string Render(string newLine) => String.Join(newLine, RenderLines());
+
+		/// <summary>
+		/// Renders all the captured events
+		/// </summary>
+		/// <returns>List of rendered log lines</returns>
+		public List<string> RenderLines() => Events.ConvertAll(x => x.ToString());
+
+		/// <summary>
+		/// Renders all the captured events to another logger
+		/// </summary>
+		/// <param name="logger"></param>
+		public void RenderTo(ILogger logger) => Events.ForEach(x => logger.Log(x.Level, x.Id, x, null, (s, e) => s.ToString()));
+
+		/// <inheritdoc/>
+		public IDisposable? BeginScope<TState>(TState state) where TState : notnull => _scopeCollection.BeginScope(state);
+
+		/// <inheritdoc/>
+		public bool IsEnabled(LogLevel logLevel) => true;
+
+		/// <inheritdoc/>
+		public void Log<TState>(LogLevel logLevel, EventId eventId, TState state, Exception? exception, Func<TState, Exception?, string> formatter)
+		{
+			LogEvent logEvent = LogEvent.FromState(logLevel, eventId, state, exception, formatter);
+			logEvent.AddProperties(_scopeCollection.GetProperties());
+			Events.Add(logEvent);
+		}
+	}
+
+	/// <summary>
+	/// Default log output device
+	/// </summary>
+	class DefaultLogger : ILogger, IDisposable
+	{
+		/// <summary>
+		/// Temporary status message displayed on the console.
+		/// </summary>
+		[DebuggerDisplay("{_headingText}")]
+		class StatusMessage
+		{
+			/// <summary>
+			/// The heading for this status message.
+			/// </summary>
+			public string _headingText;
+
+			/// <summary>
+			/// The current status text.
+			/// </summary>
+			public string _currentText;
+
+			/// <summary>
+			/// Whether the heading has been written to the console. Before the first time that lines are output to the log in the midst of a status scope, the heading will be written on a line of its own first.
+			/// </summary>
+			public bool _hasFlushedHeadingText;
+
+			/// <summary>
+			/// Constructor
+			/// </summary>
+			public StatusMessage(string headingText, string currentText)
+			{
+				_headingText = headingText;
+				_currentText = currentText;
+			}
+		}
+
+		private readonly LoggerScopeCollection _scopeCollection = new();
+		private readonly object _syncObject = new();
+
+		/// <summary>
+		/// Minimum level for outputting messages
+		/// </summary>
+		public LogLevel OutputLevel
+		{
+			get; set;
+		}
+
+		/// <summary>
+		/// Whether to include timestamps on each line of log output
+		/// </summary>
+		public bool IncludeTimestamps
+		{
+			get; set;
+		}
+
+		/// <summary>
+		/// When true, will detect warnings and errors and set the console output color to yellow and red.
+		/// </summary>
+		public bool ColorConsoleOutput
+		{
+			get; set;
+		}
+
+		/// <summary>
+		/// Whether to write JSON to stdout
+		/// </summary>
+		public bool WriteJsonToStdOut
+		{
+			get; set;
+		}
+
+		/// <summary>
+		/// When true, a timestamp will be written to the log file when the first listener is added
+		/// </summary>
+		public bool IncludeStartingTimestamp
+		{
+			get; set;
+		}
+		private bool _includeStartingTimestampWritten = false;
+
+		/// <summary>
+		/// Path to the log file being written to. May be null.
+		/// </summary>
+		public FileReference? OutputFile
+		{
+			get; private set;
+		}
+
+		/// <summary>
+		/// Whether console output is redirected. This prevents writing status updates that rely on moving the cursor.
+		/// </summary>
+		private static bool AllowStatusUpdates { get; set; } = !Console.IsOutputRedirected;
+
+		/// <summary>
+		/// When configured, this tracks time since initialization to prepend a timestamp to each log.
+		/// </summary>
+		private readonly Stopwatch _timer = Stopwatch.StartNew();
+
+		/// <summary>
+		/// Stack of status scope information.
+		/// </summary>
+		private readonly Stack<StatusMessage> _statusMessageStack = new();
+
+		/// <summary>
+		/// The currently visible status text
+		/// </summary>
+		private string _statusText = "";
+
+		/// <summary>
+		/// Parser for transforming legacy log output into structured events
+		/// </summary>
+		public LogEventParser EventParser { get; }
+
+		/// <summary>
+		/// Last time a status message was pushed to the stack
+		/// </summary>
+		private readonly Stopwatch _statusTimer = new();
+
+		/// <summary>
+		/// Background task for writing to files
+		/// </summary>
+		private Task _writeTask;
+
+		/// <summary>
+		/// Channel for new log events
+		/// </summary>
+		private Channel<JsonLogEvent> _eventChannel = Channel.CreateUnbounded<JsonLogEvent>();
+
+		/// <summary>
+		/// Output streams for structured log data
+		/// </summary>
+		private IReadOnlyList<FileStream> _jsonStreams = [];
+
+		/// <summary>
+		/// Constructor
+		/// </summary>
+		public DefaultLogger()
+		{
+			OutputLevel = LogLevel.Debug;
+			ColorConsoleOutput = true;
+			IncludeStartingTimestamp = true;
+			EventParser = new LogEventParser(this);
+
+			bool bJsonStdOut = Environment.CommandLine.Contains("JsonStdOut", StringComparison.InvariantCulture);
+			string? envVar = Environment.GetEnvironmentVariable("UE_LOG_JSON_TO_STDOUT");
+			if (bJsonStdOut || envVar != null && Int32.TryParse(envVar, out int value) && value != 0)
+			{
+				WriteJsonToStdOut = true;
+			}
+
+			_writeTask = Task.Run(() => WriteFilesAsync());
+		}
+
+		/// <inheritdoc/>
+		public void Dispose()
+		{
+			_eventChannel.Writer.TryComplete();
+			_writeTask.Wait();
+		}
+
+		/// <summary>
+		/// Flush the log queue
+		/// </summary>
+		/// <returns></returns>
+		public async Task FlushAsync()
+		{
+			lock (_syncObject)
+			{
+				Channel<JsonLogEvent> prevEventChannel = _eventChannel;
+				_eventChannel = Channel.CreateUnbounded<JsonLogEvent>();
+				prevEventChannel.Writer.TryComplete();
+			}
+
+			await _writeTask;
+			_writeTask = Task.Run(() => WriteFilesAsync());
+		}
+
+		/// <summary>
+		/// Background task to write events to sinks
+		/// </summary>
+		async Task WriteFilesAsync()
+		{
+			byte[] newline = [(byte)'\n'];
+			while (await _eventChannel.Reader.WaitToReadAsync())
+			{
+				IReadOnlyList<FileStream> streams = _jsonStreams;
+
+				JsonLogEvent logEvent;
+				while (_eventChannel.Reader.TryRead(out logEvent))
+				{
+					foreach (FileStream stream in streams)
+					{
+						await stream.WriteAsync(logEvent.Data);
+						await stream.WriteAsync(newline);
+					}
+				}
+
+				foreach (FileStream stream in streams)
+				{
+					await stream.FlushAsync();
+				}
+			}
+		}
+
+		/// <summary>
+		/// Adds a trace listener that writes to a log file
+		/// </summary>
+		/// <param name="name">Listener name</param>
+		/// <param name="outputFile">The file to write to</param>
+		/// <returns>The created trace listener</returns>
+		public TextWriterTraceListener AddFileWriter(string name, FileReference outputFile)
+		{
+			try
+			{
+				OutputFile = outputFile;
+				DirectoryReference.CreateDirectory(outputFile.Directory);
+				TextWriterTraceListener logTraceListener = new(new StreamWriter(outputFile.FullName), name);
+				lock (_syncObject)
+				{
+					Trace.Listeners.Add(logTraceListener);
+					WriteInitialTimestamp();
+
+					List<FileStream> newJsonStreams = [.. _jsonStreams];
+					newJsonStreams.Add(FileReference.Open(outputFile.ChangeExtension(".json"), FileMode.Create, FileAccess.Write, FileShare.Read | FileShare.Delete));
+					_jsonStreams = newJsonStreams;
+				}
+				return logTraceListener;
+			}
+			catch (Exception ex)
+			{
+				throw new Exception($"Error while creating log file \"{outputFile}\"", ex);
+			}
+		}
+
+		/// <summary>
+		/// Adds a <see cref="TraceListener"/> to the collection in a safe manner.
+		/// </summary>
+		/// <param name="traceListener">The <see cref="TraceListener"/> to add.</param>
+		public void AddTraceListener(TraceListener traceListener)
+		{
+			lock (_syncObject)
+			{
+				if (!Trace.Listeners.Contains(traceListener))
+				{
+					Trace.Listeners.Add(traceListener);
+					WriteInitialTimestamp();
+				}
+			}
+		}
+
+		/// <summary>
+		/// Write a timestamp to the log, once. To be called when a new listener is added.
+		/// </summary>
+		private void WriteInitialTimestamp()
+		{
+			if (IncludeStartingTimestamp && !_includeStartingTimestampWritten)
+			{
+				DateTime now = DateTime.Now;
+				this.LogDebug("{Message}", $"Log started at {now} ({now.ToUniversalTime():yyyy-MM-ddTHH\\:mm\\:ssZ})");
+				_includeStartingTimestampWritten = true;
+			}
+		}
+
+		/// <summary>
+		/// Removes a <see cref="TraceListener"/> from the collection in a safe manner.
+		/// </summary>
+		/// <param name="traceListener">The <see cref="TraceListener"/> to remove.</param>
+		public void RemoveTraceListener(TraceListener traceListener)
+		{
+			lock (_syncObject)
+			{
+				if (Trace.Listeners.Contains(traceListener))
+				{
+					Trace.Listeners.Remove(traceListener);
+				}
+			}
+		}
+
+		/// <summary>
+		/// Removes <see cref="StartupTraceListener"/> from the collection in a safe manner.
+		/// </summary>
+		public void RemoveStartupTraceListener()
+		{
+			lock (_syncObject)
+			{
+				List<StartupTraceListener> startupListeners = [.. Trace.Listeners.OfType<StartupTraceListener>()];
+				if (startupListeners.Count != 0)
+				{
+					Trace.Listeners.Remove(startupListeners.First());
+				}
+			}
+		}
+
+		/// <summary>
+		/// Adds a trace listener that writes to a log file.
+		/// If a StartupTraceListener was in use, this function will copy its captured data to the log file(s)
+		/// and remove the startup listener from the list of registered listeners.
+		/// </summary>
+		/// <param name="name">Identifier for the writer</param>
+		/// <param name="outputFile">The file to write to</param>
+		public void AddFileWriterWithoutBackup(string name, FileReference outputFile)
+		{
+			TextWriterTraceListener firstTextWriter = AddFileWriter(name, outputFile);
+
+			lock (_syncObject)
+			{
+				// find the StartupTraceListener in the listeners that was added early on
+				List<StartupTraceListener> startupListeners = [.. Trace.Listeners.OfType<StartupTraceListener>()];
+				if (startupListeners.Count != 0)
+				{
+					StartupTraceListener startupListener = startupListeners.First();
+					startupListener.CopyTo(firstTextWriter);
+					Trace.Listeners.Remove(startupListeners.First());
+				}
+			}
+		}
+
+		/// <summary>
+		/// Determines if a TextWriterTraceListener has been added to the list of trace listeners
+		/// </summary>
+		/// <returns>True if a TextWriterTraceListener has been added</returns>
+		public bool HasFileWriter()
+		{
+			lock (_syncObject)
+			{
+				return Trace.Listeners.OfType<TextWriterTraceListener>().Any();
+			}
+		}
+
+		/// <inheritdoc/>
+		public IDisposable? BeginScope<TState>(TState state) where TState : notnull
+			=> _scopeCollection.BeginScope(state);
+
+		/// <inheritdoc/>
+		public bool IsEnabled(LogLevel logLevel)
+		{
+			return logLevel >= OutputLevel;
+		}
+
+		/// <inheritdoc/>
+		public void Log<TState>(LogLevel logLevel, EventId eventId, TState state, Exception? exception, Func<TState, Exception?, string> formatter)
+		{
+			if (!IsEnabled(logLevel))
+			{
+				return;
+			}
+
+			string[] lines = formatter(state, exception).Split('\n');
+			lock (_syncObject)
+			{
+				// Output to all the other trace listeners
+				string timePrefix = String.Format("[{0:hh\\:mm\\:ss\\.fff}] ", _timer.Elapsed);
+				foreach (TraceListener? listener in Trace.Listeners)
+				{
+					if (listener != null)
+					{
+						if (listener is DefaultTraceListener && logLevel < LogLevel.Information)
+						{
+							continue;
+						}
+
+						string timePrefixActual =
+							IncludeTimestamps &&
+							listener is not DefaultTraceListener // no timestamps when writing to the Visual Studio debug window
+							? timePrefix
+							: String.Empty;
+
+						foreach (string line in lines)
+						{
+							string lineWithTime = timePrefixActual + line;
+							listener.WriteLine(lineWithTime);
+							listener.Flush();
+						}
+					}
+				}
+
+				Activity? activity = Activity.Current;
+
+				JsonLogEvent jsonLogEvent;
+				if (activity == null && _scopeCollection.IsEmpty())
+				{
+					jsonLogEvent = JsonLogEvent.FromLoggerState(logLevel, eventId, state, exception, formatter);
+				}
+				else
+				{
+					LogEvent logEvent = LogEvent.FromState(logLevel, eventId, state, exception, formatter);
+					logEvent.AddProperties(_scopeCollection.GetProperties());
+					if (activity != null)
+					{
+						logEvent.AddProperty("Activity", activity);
+					}
+					jsonLogEvent = new JsonLogEvent(logEvent);
+				}
+				_eventChannel.Writer.TryWrite(jsonLogEvent);
+
+				// Handle the console output separately; we format things differently
+				if (logLevel >= LogLevel.Information)
+				{
+					FlushStatusHeading();
+
+					bool bResetConsoleColor = false;
+					if (ColorConsoleOutput)
+					{
+						if (logLevel == LogLevel.Warning)
+						{
+							Console.ForegroundColor = ConsoleColor.Yellow;
+							bResetConsoleColor = true;
+						}
+						if (logLevel >= LogLevel.Error)
+						{
+							Console.ForegroundColor = ConsoleColor.Red;
+							bResetConsoleColor = true;
+						}
+					}
+					try
+					{
+						if (WriteJsonToStdOut)
+						{
+							Console.WriteLine(Encoding.UTF8.GetString(jsonLogEvent.Data.Span));
+						}
+						else
+						{
+							foreach (string line in lines)
+							{
+								Console.WriteLine(line);
+							}
+						}
+					}
+					catch (IOException)
+					{
+						// Potential file access/sharing issue on std out
+					}
+					finally
+					{
+						// make sure we always put the console color back.
+						if (bResetConsoleColor)
+						{
+							Console.ResetColor();
+						}
+					}
+
+					if (_statusMessageStack.Count > 0 && AllowStatusUpdates)
+					{
+						SetStatusText(_statusMessageStack.Peek()._currentText);
+					}
+				}
+			}
+		}
+
+		/// <summary>
+		/// Flushes the current status text before writing out a new log line or status message
+		/// </summary>
+		void FlushStatusHeading()
+		{
+			if (_statusMessageStack.Count > 0)
+			{
+				StatusMessage currentStatus = _statusMessageStack.Peek();
+				if (currentStatus._headingText.Length > 0 && !currentStatus._hasFlushedHeadingText && AllowStatusUpdates)
+				{
+					SetStatusText(currentStatus._headingText);
+					Console.WriteLine();
+					_statusText = "";
+					currentStatus._hasFlushedHeadingText = true;
+				}
+				else
+				{
+					SetStatusText("");
+				}
+			}
+		}
+
+		/// <summary>
+		/// Enter a scope with the given status message. The message will be written to the console without a newline, allowing it to be updated through subsequent calls to UpdateStatus().
+		/// The message will be written to the log immediately. If another line is written while in a status scope, the initial status message is flushed to the console first.
+		/// </summary>
+		/// <param name="message">The status message</param>
+		[Conditional("TRACE")]
+		public void PushStatus(string message)
+		{
+			lock (_syncObject)
+			{
+				FlushStatusHeading();
+
+				StatusMessage newStatusMessage = new(message, message);
+				_statusMessageStack.Push(newStatusMessage);
+
+				_statusTimer.Restart();
+
+				if (message.Length > 0)
+				{
+					this.LogDebug("{Message}", message);
+					SetStatusText(message);
+				}
+			}
+		}
+
+		/// <summary>
+		/// Updates the current status message. This will overwrite the previous status line.
+		/// </summary>
+		/// <param name="message">The status message</param>
+		[Conditional("TRACE")]
+		public void UpdateStatus(string message)
+		{
+			lock (_syncObject)
+			{
+				StatusMessage currentStatusMessage = _statusMessageStack.Peek();
+				currentStatusMessage._currentText = message;
+
+				if (AllowStatusUpdates || _statusTimer.Elapsed.TotalSeconds > 10.0)
+				{
+					SetStatusText(message);
+					_statusTimer.Restart();
+				}
+			}
+		}
+
+		/// <summary>
+		/// Updates the Pops the top status message from the stack. The mess
+		/// </summary>
+		[Conditional("TRACE")]
+		public void PopStatus()
+		{
+			lock (_syncObject)
+			{
+				StatusMessage currentStatusMessage = _statusMessageStack.Peek();
+				SetStatusText(currentStatusMessage._currentText);
+
+				if (_statusText.Length > 0)
+				{
+					Console.WriteLine();
+					_statusText = "";
+				}
+
+				_statusMessageStack.Pop();
+			}
+		}
+
+		/// <summary>
+		/// Update the status text. For internal use only; does not modify the StatusMessageStack objects.
+		/// </summary>
+		/// <param name="newStatusText">New status text to display</param>
+		private void SetStatusText(string newStatusText)
+		{
+			if (newStatusText.Length > 0)
+			{
+				newStatusText = LogIndent.Current + newStatusText;
+			}
+
+			if (_statusText != newStatusText)
+			{
+				int numCommonChars = 0;
+				while (numCommonChars < _statusText.Length && numCommonChars < newStatusText.Length && _statusText[numCommonChars] == newStatusText[numCommonChars])
+				{
+					numCommonChars++;
+				}
+
+				if (!AllowStatusUpdates && numCommonChars < _statusText.Length && _statusText.Length > 0)
+				{
+					// Prevent writing backspace characters if the console doesn't support it
+					Console.WriteLine();
+					_statusText = "";
+					numCommonChars = 0;
+				}
+
+				StringBuilder text = new();
+				text.Append('\b', _statusText.Length - numCommonChars);
+				text.Append(newStatusText, numCommonChars, newStatusText.Length - numCommonChars);
+				if (newStatusText.Length < _statusText.Length)
+				{
+					int numChars = _statusText.Length - newStatusText.Length;
+					text.Append(' ', numChars);
+					text.Append('\b', numChars);
+				}
+				Console.Write(text.ToString());
+
+				_statusText = newStatusText;
+				_statusTimer.Restart();
+			}
+		}
+	}
+
+	/// <summary>
+	/// Provider for default logger instances
+	/// </summary>
+	public sealed class DefaultLoggerProvider : ILoggerProvider
+	{
+		/// <inheritdoc/>
+		public ILogger CreateLogger(string categoryName)
+		{
+			return new DefaultLogger();
+		}
+
+		/// <inheritdoc/>
+		public void Dispose()
+		{
+		}
+	}
+
+	/// <summary>
+	/// Extension methods to support the default logger
+	/// </summary>
+	public static class DefaultLoggerExtensions
+	{
+		/// <summary>
+		/// Adds a regular Epic logger to the builder
+		/// </summary>
+		/// <param name="builder">Logging builder</param>
+		public static void AddEpicDefault(this ILoggingBuilder builder)
+		{
+			builder.Services.AddSingleton<ILoggerProvider, DefaultLoggerProvider>();
+		}
+	}
+}
diff --git a/Engine/Source/Programs/Shared/EpicGames.Core/Logging/LogBuffer.cs b/Engine/Source/Programs/Shared/EpicGames.Core/Logging/LogBuffer.cs
new file mode 100644
--- /dev/null
+++ b/Engine/Source/Programs/Shared/EpicGames.Core/Logging/LogBuffer.cs
@@ -0,0 +1,147 @@
+// Copyright Epic Games, Inc. All Rights Reserved.
+
+using System;
+using System.Collections.Generic;
+
+namespace EpicGames.Core
+{
+	/// <summary>
+	/// Implementation of <see cref="ILogCursor"/> used to parse events from a stream of lines
+	/// </summary>
+	public class LogBuffer : ILogCursor
+	{
+		int _lineNumber;
+		readonly string?[] _history;
+		int _historyIdx;
+		int _historyCount;
+		readonly List<string?> _nextLines;
+
+		/// <summary>
+		/// Constructor
+		/// </summary>
+		/// <param name="historySize">Number of previous lines to keep in the buffer</param>
+		public LogBuffer(int historySize)
+		{
+			_lineNumber = 1;
+			_history = new string?[historySize];
+			_nextLines = [];
+		}
+
+		/// <summary>
+		/// Whether the buffer needs more data
+		/// </summary>
+		public bool NeedMoreData
+		{
+			get;
+			private set;
+		}
+
+		/// <inheritdoc/>
+		public string? CurrentLine => this[0];
+
+		/// <inheritdoc/>
+		public int CurrentLineNumber => _lineNumber;
+
+		/// <summary>
+		/// Number of forward lines in the buffer
+		/// </summary>
+		public int Length => _nextLines.Count;
+
+		/// <summary>
+		/// Add a new line to the buffer
+		/// </summary>
+		/// <param name="line">The new line</param>
+		public void AddLine(string? line)
+		{
+			_nextLines.Add(line);
+			NeedMoreData = false;
+		}
+
+		/// <summary>
+		/// Add a collection of lines to the buffer
+		/// </summary>
+		/// <param name="lines">The new lines</param>
+		public void AddLines(IEnumerable<string?> lines)
+		{
+			_nextLines.AddRange(lines);
+			NeedMoreData = false;
+		}
+
+		/// <summary>
+		/// Move forward a given number of lines
+		/// </summary>
+		/// <param name="count">Number of lines to advance by</param>
+		public void Advance(int count)
+		{
+			for (int idx = 0; idx < count; idx++)
+			{
+				MoveNext();
+			}
+		}
+
+		/// <summary>
+		/// Move to the next line
+		/// </summary>
+		public void MoveNext()
+		{
+			if (_nextLines.Count == 0)
+			{
+				throw new InvalidOperationException("Attempt to move past end of line buffer");
+			}
+
+			_historyIdx++;
+			if (_historyIdx >= _history.Length)
+			{
+				_historyIdx = 0;
+			}
+			if (_historyCount < _history.Length)
+			{
+				_historyCount++;
+			}
+
+			_history[_historyIdx] = _nextLines[0];
+			_nextLines.RemoveAt(0);
+
+			_lineNumber++;
+		}
+
+		/// <inheritdoc/>
+		public string? this[int offset]
+		{
+			get
+			{
+				if (offset >= 0)
+				{
+					// Add new lines to the buffer
+					if (offset >= _nextLines.Count)
+					{
+						NeedMoreData = true;
+						return null;
+					}
+					return _nextLines[offset];
+				}
+				else if (offset >= -_historyCount)
+				{
+					// Retrieve a line from the history buffer
+					int idx = _historyIdx + 1 + offset;
+					if (idx < 0)
+					{
+						idx += _history.Length;
+					}
+					return _history[idx];
+				}
+				else
+				{
+					// Invalid index
+					return null;
+				}
+			}
+		}
+
+		/// <inheritdoc/>
+		public override string? ToString()
+		{
+			return CurrentLine;
+		}
+	}
+}
diff --git a/Engine/Source/Programs/Shared/EpicGames.Core/Logging/LogEvent.cs b/Engine/Source/Programs/Shared/EpicGames.Core/Logging/LogEvent.cs
new file mode 100644
--- /dev/null
+++ b/Engine/Source/Programs/Shared/EpicGames.Core/Logging/LogEvent.cs
@@ -0,0 +1,774 @@
+// Copyright Epic Games, Inc. All Rights Reserved.
+
+using System;
+using System.Buffers;
+using System.Collections;
+using System.Collections.Generic;
+using System.Diagnostics.CodeAnalysis;
+using System.Globalization;
+using System.Linq;
+using System.Text;
+using System.Text.Json;
+using System.Text.Json.Serialization;
+using Microsoft.Extensions.Logging;
+
+namespace EpicGames.Core
+{
+	/// <summary>
+	/// Static read-only utf8 strings for parsing log events
+	/// </summary>
+	public static class LogEventPropertyName
+	{
+#pragma warning disable CS1591 // Missing documentation
+		public static readonly Utf8String Time = new("time");
+		public static readonly Utf8String Frame = new("frame");
+		public static readonly Utf8String Level = new("level");
+		public static readonly Utf8String Id = new("id");
+		public static readonly Utf8String Line = new("line");
+		public static readonly Utf8String LineCount = new("lineCount");
+		public static readonly Utf8String Message = new("message");
+		public static readonly Utf8String Format = new("format");
+		public static readonly Utf8String Properties = new("properties");
+
+		public static readonly Utf8String Type = new("$type");
+		public static readonly Utf8String Text = new("$text");
+
+		public static readonly Utf8String File = new("file"); // For source file / asset types
+		public static readonly Utf8String Identifier = new("identifier"); // For symbols
+		public static readonly Utf8String RelativePath = new("relativePath");
+		public static readonly Utf8String DepotPath = new("depotPath");
+		public static readonly Utf8String Target = new("target"); // For hyperlinks
+
+		public static readonly Utf8String Exception = new("exception");
+		public static readonly Utf8String Trace = new("trace");
+		public static readonly Utf8String InnerException = new("innerException");
+		public static readonly Utf8String InnerExceptions = new("innerExceptions");
+#pragma warning restore CS1591 // Missing documentation
+	}
+
+	/// <summary>
+	/// Epic representation of a log event. Can be serialized to/from Json for the Horde dashboard, and passed directly through ILogger interfaces.
+	/// </summary>
+	[JsonConverter(typeof(LogEventConverter))]
+	public class LogEvent : IEnumerable<KeyValuePair<string, object?>>
+	{
+		/// <summary>
+		/// Time that the event was emitted
+		/// </summary>
+		public DateTime Time { get; set; }
+
+		/// <summary>
+		/// Frame that the event was emitted, or 0 when not applicable
+		/// </summary>
+		public ulong Frame { get; set; }
+
+		/// <summary>
+		/// The log level
+		/// </summary>
+		public LogLevel Level { get; set; }
+
+		/// <summary>
+		/// Unique id associated with this event. See <see cref="KnownLogEvents"/> for possible values.
+		/// </summary>
+		public EventId Id { get; set; }
+
+		/// <summary>
+		/// Index of the line within a multi-line message
+		/// </summary>
+		public int LineIndex { get; set; }
+
+		/// <summary>
+		/// Number of lines in the message
+		/// </summary>
+		public int LineCount { get; set; }
+
+		/// <summary>
+		/// The formatted message
+		/// </summary>
+		public string Message { get; set; }
+
+		/// <summary>
+		/// Message template string
+		/// </summary>
+		public string? Format { get; set; }
+
+		/// <summary>
+		/// Map of property name to value
+		/// </summary>
+		public IEnumerable<KeyValuePair<string, object?>>? Properties { get; set; }
+
+		/// <summary>
+		/// The exception value
+		/// </summary>
+		public LogException? Exception { get; }
+
+		class MergedPropertyList : IEnumerable<KeyValuePair<string, object?>>
+		{
+			readonly HashSet<string> _names = [];
+			readonly List<KeyValuePair<string, object?>> _properties = [];
+
+			public void AddRange(IEnumerable<KeyValuePair<string, object?>> properties)
+			{
+				foreach (KeyValuePair<string, object?> property in properties)
+				{
+					if (_names.Add(property.Key))
+					{
+						_properties.Add(property);
+					}
+				}
+			}
+
+			/// <inheritdoc/>
+			public IEnumerator<KeyValuePair<string, object?>> GetEnumerator()
+				=> _properties.GetEnumerator();
+
+			/// <inheritdoc/>
+			IEnumerator IEnumerable.GetEnumerator()
+				=> _properties.GetEnumerator();
+		}
+
+		static readonly JsonSerializerOptions s_jsonSerializerOptions = new() { PropertyNamingPolicy = JsonNamingPolicy.CamelCase, DefaultIgnoreCondition = JsonIgnoreCondition.WhenWritingNull };
+
+		/// <summary>
+		/// Constructor
+		/// </summary>
+		public LogEvent(DateTime time, LogLevel level, EventId eventId, string message, string? format, IEnumerable<KeyValuePair<string, object?>>? properties, LogException? exception)
+			: this(time, level, eventId, 0, 1, message, format, properties, exception)
+		{
+		}
+
+		/// <summary>
+		/// Constructor
+		/// </summary>
+		public LogEvent(DateTime time, LogLevel level, EventId eventId, int lineIndex, int lineCount, string message, string? format, IEnumerable<KeyValuePair<string, object?>>? properties, LogException? exception)
+			: this(time, 0, level, eventId, lineIndex, lineCount, message, format, properties, exception)
+		{
+		}
+
+		/// <summary>
+		/// Constructor
+		/// </summary>
+		public LogEvent(DateTime time, ulong frame, LogLevel level, EventId eventId, int lineIndex, int lineCount, string message, string? format, IEnumerable<KeyValuePair<string, object?>>? properties, LogException? exception)
+		{
+			Time = time;
+			Frame = frame;
+			Level = level;
+			Id = eventId;
+			LineIndex = lineIndex;
+			LineCount = lineCount;
+			Message = message;
+			Format = format;
+			Properties = properties;
+			Exception = exception;
+		}
+
+		/// <summary>
+		/// Add a new property to this event
+		/// </summary>
+		/// <param name="name">Name of the property to add</param>
+		/// <param name="value">Value for the property</param>
+		public void AddProperty(string name, object? value)
+			=> AddProperties([KeyValuePair.Create(name, value)]);
+
+		/// <summary>
+		/// Add new properties to this event
+		/// </summary>
+		/// <param name="properties">Properties to add</param>
+		public void AddProperties(IEnumerable<KeyValuePair<string, object?>> properties)
+		{
+			if (Properties == null)
+			{
+				Properties = properties;
+			}
+			else
+			{
+				MergedPropertyList? list = Properties as MergedPropertyList;
+				if (list == null)
+				{
+					list = new MergedPropertyList();
+					list.AddRange(Properties);
+					Properties = list;
+				}
+				list.AddRange(properties);
+			}
+		}
+
+		/// <summary>
+		/// Gets an untyped property with the given name
+		/// </summary>
+		/// <param name="name"></param>
+		/// <returns></returns>
+		public object GetProperty(string name)
+		{
+			object? value;
+			if (TryGetProperty(name, out value))
+			{
+				return value;
+			}
+			throw new KeyNotFoundException($"Property {name} not found");
+		}
+
+		/// <summary>
+		/// Gets a property with the given name
+		/// </summary>
+		/// <typeparam name="T"></typeparam>
+		/// <param name="name">Name of the property</param>
+		/// <returns></returns>
+		public T GetProperty<T>(string name) => (T)GetProperty(name);
+
+		/// <summary>
+		/// Finds a property with the given name
+		/// </summary>
+		/// <param name="name">Name of the property</param>
+		/// <param name="value">Value for the property, on success</param>
+		/// <returns>True if the property was found, false otherwise</returns>
+		public bool TryGetProperty(string name, [NotNullWhen(true)] out object? value)
+		{
+			if (Properties != null)
+			{
+				foreach (KeyValuePair<string, object?> pair in Properties)
+				{
+					if (pair.Key.Equals(name, StringComparison.Ordinal) && pair.Value != null)
+					{
+						value = pair.Value;
+						return true;
+					}
+				}
+			}
+
+			value = null;
+			return false;
+		}
+
+		/// <summary>
+		/// Finds a typed property with the given name
+		/// </summary>
+		/// <typeparam name="T">Type of the property to receive</typeparam>
+		/// <param name="name">Name of the property</param>
+		/// <param name="value">Value for the property, on success</param>
+		/// <returns>True if the property was found, false otherwise</returns>
+		public bool TryGetProperty<T>(string name, [NotNullWhen(true)] out T value)
+		{
+			object? untypedValue;
+			if (TryGetProperty(name, out untypedValue) && untypedValue is T typedValue)
+			{
+				value = typedValue;
+				return true;
+			}
+			else
+			{
+				value = default!;
+				return false;
+			}
+		}
+
+		/// <summary>
+		/// Read a log event from a utf-8 encoded json byte array
+		/// </summary>
+		/// <param name="data"></param>
+		/// <returns></returns>
+		public static LogEvent Read(ReadOnlySpan<byte> data)
+		{
+			Utf8JsonReader reader = new(data);
+			reader.Read();
+			return Read(ref reader);
+		}
+
+		/// <summary>
+		/// Read a log event from Json
+		/// </summary>
+		/// <param name="reader">The Json reader</param>
+		/// <returns>New log event</returns>
+#pragma warning disable CA1045 // Do not pass types by reference
+		public static LogEvent Read(ref Utf8JsonReader reader)
+#pragma warning restore CA1045 // Do not pass types by reference
+		{
+			DateTime time = new(0);
+			ulong frame = 0;
+			LogLevel level = LogLevel.None;
+			EventId eventId = new(0);
+			int line = 0;
+			int lineCount = 1;
+			string message = String.Empty;
+			string? format = null;
+			Dictionary<string, object?>? properties = null;
+			LogException? exception = null;
+
+			ReadOnlySpan<byte> propertyName;
+			for (; JsonExtensions.TryReadNextPropertyName(ref reader, out propertyName); reader.Skip())
+			{
+				if (Utf8StringComparer.OrdinalIgnoreCase.Equals(propertyName, LogEventPropertyName.Time.Span))
+				{
+					time = reader.GetDateTime();
+				}
+				else if (Utf8StringComparer.OrdinalIgnoreCase.Equals(propertyName, LogEventPropertyName.Frame.Span))
+				{
+					frame = reader.GetUInt64();
+				}
+				else if (Utf8StringComparer.OrdinalIgnoreCase.Equals(propertyName, LogEventPropertyName.Level.Span))
+				{
+					level = Enum.Parse<LogLevel>(reader.GetString()!);
+				}
+				else if (Utf8StringComparer.OrdinalIgnoreCase.Equals(propertyName, LogEventPropertyName.Id.Span))
+				{
+					eventId = reader.GetInt32();
+				}
+				else if (Utf8StringComparer.OrdinalIgnoreCase.Equals(propertyName, LogEventPropertyName.Line.Span))
+				{
+					line = reader.GetInt32();
+				}
+				else if (Utf8StringComparer.OrdinalIgnoreCase.Equals(propertyName, LogEventPropertyName.LineCount.Span))
+				{
+					lineCount = reader.GetInt32();
+				}
+				else if (Utf8StringComparer.OrdinalIgnoreCase.Equals(propertyName, LogEventPropertyName.Message.Span))
+				{
+					message = reader.GetString()!;
+				}
+				else if (Utf8StringComparer.OrdinalIgnoreCase.Equals(propertyName, LogEventPropertyName.Format.Span))
+				{
+					format = reader.GetString()!;
+				}
+				else if (Utf8StringComparer.OrdinalIgnoreCase.Equals(propertyName, LogEventPropertyName.Properties.Span))
+				{
+					properties = ReadProperties(ref reader);
+				}
+			}
+
+			return new LogEvent(time, frame, level, eventId, line, lineCount, message, format, properties, exception);
+		}
+
+		static Dictionary<string, object?> ReadProperties(ref Utf8JsonReader reader)
+		{
+			Dictionary<string, object?> properties = [];
+
+			ReadOnlySpan<byte> propertyName;
+			for (; JsonExtensions.TryReadNextPropertyName(ref reader, out propertyName); reader.Skip())
+			{
+				string name = Encoding.UTF8.GetString(propertyName);
+
+				object? value = ReadPropertyValue(ref reader);
+				properties.Add(name, value);
+			}
+
+			return properties;
+		}
+
+		static object? ReadPropertyValue(ref Utf8JsonReader reader)
+		{
+			switch (reader.TokenType)
+			{
+				case JsonTokenType.Null:
+					return null;
+				case JsonTokenType.True:
+					return true;
+				case JsonTokenType.False:
+					return false;
+				case JsonTokenType.StartArray:
+					return ReadArrayPropertyValue(ref reader);
+				case JsonTokenType.StartObject:
+					return ReadObjectPropertyValue(ref reader);
+				case JsonTokenType.String:
+					return reader.GetString()!;
+				case JsonTokenType.Number:
+					if (reader.TryGetInt32(out int intValue))
+					{
+						return intValue;
+					}
+					else if (reader.TryGetDouble(out double doubleValue))
+					{
+						return doubleValue;
+					}
+					else
+					{
+						return reader.GetString()!;
+					}
+				default:
+					throw new InvalidOperationException("Unhandled property type");
+			}
+		}
+
+		static List<object?> ReadArrayPropertyValue(ref Utf8JsonReader reader)
+		{
+			List<object?> result = [];
+			while (reader.Read() && reader.TokenType != JsonTokenType.EndArray)
+			{
+				result.Add(ReadPropertyValue(ref reader));
+			}
+			return result;
+		}
+
+		static object ReadObjectPropertyValue(ref Utf8JsonReader reader)
+		{
+			// Read all the properties
+			Dictionary<Utf8String, object?>? properties = [];
+			for (; JsonExtensions.TryReadNextPropertyName(ref reader, out ReadOnlySpan<byte> propertyName); reader.Skip())
+			{
+				properties.Add(new Utf8String(propertyName.ToArray()), ReadPropertyValue(ref reader));
+			}
+
+			// Check if we can convert it to a LogValue
+			if (properties.TryGetValue(LogEventPropertyName.Type, out object? type) && type is string typeStr && properties.TryGetValue(LogEventPropertyName.Text, out object? text) && text is string textStr)
+			{
+				properties.Remove(LogEventPropertyName.Type);
+				properties.Remove(LogEventPropertyName.Text);
+
+				if (properties.Count == 0)
+				{
+					return new LogValue(new Utf8String(typeStr), textStr);
+				}
+				else
+				{
+					return new LogValue(new Utf8String(typeStr), textStr, properties);
+				}
+			}
+
+			return properties;
+		}
+
+		/// <summary>
+		/// Writes a log event to Json
+		/// </summary>
+		/// <param name="writer"></param>
+		public void Write(Utf8JsonWriter writer)
+		{
+			writer.WriteStartObject();
+			writer.WriteString(LogEventPropertyName.Time.Span, Time.ToString("s", CultureInfo.InvariantCulture));
+
+			if (Frame != 0)
+			{
+				writer.WriteNumber(LogEventPropertyName.Frame.Span, Frame);
+			}
+
+			writer.WriteString(LogEventPropertyName.Level.Span, Level.ToString());
+			writer.WriteString(LogEventPropertyName.Message.Span, Message);
+
+			if (Id.Id != 0)
+			{
+				writer.WriteNumber(LogEventPropertyName.Id.Span, Id.Id);
+			}
+
+			if (LineIndex > 0)
+			{
+				writer.WriteNumber(LogEventPropertyName.Line.Span, LineIndex);
+			}
+
+			if (LineCount > 1)
+			{
+				writer.WriteNumber(LogEventPropertyName.LineCount.Span, LineCount);
+			}
+
+			if (Format != null)
+			{
+				writer.WriteString(LogEventPropertyName.Format.Span, Format);
+			}
+
+			if (Properties != null && Properties.Any())
+			{
+				writer.WriteStartObject(LogEventPropertyName.Properties.Span);
+				foreach ((string name, object? value) in Properties!)
+				{
+					if (!name.Equals(MessageTemplate.FormatPropertyName, StringComparison.Ordinal))
+					{
+						if (name.StartsWith('@'))
+						{
+							writer.WritePropertyName(name[1..]);
+							JsonSerializer.Serialize(writer, value, value?.GetType() ?? typeof(object), s_jsonSerializerOptions);
+						}
+						else
+						{
+							writer.WritePropertyName(name);
+							LogValueFormatter.Format(value, writer);
+						}
+					}
+				}
+				writer.WriteEndObject();
+			}
+
+			if (Exception != null)
+			{
+				writer.WriteStartObject(LogEventPropertyName.Exception.Span);
+				WriteException(ref writer, Exception);
+				writer.WriteEndObject();
+			}
+			writer.WriteEndObject();
+		}
+
+		/// <summary>
+		/// Writes an exception to a json object
+		/// </summary>
+		/// <param name="writer">Writer to receive the exception data</param>
+		/// <param name="exception">The exception</param>
+		static void WriteException(ref Utf8JsonWriter writer, LogException exception)
+		{
+			writer.WriteString("message", exception.Message);
+			writer.WriteString("trace", exception.Trace);
+
+			if (exception.InnerException != null)
+			{
+				writer.WriteStartObject("innerException");
+				WriteException(ref writer, exception.InnerException);
+				writer.WriteEndObject();
+			}
+
+			if (exception.InnerExceptions != null)
+			{
+				writer.WriteStartArray("innerExceptions");
+				for (int idx = 0; idx < 16 && idx < exception.InnerExceptions.Count; idx++) // Cap number of exceptions returned to avoid huge messages
+				{
+					LogException innerException = exception.InnerExceptions[idx];
+					writer.WriteStartObject();
+					WriteException(ref writer, innerException);
+					writer.WriteEndObject();
+				}
+				writer.WriteEndArray();
+			}
+		}
+
+		/// <summary>
+		/// Create a new log event
+		/// </summary>
+		public static LogEvent Create(LogLevel level, string format, params object[] args)
+			=> Create(level, KnownLogEvents.None, null, format, args);
+
+		/// <summary>
+		/// Create a new log event
+		/// </summary>
+		public static LogEvent Create(LogLevel level, EventId eventId, string format, params object[] args)
+			=> Create(level, eventId, null, format, args);
+
+		/// <summary>
+		/// Create a new log event
+		/// </summary>
+		public static LogEvent Create(LogLevel level, EventId eventId, Exception? exception, string format, params object[] args)
+		{
+			Dictionary<string, object?> properties = [];
+			MessageTemplate.ParsePropertyValues(format, args, properties);
+
+			string message = MessageTemplate.Render(format, properties!);
+			return new LogEvent(DateTime.UtcNow, level, eventId, message, format, properties, LogException.FromException(exception));
+		}
+
+		/// <summary>
+		/// Creates a log event from an ILogger parameters
+		/// </summary>
+		public static LogEvent FromState<TState>(LogLevel level, EventId eventId, TState state, Exception? exception, Func<TState, Exception?, string> formatter)
+		{
+			_ = formatter;
+
+			if (state is LogEvent logEvent)
+			{
+				return logEvent;
+			}
+			if (state is JsonLogEvent jsonLogEvent)
+			{
+				return Read(jsonLogEvent.Data.Span);
+			}
+
+			DateTime time = DateTime.UtcNow;
+
+			// Try to log the event
+			IEnumerable<KeyValuePair<string, object?>>? values = state as IEnumerable<KeyValuePair<string, object?>>;
+			string? format = values?.FirstOrDefault(x => x.Key.Equals(MessageTemplate.FormatPropertyName, StringComparison.Ordinal)).Value?.ToString();
+			string message = (format == null) ? formatter(state, exception) : MessageTemplate.Render(format, values);
+			return new LogEvent(time, level, eventId, message, format, values, LogException.FromException(exception));
+		}
+
+		/// <summary>
+		/// Enumerates all the properties in this object
+		/// </summary>
+		/// <returns>Property pairs</returns>
+		public IEnumerator<KeyValuePair<string, object?>> GetEnumerator()
+		{
+			if (Format != null)
+			{
+				yield return new KeyValuePair<string, object?>(MessageTemplate.FormatPropertyName, Format.ToString());
+			}
+
+			if (Properties != null)
+			{
+				foreach ((string name, object? value) in Properties)
+				{
+					if (!name.Equals(MessageTemplate.FormatPropertyName, StringComparison.OrdinalIgnoreCase))
+					{
+						yield return new KeyValuePair<string, object?>(name, value?.ToString());
+					}
+				}
+			}
+		}
+
+		/// <summary>
+		/// Enumerates all the properties in this object
+		/// </summary>
+		/// <returns>Property pairs</returns>
+		System.Collections.IEnumerator System.Collections.IEnumerable.GetEnumerator()
+		{
+			foreach (KeyValuePair<string, object?> pair in this)
+			{
+				yield return pair;
+			}
+		}
+
+		/// <summary>
+		/// Serialize a message template to JOSN
+		/// </summary>
+		public byte[] ToJsonBytes()
+		{
+			ArrayBufferWriter<byte> buffer = new();
+			using (Utf8JsonWriter writer = new(buffer))
+			{
+				Write(writer);
+			}
+			return buffer.WrittenSpan.ToArray();
+		}
+
+		/// <summary>
+		/// Serialize a message template to JOSN
+		/// </summary>
+		public string ToJson()
+		{
+			ArrayBufferWriter<byte> buffer = new();
+			using (Utf8JsonWriter writer = new(buffer))
+			{
+				Write(writer);
+			}
+			return Encoding.UTF8.GetString(buffer.WrittenSpan);
+		}
+
+		/// <inheritdoc/>
+		public override string ToString() => Message;
+	}
+
+	/// <summary>
+	/// Information about an exception in a log event
+	/// </summary>
+	public sealed class LogException
+	{
+		/// <summary>
+		/// Exception message
+		/// </summary>
+		public string Message { get; set; }
+
+		/// <summary>
+		/// Stack trace for the exception
+		/// </summary>
+		public string Trace { get; set; }
+
+		/// <summary>
+		/// Optional inner exception information
+		/// </summary>
+		public LogException? InnerException { get; set; }
+
+		/// <summary>
+		/// Multiple inner exceptions, in the case of an <see cref="AggregateException"/>
+		/// </summary>
+		public List<LogException> InnerExceptions { get; } = [];
+
+		/// <summary>
+		/// Constructor
+		/// </summary>
+		/// <param name="message"></param>
+		/// <param name="trace"></param>
+		public LogException(string message, string trace)
+		{
+			Message = message;
+			Trace = trace;
+		}
+
+		/// <summary>
+		/// Constructor
+		/// </summary>
+		/// <param name="exception"></param>
+		[return: NotNullIfNotNull(nameof(exception))]
+		public static LogException? FromException(Exception? exception)
+		{
+			LogException? result = null;
+			if (exception != null)
+			{
+				result = new LogException(exception.Message, exception.StackTrace ?? String.Empty);
+
+				if (exception.InnerException != null)
+				{
+					result.InnerException = FromException(exception.InnerException);
+				}
+
+				AggregateException? aggregateException = exception as AggregateException;
+				if (aggregateException != null && aggregateException.InnerExceptions.Count > 0)
+				{
+					for (int idx = 0; idx < 16 && idx < aggregateException.InnerExceptions.Count; idx++) // Cap number of exceptions returned to avoid huge messages
+					{
+						LogException innerException = FromException(aggregateException.InnerExceptions[idx]);
+						result.InnerExceptions.Add(innerException);
+					}
+				}
+			}
+			return result;
+		}
+	}
+
+	/// <summary>
+	/// Interface for a log event sink
+	/// </summary>
+	public interface ILogEventSink
+	{
+		/// <summary>
+		/// Process log event 
+		/// </summary>
+		/// <param name="logEvent">Log event</param>
+		void ProcessEvent(LogEvent logEvent);
+	}
+
+	/// <summary>
+	/// Simple filtering log event sink with a callback for convenience
+	/// </summary>
+	public class FilteringEventSink : ILogEventSink
+	{
+		/// <summary>
+		/// Log events received
+		/// </summary>
+		public IReadOnlyList<LogEvent> LogEvents => _logEvents;
+
+		private readonly List<LogEvent> _logEvents = [];
+		private readonly int[] _includeEventIds;
+		private readonly Action<LogEvent>? _eventCallback;
+
+		/// <summary>
+		/// Constructor
+		/// </summary>
+		/// <param name="includeEventIds">Event IDs to include</param>
+		/// <param name="eventCallback">Optional callback function for each event received</param>
+		public FilteringEventSink(int[] includeEventIds, Action<LogEvent>? eventCallback = null)
+		{
+			_includeEventIds = includeEventIds;
+			_eventCallback = eventCallback;
+		}
+
+		/// <inheritdoc/>
+		public void ProcessEvent(LogEvent logEvent)
+		{
+			if (_includeEventIds.Any(x => x == logEvent.Id.Id))
+			{
+				_logEvents.Add(logEvent);
+				_eventCallback?.Invoke(logEvent);
+			}
+		}
+	}
+
+	/// <summary>
+	/// Converter for serialization of <see cref="LogEvent"/> instances to Json streams
+	/// </summary>
+	public class LogEventConverter : JsonConverter<LogEvent>
+	{
+		/// <inheritdoc/>
+		public override LogEvent Read(ref Utf8JsonReader reader, Type typeToConvert, JsonSerializerOptions options)
+		{
+			return LogEvent.Read(ref reader);
+		}
+
+		/// <inheritdoc/>
+		public override void Write(Utf8JsonWriter writer, LogEvent value, JsonSerializerOptions options)
+		{
+			value.Write(writer);
+		}
+	}
+}
diff --git a/Engine/Source/Programs/Shared/EpicGames.Core/Logging/LogEventBuilder.cs b/Engine/Source/Programs/Shared/EpicGames.Core/Logging/LogEventBuilder.cs
new file mode 100644
--- /dev/null
+++ b/Engine/Source/Programs/Shared/EpicGames.Core/Logging/LogEventBuilder.cs
@@ -0,0 +1,507 @@
+// Copyright Epic Games, Inc. All Rights Reserved.
+
+using System;
+using System.Collections.Generic;
+using System.Text;
+using System.Text.RegularExpressions;
+using Microsoft.Extensions.Logging;
+
+namespace EpicGames.Core
+{
+	/// <summary>
+	/// Stores information about a span within a line
+	/// </summary>
+	public class LogEventSpan
+	{
+		/// <summary>
+		/// Starting offset within the line
+		/// </summary>
+		public int Offset
+		{
+			get;
+		}
+
+		/// <summary>
+		/// Text for this span
+		/// </summary>
+		public string Text
+		{
+			get;
+		}
+
+		/// <summary>
+		/// Storage for properties
+		/// </summary>
+		public Dictionary<string, object> Properties
+		{
+			get;
+		}
+
+		/// <summary>
+		/// Constructor
+		/// </summary>
+		/// <param name="offset">Starting offset within the line</param>
+		/// <param name="text">The text for this span</param>
+		public LogEventSpan(int offset, string text)
+		{
+			Offset = offset;
+			Text = text;
+			Properties = new Dictionary<string, object>(StringComparer.OrdinalIgnoreCase);
+		}
+
+		/// <summary>
+		/// Converts this object to a string. This determines how the span will be rendered by the default console logger, so should return the original text.
+		/// </summary>
+		/// <returns>Original text for this span</returns>
+		public override string ToString()
+		{
+			return Text;
+		}
+	}
+
+	/// <summary>
+	/// Individual line in the log output
+	/// </summary>
+	public class LogEventLine
+	{
+		/// <summary>
+		/// The raw text
+		/// </summary>
+		public string Text
+		{
+			get;
+		}
+
+		/// <summary>
+		/// List of spans for markup
+		/// </summary>
+		public Dictionary<string, LogEventSpan> Spans
+		{
+			get;
+		}
+
+		/// <summary>
+		/// Constructor
+		/// </summary>
+		/// <param name="text">Text for the line</param>
+		public LogEventLine(string text)
+		{
+			Text = text;
+			Spans = [];
+		}
+
+		/// <summary>
+		/// Adds a span containing markup on the source text
+		/// </summary>
+		/// <param name="offset">Offset within the line</param>
+		/// <param name="length">Length of the span</param>
+		/// <param name="name">Name to use to identify the item in the format string</param>
+		/// <returns>New span for the given range</returns>
+		public LogEventSpan AddSpan(int offset, int length, string name)
+		{
+			LogEventSpan span = new(offset, Text.Substring(offset, length));
+			Spans.Add(name, span);
+			return span;
+		}
+
+		/// <summary>
+		/// Adds a span containing markup for a regex match group
+		/// </summary>
+		/// <param name="group">The match group</param>
+		/// <param name="name">Name to use to identify the item in the format string</param>
+		public LogEventSpan AddSpan(Group group, string name)
+		{
+			return AddSpan(group.Index, group.Length, name);
+		}
+
+		/// <summary>
+		/// Adds a span naming a regex match group, using the name of the group
+		/// </summary>
+		/// <param name="group">The match group</param>
+		public LogEventSpan AddSpan(Group group)
+		{
+			return AddSpan(group.Index, group.Length, group.Name);
+		}
+
+		/// <summary>
+		/// Adds a span naming a regex match group, using the name of the group
+		/// </summary>
+		/// <param name="group">The match group</param>
+		/// <param name="name">Name to use to identify the item in the format string</param>
+		public LogEventSpan? TryAddSpan(Group group, string name)
+		{
+			if (group.Success)
+			{
+				return AddSpan(group, name);
+			}
+			else
+			{
+				return null;
+			}
+		}
+
+		/// <summary>
+		/// Adds a span naming a regex match group, using the name of the group
+		/// </summary>
+		/// <param name="group">The match group</param>
+		public LogEventSpan? TryAddSpan(Group group)
+		{
+			if (group.Success)
+			{
+				return AddSpan(group);
+			}
+			else
+			{
+				return null;
+			}
+		}
+
+		/// <inheritdoc/>
+		public override string ToString()
+		{
+			return Text;
+		}
+	}
+
+	/// <summary>
+	/// Allows building log events by annotating a window of lines around the current cursor position.
+	/// </summary>
+	public class LogEventBuilder
+	{
+		class LogSpan
+		{
+			public string _name;
+			public int _offset;
+			public int _length;
+			public object? _value;
+
+			public LogSpan(string name, int offset, int length, object? value)
+			{
+				_name = name;
+				_offset = offset;
+				_length = length;
+				_value = value;
+			}
+		}
+
+		class LogLine
+		{
+			public string _message;
+			public string? _format;
+			public Dictionary<string, object>? _properties;
+
+			public LogLine(string message, string? format, Dictionary<string, object>? properties)
+			{
+				_message = message;
+				_format = format;
+				_properties = properties;
+			}
+		}
+
+		/// <summary>
+		/// The current cursor position
+		/// </summary>
+		public ILogCursor Current { get; private set; }
+
+		/// <summary>
+		/// The next cursor position
+		/// </summary>
+		public ILogCursor Next { get; private set; }
+
+		/// <summary>
+		/// Events which have been parsed so far
+		/// </summary>
+		List<LogLine>? _lines;
+
+		/// <summary>
+		/// Spans for the current line
+		/// </summary>
+		List<LogSpan>? _spans;
+
+		/// <summary>
+		/// Additional properties for this line
+		/// </summary>
+		Dictionary<string, object>? _properties;
+
+		/// <summary>
+		/// Starts building a log event at the current cursor position
+		/// </summary>
+		/// <param name="cursor">The current cursor position</param>
+		/// <param name="lineCount">Number of lines to consume</param>
+		public LogEventBuilder(ILogCursor cursor, int lineCount = 1)
+		{
+			Current = cursor;
+			Next = cursor.Rebase(1);
+
+			if (lineCount > 1)
+			{
+				MoveNext(lineCount - 1);
+			}
+		}
+
+		/// <summary>
+		/// Creates a log event from the current line
+		/// </summary>
+		/// <returns></returns>
+		LogLine CreateLine()
+		{
+			int offset = 0;
+			string currentLine = Current.CurrentLine!;
+
+			string? format = null;
+			if (_spans != null)
+			{
+				StringBuilder builder = new();
+				foreach (LogSpan span in _spans)
+				{
+					if (span._offset >= offset)
+					{
+						builder.Append(currentLine, offset, span._offset - offset);
+						builder.Append($"{{{span._name}}}");
+						offset = span._offset + span._length;
+					}
+				}
+				builder.Append(currentLine, offset, currentLine.Length - offset);
+				format = builder.ToString();
+			}
+
+			return new LogLine(currentLine, format, _properties);
+		}
+
+		/// <summary>
+		/// Adds a span containing markup on the source text
+		/// </summary>
+		/// <param name="name">Name to use to identify the item in the format string</param>
+		/// <param name="offset">Offset within the line</param>
+		/// <param name="length">Length of the span</param>
+		/// <param name="value">Data of the span</param>
+		/// <returns>New span for the given range</returns>
+		public void Annotate(string name, int offset, int length, object? value = null)
+		{
+			LogSpan span = new(name, offset, length, value);
+
+			_properties ??= new Dictionary<string, object>(StringComparer.OrdinalIgnoreCase);
+			_properties.Add(name, span);
+
+			_spans ??= [];
+			for (int insertIdx = _spans.Count; ; insertIdx--)
+			{
+				if (insertIdx == 0 || _spans[insertIdx - 1]._offset < offset)
+				{
+					_spans.Insert(insertIdx, span);
+					break;
+				}
+			}
+		}
+
+		/// <summary>
+		/// Adds a span containing markup for a regex match group
+		/// </summary>
+		/// <param name="group">The match group</param>
+		/// <param name="name">Name to use to identify the item in the format string</param>
+		/// <param name="value">Optional value for the annotation</param>
+		public void Annotate(string name, Group group, object? value = null)
+		{
+			Annotate(name, group.Index, group.Length, value);
+		}
+
+		/// <summary>
+		/// Adds a span naming a regex match group, using the name of the group
+		/// </summary>
+		/// <param name="group">The match group</param>
+		/// <param name="value">Optional value for the annotation</param>
+		public void Annotate(Group group, object? value = null)
+		{
+			Annotate(group.Name, group.Index, group.Length, value);
+		}
+
+		/// <summary>
+		/// Adds a span naming a regex match group, using the name of the group
+		/// </summary>
+		/// <param name="group">The match group</param>
+		/// <param name="name">Name to use to identify the item in the format string</param>
+		/// <param name="value">Optional value for the annotation</param>
+		public bool TryAnnotate(string name, Group group, object? value = null)
+		{
+			if (group.Success)
+			{
+				Annotate(name, group, value);
+				return true;
+			}
+			return false;
+		}
+
+		/// <summary>
+		/// Adds a span naming a regex match group, using the name of the group
+		/// </summary>
+		/// <param name="group">The match group</param>
+		/// <param name="value">Optional value for the annotation</param>
+		public bool TryAnnotate(Group group, object? value = null)
+		{
+			if (group.Success)
+			{
+				Annotate(group, value);
+				return true;
+			}
+			return false;
+		}
+
+		/// <summary>
+		/// Adds an additional named property
+		/// </summary>
+		/// <param name="name">Name of the argument</param>
+		/// <param name="value">Value to associate with it</param>
+		public void AddProperty(string name, object value)
+		{
+			_properties ??= new Dictionary<string, object>(StringComparer.OrdinalIgnoreCase);
+			_properties.Add(name, value);
+		}
+
+		string? GetFirstLine()
+		{
+			if (_lines == null || _lines.Count == 0)
+			{
+				return Current.CurrentLine;
+			}
+			else
+			{
+				return _lines[0]._message;
+			}
+		}
+
+		/// <summary>
+		/// Check if the next line is aligned or indented from the first line
+		/// </summary>
+		/// <returns>True if the next line is aligned with this</returns>
+		public bool IsNextLineAligned() => Next.IsAligned(0, GetFirstLine());
+
+		/// <summary>
+		/// Check if the next line is indented from the first line
+		/// </summary>
+		/// <returns>True if the next line is aligned with this</returns>
+		public bool IsNextLineHanging() => Next.IsHanging(0, GetFirstLine());
+
+		/// <summary>
+		/// Complete the current line and move to the next
+		/// </summary>
+		public void MoveNext()
+		{
+			_lines ??= [];
+			_lines.Add(CreateLine());
+
+			_spans = null;
+			_properties = null;
+
+			Current = Next;
+			Next = Next.Rebase(1);
+		}
+
+		/// <summary>
+		/// Advance by the given number of lines
+		/// </summary>
+		/// <param name="count"></param>
+		public void MoveNext(int count)
+		{
+			for (int idx = 0; idx < count; idx++)
+			{
+				MoveNext();
+			}
+		}
+
+		/// <summary>
+		/// Returns an array of log events
+		/// </summary>
+		/// <returns></returns>
+		public LogEvent[] ToArray(LogLevel level, EventId eventId)
+		{
+			DateTime time = DateTime.UtcNow;
+
+			int numLines = _lines?.Count ?? 0;
+			int numEvents = numLines;
+			if (Current.CurrentLine != null)
+			{
+				numEvents++;
+			}
+
+			LogEvent[] events = new LogEvent[numEvents];
+			for (int idx = 0; idx < numLines; idx++)
+			{
+				events[idx] = CreateEvent(time, level, eventId, idx, numEvents, _lines![idx]);
+			}
+			if (Current.CurrentLine != null)
+			{
+				events[numEvents - 1] = CreateEvent(time, level, eventId, numEvents - 1, numEvents, CreateLine());
+			}
+
+			return events;
+		}
+
+		static LogEvent CreateEvent(DateTime time, LogLevel level, EventId eventId, int lineIndex, int lineCount, LogLine line)
+		{
+			Dictionary<string, object?>? properties = null;
+			if (line._properties != null)
+			{
+				properties = [];
+				foreach ((string name, object value) in line._properties)
+				{
+					object newValue;
+					if (value is LogSpan span)
+					{
+						string text = line._message.Substring(span._offset, span._length);
+						if (span._value == null)
+						{
+							newValue = text;
+						}
+						else if (span._value is LogValue newLogValue)
+						{
+							newLogValue.Text = text;
+							newValue = newLogValue;
+						}
+						else
+						{
+							newValue = text;
+						}
+					}
+					else
+					{
+						newValue = value;
+					}
+					properties[name] = newValue;
+				}
+			}
+			return new LogEvent(time, level, eventId, lineIndex, lineCount, line._message, line._format, properties, null);
+		}
+
+		/// <summary>
+		/// Creates a match object at the given priority
+		/// </summary>
+		/// <param name="level"></param>
+		/// <param name="eventId">The event id</param>
+		/// <param name="priority"></param>
+		/// <param name="isComplete"></param>
+		/// <returns></returns>
+		public LogEventMatch ToMatch(LogEventPriority priority, LogLevel level, EventId eventId, bool isComplete = true)
+		{
+			return new LogEventMatch(priority, ToArray(level, eventId), isComplete);
+		}
+		
+		/// <summary>
+		/// Creates a match object at the given priority with hanging lines
+		/// </summary>
+		/// <param name="level"></param>
+		/// <param name="eventId">The event id</param>
+		/// <param name="priority"></param>
+		/// <returns></returns>
+		public LogEventMatch ToMatchWithHangingLines(LogEventPriority priority, LogLevel level, EventId eventId)
+		{
+			// If there are more lines that are not hanging, is must be complete
+			bool isComplete = Next.Length > 0;
+			for (; isComplete && IsNextLineHanging(); isComplete = Next.Length > 0)
+			{
+				MoveNext();
+			}
+
+			return new LogEventMatch(priority, ToArray(level, eventId), isComplete);
+		}
+	}
+}
diff --git a/Engine/Source/Programs/Shared/EpicGames.Core/Logging/LogEventMarkup.cs b/Engine/Source/Programs/Shared/EpicGames.Core/Logging/LogEventMarkup.cs
new file mode 100644
--- /dev/null
+++ b/Engine/Source/Programs/Shared/EpicGames.Core/Logging/LogEventMarkup.cs
@@ -0,0 +1,113 @@
+// Copyright Epic Games, Inc. All Rights Reserved.
+
+using System;
+using System.Collections.Generic;
+using System.IO;
+using System.Text.RegularExpressions;
+
+namespace EpicGames.Core
+{
+	/// <summary>
+	/// Extension methods to allow adding markup to log event spans
+	/// </summary>
+	public static class LogEventMarkup
+	{
+#pragma warning disable CS1591
+		public static LogValue Channel => new(LogValueType.Channel, "");
+		public static LogValue Severity => new(LogValueType.Severity, "");
+		public static LogValue Message => new(LogValueType.Message, "");
+		public static LogValue LineNumber => new(LogValueType.LineNumber, "");
+		public static LogValue ColumnNumber => new(LogValueType.ColumnNumber, "");
+		public static LogValue Symbol => new(LogValueType.Symbol, "");
+		public static LogValue ErrorCode => new(LogValueType.ErrorCode, "");
+		public static LogValue ToolName => new(LogValueType.ToolName, "");
+		public static LogValue ScreenshotTest => new(LogValueType.ScreenshotTest, "");
+#pragma warning restore CS1591
+
+		static readonly char[] s_pathChars = ['/', '\\'];
+
+		static string RemoveRelativeDirs(string path)
+		{
+			int idx0 = path.IndexOfAny(s_pathChars);
+			if (idx0 != -1)
+			{
+				int idx1 = path.IndexOfAny(s_pathChars, idx0 + 1);
+				while (idx1 != -1)
+				{
+					int idx2 = path.IndexOfAny(s_pathChars, idx1 + 1);
+					if (idx2 == -1)
+					{
+						break;
+					}
+
+					if (idx2 == idx1 + 3 && path[idx1 + 1] == '.' && path[idx1 + 2] == '.')
+					{
+						path = path.Remove(idx0, idx2 - idx0);
+						idx1 = path.IndexOfAny(s_pathChars, idx0 + 1);
+					}
+					else
+					{
+						idx0 = idx1;
+						idx1 = idx2;
+					}
+				}
+			}
+			return path;
+		}
+
+		/// <summary>
+		/// Marks a span of text as a source file
+		/// </summary>
+		public static void AnnotateSourceFile(this LogEventBuilder builder, Group group, string? baseDir)
+		{
+			Dictionary<Utf8String, object?>? properties = null;
+			if (!String.IsNullOrEmpty(baseDir))
+			{
+				string file = group.Value;
+				if (!Path.IsPathRooted(file))
+				{
+					try
+					{
+						string combinedPath = RemoveRelativeDirs(Path.Combine(baseDir, file));
+						properties = [];
+						properties[LogEventPropertyName.File] = combinedPath;
+					}
+					catch
+					{
+					}
+				}
+			}
+			builder.Annotate(group, new LogValue(LogValueType.SourceFile, group.Value, properties));
+		}
+
+		/// <summary>
+		/// Marks a span of text as a source file
+		/// </summary>
+		public static void AnnotateAsset(this LogEventBuilder builder, Group group)
+		{
+			builder.Annotate(group, new LogValue(LogValueType.Asset, group.Value));
+		}
+
+		/// <summary>
+		/// Marks a span of text as a symbol
+		/// </summary>
+		public static void AnnotateSymbol(this LogEventBuilder builder, Group group)
+		{
+			string identifier = group.Value;
+
+			// Remove any __declspec qualifiers
+			identifier = Regex.Replace(identifier, "(?<![^a-zA-Z_])__declspec\\([^\\)]+\\)", "");
+
+			// Remove any argument lists for functions (anything after the first paren)
+			identifier = Regex.Replace(identifier, "\\(.*$", "");
+
+			// Remove any decorators and type information (greedy match up to the last space)
+			identifier = Regex.Replace(identifier, "^.* ", "");
+
+			// Add it to the list
+			Dictionary<Utf8String, object?> properties = [];
+			properties[LogEventPropertyName.Identifier] = identifier;
+			builder.Annotate(group, new LogValue(LogValueType.Symbol, "", properties));
+		}
+	}
+}
diff --git a/Engine/Source/Programs/Shared/EpicGames.Core/Logging/LogEventParser.cs b/Engine/Source/Programs/Shared/EpicGames.Core/Logging/LogEventParser.cs
new file mode 100644
--- /dev/null
+++ b/Engine/Source/Programs/Shared/EpicGames.Core/Logging/LogEventParser.cs
@@ -0,0 +1,1076 @@
+// Copyright Epic Games, Inc. All Rights Reserved.
+
+using System;
+using System.Collections.Generic;
+using System.Diagnostics;
+using System.Linq;
+using System.Reflection;
+using System.Text;
+using System.Text.RegularExpressions;
+using System.Threading.Tasks;
+using Microsoft.Extensions.Logging;
+
+namespace EpicGames.Core
+{
+	/// <summary>
+	/// Confidence of a matched log event being the correct derivation
+	/// </summary>
+	public enum LogEventPriority
+	{
+		/// <summary>
+		/// Unspecified priority
+		/// </summary>
+		None,
+
+		/// <summary>
+		/// Lowest confidence match
+		/// </summary>
+		Lowest,
+
+		/// <summary>
+		/// Low confidence match
+		/// </summary>
+		Low,
+
+		/// <summary>
+		/// Below normal confidence match
+		/// </summary>
+		BelowNormal,
+
+		/// <summary>
+		/// Normal confidence match
+		/// </summary>
+		Normal,
+
+		/// <summary>
+		/// Above normal confidence match
+		/// </summary>
+		AboveNormal,
+
+		/// <summary>
+		/// High confidence match
+		/// </summary>
+		High,
+
+		/// <summary>
+		/// Highest confidence match
+		/// </summary>
+		Highest,
+	}
+
+	/// <summary>
+	/// Information about a matched event
+	/// </summary>
+	public class LogEventMatch
+	{
+		/// <summary>
+		/// Confidence of the match
+		/// </summary>
+		public LogEventPriority Priority { get; }
+
+		/// <summary>
+		/// Matched events
+		/// </summary>
+		public List<LogEvent> Events { get; }
+
+		/// <summary>
+		/// Match is complete
+		/// </summary>
+		public bool IsComplete { get; }
+		
+		/// <summary>
+		/// Constructor
+		/// </summary>
+		public LogEventMatch(LogEventPriority priority, LogEvent logEvent, bool isComplete)
+		{
+			Priority = priority;
+			Events = [logEvent];
+			IsComplete = isComplete;
+		}
+
+		/// <summary>
+		/// Constructor
+		/// </summary>
+		public LogEventMatch(LogEventPriority priority, IEnumerable<LogEvent> events, bool isComplete)
+		{
+			Priority = priority;
+			Events = [.. events];
+			IsComplete = isComplete;
+		}
+	}
+
+	/// <summary>
+	/// Interface for a class which matches error strings
+	/// </summary>
+	public interface ILogEventMatcher
+	{
+		/// <summary>
+		/// Attempt to match events from the given input buffer
+		/// </summary>
+		/// <param name="cursor">The input buffer</param>
+		/// <returns>Information about the error that was matched, or null if an error was not matched</returns>
+		LogEventMatch? Match(ILogCursor cursor);
+
+		/// <summary>
+		/// Get a name for this matcher, used to support dynamic matcher toggling. Names must not contain whitespace characters.
+		/// </summary>
+		/// <remarks>
+		/// This is an interface method with a default interface to allow for edge-cases like multiple matchers of the same type
+		/// but differing behaviour, matchers that want to be enabled or disabled together, inheritance etc.
+		/// </remarks>
+		string GetMatcherName() => GetType().Name;
+	}
+
+	/// <summary>
+	/// An ignore pattern with metadata about its source file
+	/// </summary>
+	/// <param name="Pattern">The compiled regex pattern</param>
+	/// <param name="SourceFile">Absolute path to the file that defined this pattern</param>
+	public readonly record struct IgnorePatternEntry(Regex Pattern, string SourceFile)
+	{
+		/// <summary>
+		/// Implicit conversion to Regex for backward compatibility with code that expects IgnorePatterns to contain Regex objects.
+		/// </summary>
+		public static implicit operator Regex(IgnorePatternEntry entry) => entry.Pattern;
+
+		/// <summary>
+		/// Implicit conversion from Regex for backward compatibility with code that adds Regex objects directly.
+		/// </summary>
+		public static implicit operator IgnorePatternEntry(Regex pattern) => new(pattern, "unknown");
+	}
+
+	/// <summary>
+	/// Extension methods for backward compatibility with code that passes Regex collections to IgnorePatternEntry lists.
+	/// </summary>
+	public static class IgnorePatternEntryExtensions
+	{
+		/// <summary>
+		/// Add a range of Regex objects as ignore pattern entries.
+		/// </summary>
+		public static void AddRange(this List<IgnorePatternEntry> list, IEnumerable<Regex> patterns)
+		{
+			foreach (Regex regex in patterns)
+			{
+				list.Add(new IgnorePatternEntry(regex, "unknown"));
+			}
+		}
+	}
+
+	/// <summary>
+	/// Turns raw text output into structured logging events
+	/// </summary>
+	/// <remarks>
+	/// Log events matching the expected format will be processed by the parser, taking the specified action if possible,
+	/// and may not be output to users. The format is specified in <see cref="MatcherToggleRegex"/>, in addition to the commands
+	/// <c>&lt;-- Suspend Log Parsing --&gt;</c> and <c>&lt;-- Resume Log Parsing --&gt;</c>.
+	/// </remarks>
+	public partial class LogEventParser : IDisposable
+	{
+		/// <summary>
+		/// Timeout applied to each ignore pattern regex match.
+		/// </summary>
+		static readonly TimeSpan s_ignorePatternTimeout = TimeSpan.FromMilliseconds(200);
+
+		/// <summary>
+		/// Timeout applied to each ignore pattern regex match. Exposed for derived classes.
+		/// </summary>
+		protected static TimeSpan IgnorePatternTimeout => s_ignorePatternTimeout;
+
+		/// <summary>
+		/// A list of enabled event matchers for this parser.
+		/// </summary>
+		public IEnumerable<ILogEventMatcher> EnabledMatchers => [.. _enabledMatchers];
+
+		/// <summary>
+		/// List of patterns to ignore
+		/// </summary>
+		public List<IgnorePatternEntry> IgnorePatterns { get; } = [];
+
+		/// <summary>
+		/// A list of matchers all matchers, including disabled matchers.
+		/// </summary>
+		private readonly List<ILogEventMatcher> _matchers = [];
+
+		/// <summary>
+		/// A list of enabled matchers, to allow quicker iteration when parsing log lines.
+		/// It's expected that we'll parse significantly more log lines than we will toggle matcher states.
+		/// </summary>
+		private readonly List<ILogEventMatcher> _enabledMatchers = [];
+
+		/// <summary>
+		/// A mapping of matcher names to how many times they've been enabled / disabled. A negative value indicates a status of disabled.
+		/// Items not present are enabled by default.
+		/// </summary>
+		private readonly Dictionary<string, int> _matcherNameToEnabled = new(StringComparer.OrdinalIgnoreCase);
+
+		/// <summary>
+		/// The layers of enabling / disabling matchers, so we can undo applications.
+		/// </summary>
+		private readonly Stack<(string MatcherName, string? GroupName, bool Enabled)> _matcherToggleLayers = [];
+
+		/// <summary>
+		/// Buffer of input lines
+		/// </summary>
+		readonly LogBuffer _buffer;
+
+		/// <summary>
+		/// Buffer for holding partial line data
+		/// </summary>
+		readonly ByteArrayBuilder _partialLine = new();
+
+		/// <summary>
+		/// Whether matching is currently enabled
+		/// </summary>
+		int _matchingEnabled;
+
+		/// <summary>
+		/// The inner logger
+		/// </summary>
+		ILogger _logger;
+
+		/// <summary>
+		/// Log events sinks in addition to <see cref="_logger" />
+		/// </summary>
+		readonly List<ILogEventSink> _logEventSinks = [];
+
+		/// <summary>
+		/// Timer for the parser being active
+		/// </summary>
+		readonly Stopwatch _timer = Stopwatch.StartNew();
+
+		/// <summary>
+		/// Amount of time that the log parser has been processing events
+		/// </summary>
+		readonly Stopwatch _activeTimer = new();
+
+		/// <summary>
+		/// Number of lines parsed in the last interval
+		/// </summary>
+		int _linesParsed = 0;
+
+		/// <summary>
+		/// Incomplete matcher being processed
+		/// </summary>
+		ILogEventMatcher? _incompleteMatcher;
+		
+		/// <summary>
+		/// Incomplete event count being processed
+		/// </summary>
+		int _incompleteEventCount;
+		
+		/// <summary>
+		/// Public accessor for the logger
+		/// </summary>
+		public ILogger Logger
+		{
+			get => _logger;
+			set => _logger = value;
+		}
+
+		/// <summary>
+		/// The pattern matcher toggle commands must conform to.
+		/// </summary>
+		/// <remarks>
+		/// <c>LayerName</c> is an arbitrary string to allow validation that users are popping the layer they expect to be - if not, a warning will be issued.
+		/// 
+		/// <c>MatcherName</c> is either the literal name of a matcher, or the string <c>*</c> for all matchers.
+		/// </remarks>
+		[GeneratedRegex("^(?:<|\\\\u003C)-- logeventparser (layer=(?<LayerName>\\S+)? )?((?<Action>suspend|resume) (?<MatcherName>\\S+)|(?<Action>pop)) --(?:>|\\\\u003E)$", RegexOptions.IgnoreCase | RegexOptions.ExplicitCapture)]
+		public static partial Regex MatcherToggleRegex();
+
+		/// <summary>
+		/// Constructor
+		/// </summary>
+		/// <param name="logger">The logger to receive parsed output messages</param>
+		/// <param name="logEventSinks">Additional sinks to receive log events</param>
+		public LogEventParser(ILogger logger, List<ILogEventSink>? logEventSinks = null)
+		{
+			_logger = logger;
+			_buffer = new LogBuffer(50);
+
+			if (logEventSinks != null)
+			{
+				_logEventSinks.AddRange(logEventSinks);
+			}
+		}
+
+		/// <inheritdoc/>
+		public void Dispose()
+		{
+			Dispose(true);
+			GC.SuppressFinalize(this);
+		}
+
+		/// <summary>
+		/// Standard Dispose pattern method
+		/// </summary>
+		/// <param name="disposing"></param>
+		protected virtual void Dispose(bool disposing)
+		{
+			if (disposing)
+			{
+				Flush();
+			}
+		}
+
+		/// <summary>
+		/// Add the selected matchers to the log event parser
+		/// </summary>
+		/// <param name="matchers">The matchers to add</param>
+		public void AddMatchers(IEnumerable<ILogEventMatcher> matchers)
+		{
+			foreach (ILogEventMatcher matcher in matchers)
+			{
+				AddMatcher(matcher);
+			}
+		}
+
+		/// <summary>
+		/// Enumerate all the types that implement <see cref="ILogEventMatcher"/> in the given assembly, and create instances of them
+		/// </summary>
+		/// <param name="assembly">The assembly to enumerate matchers from</param>
+		public void AddMatchersFromAssembly(Assembly assembly)
+		{
+			foreach (Type type in assembly.GetTypes())
+			{
+				if (type.IsClass && typeof(ILogEventMatcher).IsAssignableFrom(type))
+				{
+					ILogEventMatcher matcher = (ILogEventMatcher)Activator.CreateInstance(type)!;
+					AddMatcher(matcher);
+				}
+			}
+		}
+
+		/// <summary>
+		/// Read ignore patterns from the given root directory
+		/// </summary>
+		/// <param name="rootDir"></param>
+		/// <returns></returns>
+		[Obsolete("Use ReadIgnorePatternsFromWorkspaceAsync or DiscoverIgnorePatternFiles instead")]
+		public async Task ReadIgnorePatternsAsync(DirectoryReference rootDir)
+		{
+			Stopwatch timer = Stopwatch.StartNew();
+
+			List<DirectoryReference> baseDirs = [rootDir];
+			AddRestrictedDirs(baseDirs, "Restricted");
+			AddRestrictedDirs(baseDirs, "Platforms");
+
+			List<(FileReference, Task)> tasks = [];
+			foreach (DirectoryReference baseDir in baseDirs)
+			{
+				FileReference ignorePatternFile = FileReference.Combine(baseDir, "Build", "Horde", "IgnorePatterns.txt");
+				if (FileReference.Exists(ignorePatternFile))
+				{
+					_logger.LogDebug("Reading ignore patterns from {File}...", ignorePatternFile);
+					tasks.Add((ignorePatternFile, ReadIgnorePatternsAsync(ignorePatternFile)));
+				}
+			}
+
+			foreach ((FileReference file, Task task) in tasks)
+			{
+				try
+				{
+					await task;
+				}
+				catch (Exception ex)
+				{
+					_logger.LogError(ex, "Exception reading patterns from {File}: {Message}", file, ex.Message);
+				}
+			}
+
+			_logger.LogDebug("Took {TimeMs}ms to read ignore patterns", timer.ElapsedMilliseconds);
+		}
+
+		/// <summary>
+		/// Read ignore patterns from a single file
+		/// </summary>
+		/// <param name="ignorePatternFile"></param>
+		/// <returns></returns>
+		public async Task ReadIgnorePatternsAsync(FileReference ignorePatternFile)
+		{
+			string[] lines = await FileReference.ReadAllLinesAsync(ignorePatternFile);
+
+			List<IgnorePatternEntry> entries = [];
+			for (int idx = 0; idx < lines.Length; idx++)
+			{
+				string trimLine = lines[idx].Trim();
+				if (trimLine.Length > 0 && trimLine[0] != '#')
+				{
+					try
+					{
+						Regex regex = new(trimLine, RegexOptions.Compiled, s_ignorePatternTimeout);
+						entries.Add(new IgnorePatternEntry(regex, ignorePatternFile.FullName));
+					}
+					catch (ArgumentException ex)
+					{
+						_logger.LogWarning("Invalid regex in {File} line {LineNumber}: {Pattern} ({Message})",
+							ignorePatternFile, idx + 1, trimLine, ex.Message);
+					}
+				}
+			}
+
+			lock (IgnorePatterns)
+			{
+				IgnorePatterns.AddRange(entries);
+			}
+		}
+
+		static void AddRestrictedDirs(List<DirectoryReference> directories, string subFolder)
+		{
+			int numDirs = directories.Count;
+			for (int idx = 0; idx < numDirs; idx++)
+			{
+				DirectoryReference subDir = DirectoryReference.Combine(directories[idx], subFolder);
+				if (DirectoryReference.Exists(subDir))
+				{
+					directories.AddRange(DirectoryReference.EnumerateDirectories(subDir));
+				}
+			}
+		}
+
+		static IgnorePatternEntry? FindMatchingIgnorePattern(List<IgnorePatternEntry> ignorePatterns, string line, ILogger logger)
+		{
+			foreach (IgnorePatternEntry entry in ignorePatterns)
+			{
+				try
+				{
+					if (entry.Pattern.IsMatch(line))
+					{
+						return entry;
+					}
+				}
+				catch (RegexMatchTimeoutException)
+				{
+					logger.LogWarning("Ignore pattern timed out on line (pattern: {Pattern}): {Line}",
+						entry.Pattern.ToString(), line.Length > 200 ? line[..200] + "..." : line);
+				}
+			}
+			return null;
+		}
+
+		/// <summary>
+		/// Writes a line that was suppressed by an ignore pattern, including structured properties
+		/// identifying which pattern file and regex caused the suppression.
+		/// </summary>
+		/// <param name="line">The original log line text</param>
+		/// <param name="ignoredBy">The ignore pattern entry that matched</param>
+		/// <param name="originalLevel">The log level the line would have been emitted at before suppression</param>
+		void WriteIgnoredLine(string line, IgnorePatternEntry ignoredBy, LogLevel originalLevel)
+		{
+			LogValue ignoredLineValue = new(LogValueType.IgnoredLine, line, new Dictionary<Utf8String, object?>
+			{
+				[new Utf8String("ignorePatternFile")] = ignoredBy.SourceFile,
+				[new Utf8String("ignorePattern")] = ignoredBy.Pattern.ToString(),
+				[new Utf8String("originalLevel")] = originalLevel.ToString(),
+			});
+
+			KeyValuePair<string, object?>[] state =
+			[
+				new(MessageTemplate.FormatPropertyName, "{IgnoredLine}"),
+				new("IgnoredLine", ignoredLineValue),
+			];
+
+			_logger.Log(LogLevel.Information, KnownLogEvents.None, state, null, (s, _) => line);
+		}
+
+		/// <summary>
+		/// Checks whether a structured JSON log event should be suppressed by an ignore pattern.
+		/// Only events at Warning level or above are checked. If a matching pattern is found,
+		/// the event is re-logged at Information level (preventing build health issue creation).
+		/// </summary>
+		/// <returns>True if the event was suppressed, false if it should be logged normally</returns>
+		bool TryIgnoreStructuredEvent(JsonLogEvent jsonEvent)
+		{
+			if (jsonEvent.Level >= LogLevel.Warning && IgnorePatterns.Count > 0)
+			{
+				string message = jsonEvent.GetRenderedMessage().ToString();
+				IgnorePatternEntry? ignoredBy = FindMatchingIgnorePattern(IgnorePatterns, message, _logger);
+				if (ignoredBy != null)
+				{
+					WriteIgnoredLine(message, ignoredBy.Value, jsonEvent.Level);
+					return true;
+				}
+			}
+			return false;
+		}
+
+		/// <summary>
+		/// Discover all IgnorePatterns.txt files in a workspace by scanning Engine, project,
+		/// and Restricted/Platforms subdirectories.
+		/// </summary>
+		public static List<FileReference> DiscoverIgnorePatternFiles(DirectoryReference workspaceDir, string enginePath = "Engine")
+		{
+			HashSet<DirectoryReference> scannedDirs = [];
+			List<DirectoryReference> baseDirs = [];
+
+			// Engine directory
+			DirectoryReference engineDir = DirectoryReference.Combine(workspaceDir, enginePath);
+			scannedDirs.Add(engineDir);
+			baseDirs.Add(engineDir);
+
+			// All top-level workspace directories
+			if (DirectoryReference.Exists(workspaceDir))
+			{
+				foreach (DirectoryReference subDir in DirectoryReference.EnumerateDirectories(workspaceDir))
+				{
+					if (scannedDirs.Add(subDir))
+					{
+						baseDirs.Add(subDir);
+					}
+				}
+			}
+
+			// Restricted/Platforms under each base directory
+			int count = baseDirs.Count;
+			for (int idx = 0; idx < count; idx++)
+			{
+				foreach (string subFolder in new[] { "Restricted", "Platforms" })
+				{
+					DirectoryReference subDir = DirectoryReference.Combine(baseDirs[idx], subFolder);
+					if (DirectoryReference.Exists(subDir))
+					{
+						foreach (DirectoryReference child in DirectoryReference.EnumerateDirectories(subDir))
+						{
+							if (scannedDirs.Add(child))
+							{
+								baseDirs.Add(child);
+							}
+						}
+					}
+				}
+			}
+
+			// Find all existing IgnorePatterns.txt files
+			List<FileReference> files = [];
+			foreach (DirectoryReference baseDir in baseDirs)
+			{
+				FileReference file = FileReference.Combine(baseDir, "Build", "Horde", "IgnorePatterns.txt");
+				if (FileReference.Exists(file))
+				{
+					files.Add(file);
+				}
+			}
+			return files;
+		}
+
+		/// <summary>
+		/// Legacy ignore pattern file discovery - scans only Engine + Engine/Restricted/* + Engine/Platforms/*.
+		/// Use <see cref="DiscoverIgnorePatternFiles"/> for expanded workspace-wide scanning.
+		/// </summary>
+		public static List<FileReference> DiscoverLegacyIgnorePatternFiles(DirectoryReference workspaceDir, string enginePath = "Engine")
+		{
+			DirectoryReference engineDir = DirectoryReference.Combine(workspaceDir, enginePath);
+			List<DirectoryReference> baseDirs = [engineDir];
+			AddRestrictedDirs(baseDirs, "Restricted");
+			AddRestrictedDirs(baseDirs, "Platforms");
+
+			List<FileReference> files = [];
+			foreach (DirectoryReference baseDir in baseDirs)
+			{
+				FileReference file = FileReference.Combine(baseDir, "Build", "Horde", "IgnorePatterns.txt");
+				if (FileReference.Exists(file))
+				{
+					files.Add(file);
+				}
+			}
+			return files;
+		}
+
+		/// <summary>
+		/// Read ignore patterns from the workspace, scanning Engine and project directories
+		/// </summary>
+		public async Task ReadIgnorePatternsFromWorkspaceAsync(DirectoryReference workspaceDir, string enginePath = "Engine")
+		{
+			Stopwatch timer = Stopwatch.StartNew();
+
+			List<FileReference> files = DiscoverIgnorePatternFiles(workspaceDir, enginePath);
+
+			List<(FileReference, Task)> tasks = [];
+			foreach (FileReference file in files)
+			{
+				_logger.LogDebug("Reading ignore patterns from {File}...", file);
+				tasks.Add((file, ReadIgnorePatternsAsync(file)));
+			}
+
+			foreach ((FileReference file, Task task) in tasks)
+			{
+				try
+				{
+					await task;
+				}
+				catch (Exception ex)
+				{
+					_logger.LogError(ex, "Exception reading patterns from {File}: {Message}", file, ex.Message);
+				}
+			}
+
+			_logger.LogDebug("Read {PatternCount} ignore patterns from {FileCount} files in {TimeMs}ms",
+				IgnorePatterns.Count, tasks.Count, timer.ElapsedMilliseconds);
+		}
+
+		/// <summary>
+		/// Read ignore patterns from additional file paths
+		/// </summary>
+		public async Task ReadIgnorePatternsAsync(IEnumerable<FileReference> additionalFiles)
+		{
+			List<(FileReference, Task)> tasks = [];
+			foreach (FileReference file in additionalFiles)
+			{
+				if (FileReference.Exists(file))
+				{
+					_logger.LogDebug("Reading ignore patterns from {File}...", file);
+					tasks.Add((file, ReadIgnorePatternsAsync(file)));
+				}
+				else
+				{
+					_logger.LogWarning("Ignore pattern file not found: {File}", file);
+				}
+			}
+
+			foreach ((FileReference file, Task task) in tasks)
+			{
+				try
+				{
+					await task;
+				}
+				catch (Exception ex)
+				{
+					_logger.LogError(ex, "Exception reading patterns from {File}: {Message}", file, ex.Message);
+				}
+			}
+		}
+
+		/// <summary>
+		/// Writes a line to the event filter
+		/// </summary>
+		/// <param name="line">The line to output</param>
+		public void WriteLine(string line)
+		{
+			bool invalidJson = false;
+
+			if (line.Length > 0 && line[0] == '{')
+			{
+				int length = line.Length;
+				while (length > 0 && Char.IsWhiteSpace(line[length - 1]))
+				{
+					length--;
+				}
+
+				byte[] data = Encoding.UTF8.GetBytes(line, 0, length);
+				try
+				{
+					JsonLogEvent jsonEvent;
+					if (JsonLogEvent.TryParse(data, out jsonEvent))
+					{
+						ProcessData(true);
+
+						if (Regex.IsMatch(line, @"\\u003c-- Suspend Log Parsing --\\u003e", RegexOptions.IgnoreCase))
+						{
+							_matchingEnabled--;
+						}
+						else if (Regex.IsMatch(line, @"\\u003c-- Resume Log Parsing --\\u003e", RegexOptions.IgnoreCase))
+						{
+							_matchingEnabled++;
+						}
+
+						if (!TryIgnoreStructuredEvent(jsonEvent))
+						{
+							_logger.Log(jsonEvent.Level, jsonEvent.EventId, jsonEvent, null, JsonLogEvent.Format);
+						}
+						return;
+					}
+
+					invalidJson = true;
+				}
+				catch (Exception ex)
+				{
+					_logger.LogError(ex, "Exception while parsing log event: {Message}", ex.Message);
+				}
+			}
+
+			_buffer.AddLine(StringUtils.ParseEscapeCodes(line));
+
+			ProcessData(false, disableMatching: invalidJson);
+		}
+
+		/// <summary>
+		/// Writes data to the log parser
+		/// </summary>
+		/// <param name="data">Data to write</param>
+		public void WriteData(ReadOnlyMemory<byte> data)
+		{
+			int baseIdx = 0;
+			int scanIdx = 0;
+			ReadOnlySpan<byte> span = data.Span;
+
+			// Handle a partially existing line
+			if (_partialLine.Length > 0)
+			{
+				for (; scanIdx < span.Length; scanIdx++)
+				{
+					if (span[scanIdx] == '\n')
+					{
+						_partialLine.WriteFixedLengthBytes(span.Slice(baseIdx, scanIdx - baseIdx));
+						FlushPartialLine();
+						baseIdx = ++scanIdx;
+						break;
+					}
+				}
+			}
+
+			// Handle any complete lines
+			for (; scanIdx < span.Length; scanIdx++)
+			{
+				if (span[scanIdx] == '\n')
+				{
+					AddLine(data.Slice(baseIdx, scanIdx - baseIdx));
+					baseIdx = scanIdx + 1;
+				}
+			}
+
+			// Add the rest of the text to the partial line buffer
+			_partialLine.WriteFixedLengthBytes(span.Slice(baseIdx));
+
+			// Process the new data
+			ProcessData(false);
+		}
+
+		/// <summary>
+		/// Flushes the current contents of the parser
+		/// </summary>
+		public void Flush()
+		{
+			// If there's a partially written line, write that out first
+			if (_partialLine.Length > 0)
+			{
+				FlushPartialLine();
+			}
+
+			// Process any remaining data
+			ProcessData(true);
+		}
+
+		/// <summary>
+		/// Add the new matcher, making sure its enabled state is correct.
+		/// </summary>
+		private void AddMatcher(ILogEventMatcher matcher)
+		{
+			if (matcher is null)
+			{
+				return;
+			}
+
+			_matchers.Add(matcher);
+			// This is a little inefficient, but this function shouldn't be called often.
+			RecalculateEnabledMatchers();
+		}
+
+		/// <summary>
+		/// Set the given matcher to the given enabled state.
+		/// </summary>
+		private void PushToggle(string matcherName, string layerName, bool enabled)
+		{
+			if (_matcherToggleLayers.Count > 1024)
+			{
+				_logger.LogWarning(
+					"[LogEventParser] Over 1024 log event parser toggle layers have been pushed, no longer processing any more to avoid unbounded memory usage. Most recent toggle: matcher='{MatcherName}', layer='{LayerName}', enabled={Enabled}",
+					matcherName,
+					layerName,
+					enabled);
+				return;
+			}
+
+			// Allow debugging if all log matchers get disabled and this isn't popped.
+			_logger.LogTrace(
+				"[LogEventParser] Most recent toggle: matcher='{MatcherName}', layer='{LayerName}', enabled={Enabled}",
+				matcherName,
+				layerName,
+				enabled);
+
+			_matcherToggleLayers.Push((matcherName, layerName, enabled));
+			_matcherNameToEnabled.TryGetValue(matcherName, out int enabledCount);
+			enabledCount += enabled ? 1 : -1;
+			_matcherNameToEnabled[matcherName] = enabledCount;
+
+			RecalculateEnabledMatchers();
+		}
+
+		private void PopToggle(string layerName)
+		{
+			_logger.LogTrace("[LogEventParser] Popping most recent toggle: layer='{LayerName}'", layerName);
+
+			if (_matcherToggleLayers.Count == 0)
+			{
+				_logger.LogWarning("[LogEventParser] Popped empty stack, layer='{LayerName}'", layerName);
+				return;
+			}
+
+			(string matcherName, string? layerNameFromPop, bool enabled) = _matcherToggleLayers.Pop();
+
+			if (layerName.Length != 0 && !String.Equals(layerName, layerNameFromPop, StringComparison.OrdinalIgnoreCase))
+			{
+				_logger.LogWarning(
+					"[LogEventParser] Given layer '{LayerName}' is not the same as last pushed layer '{LayerNameFromPop}', please check your logic.",
+					layerName,
+					layerNameFromPop);
+				// Arguable if we should pop anyway, but I think it's slightly safer to do that.
+			}
+
+			// We're popping a layer, so if it was enabled, we're adding a disable and vice versa - hence why the sign is flipped compared to PushToggle.
+			_matcherNameToEnabled[matcherName] += enabled ? -1 : 1;
+
+			RecalculateEnabledMatchers();
+		}
+
+		/// <summary>
+		/// Clear and repopulate the list of enabled matchers.
+		/// </summary>
+		private void RecalculateEnabledMatchers()
+		{
+			_enabledMatchers.Clear();
+
+			// To allow disabling all loggers then enabling specific ones, we add the value for each logger to the value for "*".
+			_matcherNameToEnabled.TryGetValue("*", out int allCount);
+
+			foreach (ILogEventMatcher matcher in _matchers)
+			{
+				_matcherNameToEnabled.TryGetValue(matcher.GetMatcherName(), out int enabledCount);
+
+				if (enabledCount + allCount >= 0)
+				{
+					_enabledMatchers.Add(matcher);
+				}
+			}
+		}
+
+		/// <summary>
+		/// Adds a raw utf-8 string to the buffer
+		/// </summary>
+		/// <param name="data">The string data</param>
+		private void AddLine(ReadOnlyMemory<byte> data)
+		{
+			if (data.Length > 0 && data.Span[data.Length - 1] == '\r')
+			{
+				data = data.Slice(0, data.Length - 1);
+			}
+			if (data.Length > 0 && data.Span[0] == '{')
+			{
+				JsonLogEvent jsonEvent;
+				if (JsonLogEvent.TryParse(data, out jsonEvent))
+				{
+					ProcessData(true);
+					if (!TryIgnoreStructuredEvent(jsonEvent))
+					{
+						_logger.LogJsonLogEvent(jsonEvent);
+					}
+					return;
+				}
+			}
+			_buffer.AddLine(StringUtils.ParseEscapeCodes(Encoding.UTF8.GetString(data.Span)));
+		}
+
+		/// <summary>
+		/// Writes the current partial line data, with the given data appended to it, then clear the buffer
+		/// </summary>
+		private void FlushPartialLine()
+		{
+			AddLine(_partialLine.ToByteArray());
+			_partialLine.Clear();
+		}
+
+		/// <summary>
+		/// Check for and handle a log event parser command.
+		/// </summary>
+		/// <param name="line">The line to check for a log event parser command.</param>
+		/// <returns>Whether a log event parser command was found and handled.</returns>
+		private bool HandleLogEventParserCommand(string line)
+		{
+			Match match = MatcherToggleRegex().Match(line);
+			if (match.Success)
+			{
+				string action = match.Groups["Action"].Value;
+				string matcherName = match.Groups["MatcherName"].Value;
+				string layerName = match.Groups["LayerName"].Value;
+				switch (action.ToUpperInvariant())
+				{
+					case "RESUME":
+						PushToggle(matcherName, layerName, true);
+						break;
+					case "SUSPEND":
+						PushToggle(matcherName, layerName, false);
+						break;
+					case "POP":
+						PopToggle(layerName);
+						break;
+				}
+				_buffer.MoveNext();
+				return true;
+			}
+
+			return false;
+		}
+
+		/// <summary>
+		/// Process any data in the buffer
+		/// </summary>
+		/// <param name="bFlush">Whether we've reached the end of the stream</param>
+		/// <param name="disableMatching">Whether to skip matching</param>
+		void ProcessData(bool bFlush, bool disableMatching = false)
+		{
+			_activeTimer.Start();
+			int startLineCount = _buffer.Length;
+			if (ProcessIncomplete(bFlush))
+			{
+				while (_buffer.Length > 0)
+				{
+					string? line = _buffer[0];
+					if (line == null)
+					{
+						line = "";
+					}
+					else if (HandleLogEventParserCommand(line))
+					{
+						continue;
+					}
+					else if (Regex.IsMatch(line, "<-- Suspend Log Parsing -->", RegexOptions.IgnoreCase))
+					{
+						_matchingEnabled--;
+					}
+					else if (Regex.IsMatch(line, "<-- Resume Log Parsing -->", RegexOptions.IgnoreCase))
+					{
+						_matchingEnabled++;
+					}
+					else if (_matchingEnabled >= 0 && !disableMatching)
+					{
+						(ILogEventMatcher? matcher, LogEventMatch? match) = MatchEvent();
+
+						// Bail out if we need more data, or match is incomplete
+						if (_buffer.Length < 1024 && !bFlush && (_buffer.NeedMoreData || match is { IsComplete: false }))
+						{
+							if (!_buffer.NeedMoreData && FindMatchingIgnorePattern(IgnorePatterns, line, _logger) == null)
+							{
+								WriteEvents(match!.Events);
+								Debug.Assert(matcher != null);
+								_incompleteMatcher = matcher;
+								_incompleteEventCount = match.Events.Count;
+							}
+							break;
+						}
+
+						// Report the error to the listeners
+						// If we did match something, check if it's not negated by an ignore pattern. We typically have relatively few errors and many more ignore patterns than matchers, so it's quicker
+						// to check them in response to an identified error than to treat them as matchers of their own.
+						if (match != null)
+						{
+							IgnorePatternEntry? ignoredBy = FindMatchingIgnorePattern(IgnorePatterns, line, _logger);
+							if (ignoredBy == null)
+							{
+								WriteEvents(match.Events);
+								_buffer.Advance(match.Events.Count);
+								continue;
+							}
+
+							WriteIgnoredLine(line, ignoredBy.Value, match.Events[0].Level);
+							_buffer.MoveNext();
+							continue;
+						}
+					}
+
+					_logger.Log(LogLevel.Information, KnownLogEvents.None, line, null, (state, exception) => state);
+					_buffer.MoveNext();
+				}
+			}
+			_linesParsed += startLineCount - _buffer.Length;
+			_activeTimer.Stop();
+
+			const double UpdateIntervalSeconds = 30.0;
+			double elapsedSeconds = _timer.Elapsed.TotalSeconds;
+			if (elapsedSeconds > UpdateIntervalSeconds)
+			{
+				const double WarnPct = 0.5;
+				double activeSeconds = _activeTimer.Elapsed.TotalSeconds;
+				double activePct = activeSeconds / elapsedSeconds;
+
+				if (activePct > WarnPct)
+				{
+					_logger.LogInformation(KnownLogEvents.Systemic_LogParserBottleneck, "EpicGames.Core.LogEventParser is taking a significant amount of CPU time: {Active:n1}s/{Total:n1}s ({Pct:n1}%). Processed {NumLines} lines in last {Interval} seconds ({NumLinesInBuffer} in buffer).", activeSeconds, elapsedSeconds, activePct * 100.0, _linesParsed, UpdateIntervalSeconds, _buffer.Length);
+				}
+
+				_activeTimer.Reset();
+				_timer.Restart();
+				_linesParsed = 0;
+			}
+		}
+
+		/// <summary>
+		/// Process any incomplete match
+		/// </summary>
+		/// <param name="bFlush">Whether we've reached the end of the stream</param>
+		/// <returns>Whether it finished processing the incomplete match</returns>
+		bool ProcessIncomplete(bool bFlush)
+		{
+			if (_incompleteMatcher != null)
+			{
+				LogEventMatch match = _incompleteMatcher.Match(_buffer)!;
+				Debug.Assert(match.Events.Count >= _incompleteEventCount);
+				WriteEvents(match.Events.Skip(_incompleteEventCount));
+
+				// Bail out if we need more data
+				if (_buffer.Length < 1024 && !bFlush && !match.IsComplete)
+				{
+					_incompleteEventCount = match.Events.Count;
+					return false;
+				}
+
+				_incompleteMatcher = null;
+				_incompleteEventCount = 0;
+				_buffer.Advance(match.Events.Count);
+			}
+
+			return true;
+		}
+		
+		/// <summary>
+		/// Try to match an event from the current buffer
+		/// </summary>
+		/// <returns>The matcher and the match</returns>
+		private (ILogEventMatcher?, LogEventMatch?) MatchEvent()
+		{
+			ILogEventMatcher? currentMatcher = null;
+			LogEventMatch? currentMatch = null;
+			foreach (ILogEventMatcher matcher in _enabledMatchers)
+			{
+				try
+				{
+					LogEventMatch? match = matcher.Match(_buffer);
+					if (match != null && (currentMatch == null || match.Priority > currentMatch.Priority))
+					{
+						currentMatcher = matcher;
+						currentMatch = match;
+					}
+				}
+				catch (Exception ex)
+				{
+					_logger.LogWarning(KnownLogEvents.Systemic_LogEventMatcher, ex, "Exception while parsing log events with {Type}. Buffer size {Length}.", matcher.GetType().Name, _buffer.Length);
+				}
+			}
+			return (currentMatcher, currentMatch);
+		}
+
+		/// <summary>
+		/// Writes an event to the log
+		/// </summary>
+		/// <param name="logEvents">The event to write</param>
+		protected virtual void WriteEvents(IEnumerable<LogEvent> logEvents)
+		{
+			foreach (LogEvent logEvent in logEvents)
+			{
+				_logger.Log(logEvent.Level, logEvent.Id, logEvent, null, (state, exception) => state.ToString());
+				foreach (ILogEventSink logEventSink in _logEventSinks)
+				{
+					logEventSink.ProcessEvent(logEvent);
+				}
+			}
+		}
+	}
+}
diff --git a/Engine/Source/Programs/Shared/EpicGames.Core/Logging/LogIndent.cs b/Engine/Source/Programs/Shared/EpicGames.Core/Logging/LogIndent.cs
new file mode 100644
--- /dev/null
+++ b/Engine/Source/Programs/Shared/EpicGames.Core/Logging/LogIndent.cs
@@ -0,0 +1,60 @@
+// Copyright Epic Games, Inc. All Rights Reserved.
+
+using System;
+using System.Threading;
+
+namespace EpicGames.Core
+{
+	/// <summary>
+	/// Static class for tracking a stack of log indent prefixes
+	/// </summary>
+	public static class LogIndent
+	{
+		/// <summary>
+		/// Tracks the current state
+		/// </summary>
+		class State
+		{
+			public State? PrevState { get; }
+			public string Indent { get; }
+
+			public State(State? prevState, string indent)
+			{
+				PrevState = prevState;
+				Indent = indent;
+			}
+		}
+
+		/// <summary>
+		/// The current state value
+		/// </summary>
+		static readonly AsyncLocal<State?> s_currentState = new();
+
+		/// <summary>
+		/// Gets the current indent string
+		/// </summary>
+		public static string Current => s_currentState.Value?.Indent ?? String.Empty;
+
+		/// <summary>
+		/// Push a new indent onto the stack
+		/// </summary>
+		/// <param name="indent">The indent to add</param>
+		public static void Push(string indent)
+		{
+			State? prevState = s_currentState.Value;
+			if (prevState != null)
+			{
+				indent = prevState.Indent + indent;
+			}
+			s_currentState.Value = new State(prevState, indent);
+		}
+
+		/// <summary>
+		/// Pops an indent off the stack
+		/// </summary>
+		public static void Pop()
+		{
+			s_currentState.Value = s_currentState.Value!.PrevState;
+		}
+	}
+}
diff --git a/Engine/Source/Programs/Shared/EpicGames.Core/Logging/LogIndentScope.cs b/Engine/Source/Programs/Shared/EpicGames.Core/Logging/LogIndentScope.cs
new file mode 100644
--- /dev/null
+++ b/Engine/Source/Programs/Shared/EpicGames.Core/Logging/LogIndentScope.cs
@@ -0,0 +1,38 @@
+// Copyright Epic Games, Inc. All Rights Reserved.
+
+using System;
+
+namespace EpicGames.Core
+{
+	/// <summary>
+	/// Class to apply a log indent for the lifetime of an object 
+	/// </summary>
+	public sealed class LogIndentScope : IDisposable
+	{
+		/// <summary>
+		/// Whether the object has been disposed
+		/// </summary>
+		bool _disposed;
+
+		/// <summary>
+		/// Constructor
+		/// </summary>
+		/// <param name="indent">Indent to append to the existing indent</param>
+		public LogIndentScope(string indent)
+		{
+			LogIndent.Push(indent);
+		}
+
+		/// <summary>
+		/// Restore the log indent to normal
+		/// </summary>
+		public void Dispose()
+		{
+			if (!_disposed)
+			{
+				LogIndent.Pop();
+				_disposed = true;
+			}
+		}
+	}
+}
diff --git a/Engine/Source/Programs/Shared/EpicGames.Core/Logging/LogStatusScope.cs b/Engine/Source/Programs/Shared/EpicGames.Core/Logging/LogStatusScope.cs
new file mode 100644
--- /dev/null
+++ b/Engine/Source/Programs/Shared/EpicGames.Core/Logging/LogStatusScope.cs
@@ -0,0 +1,95 @@
+// Copyright Epic Games, Inc. All Rights Reserved.
+
+using System;
+using System.Text;
+
+namespace EpicGames.Core
+{
+	/// <summary>
+	/// Writes a status message to the log, which can be updated with a progress indicator as a slow task is being performed.
+	/// </summary>
+	public sealed class LogStatusScope : IDisposable
+	{
+		/// <summary>
+		/// The base status message
+		/// </summary>
+		string _message;
+
+		/// <summary>
+		/// Constructor
+		/// </summary>
+		/// <param name="message">The status message</param>
+		public LogStatusScope(string message)
+		{
+			_message = message;
+			Log.PushStatus(message);
+		}
+
+		/// <summary>
+		/// Constructor
+		/// </summary>
+		/// <param name="format">The format specifier for the message</param>
+		/// <param name="args">Arguments for the status message</param>
+		public LogStatusScope(string format, params object[] args)
+			: this(String.Format(format, args))
+		{
+		}
+
+		/// <summary>
+		/// Updates the base status message passed into the constructor.
+		/// </summary>
+		/// <param name="message">The status message</param>
+		public void SetMessage(string message)
+		{
+			_message = message;
+			Log.UpdateStatus(message);
+		}
+
+		/// <summary>
+		/// Updates the base status message passed into the constructor.
+		/// </summary>
+		/// <param name="format">The format specifier for the message</param>
+		/// <param name="args">Arguments for the status message</param>
+		public void SetMessage(string format, params object[] args)
+		{
+			SetMessage(String.Format(format, args));
+		}
+
+		/// <summary>
+		/// Appends a progress string to the status message. Overwrites any previous progress message.
+		/// </summary>
+		/// <param name="progress">The progress message</param>
+		public void SetProgress(string progress)
+		{
+			StringBuilder fullMessage = new(_message);
+			fullMessage.Append(' ');
+			fullMessage.Append(progress);
+			Log.UpdateStatus(fullMessage.ToString());
+		}
+
+		/// <summary>
+		/// Appends a progress string to the status message. Overwrites any previous progress message.
+		/// </summary>
+		/// <param name="format">The format specifier for the message</param>
+		/// <param name="args">Arguments for the status message</param>
+		public void SetProgress(string format, params object[] args)
+		{
+			StringBuilder fullMessage = new(_message);
+			fullMessage.Append(' ');
+			fullMessage.AppendFormat(format, args);
+			Log.UpdateStatus(fullMessage.ToString());
+		}
+
+		/// <summary>
+		/// Pops the status message from the log.
+		/// </summary>
+		public void Dispose()
+		{
+			if (_message != null)
+			{
+				Log.PopStatus();
+				_message = null!;
+			}
+		}
+	}
+}
diff --git a/Engine/Source/Programs/Shared/EpicGames.Core/Logging/LogValue.cs b/Engine/Source/Programs/Shared/EpicGames.Core/Logging/LogValue.cs
new file mode 100644
--- /dev/null
+++ b/Engine/Source/Programs/Shared/EpicGames.Core/Logging/LogValue.cs
@@ -0,0 +1,151 @@
+// Copyright Epic Games, Inc. All Rights Reserved.
+
+using System;
+using System.Collections.Generic;
+using System.Reflection;
+using System.Text.Json.Serialization;
+
+namespace EpicGames.Core
+{
+	/// <summary>
+	/// Different types of log values (stored in the $type field of log properties)
+	/// </summary>
+	public static class LogValueType
+	{
+#pragma warning disable CS1591
+		public static readonly Utf8String Asset = new("Asset");
+		public static readonly Utf8String SourceFile = new("SourceFile");
+		public static readonly Utf8String Object = new("Object"); // Arbitrary structured object
+		public static readonly Utf8String Channel = new("Channel");
+		public static readonly Utf8String Severity = new("Severity");
+		public static readonly Utf8String Message = new("Message");
+		public static readonly Utf8String LineNumber = new("Line");
+		public static readonly Utf8String ColumnNumber = new("Column");
+		public static readonly Utf8String Symbol = new("Symbol");
+		public static readonly Utf8String ErrorCode = new("ErrorCode");
+		public static readonly Utf8String ToolName = new("ToolName");
+		public static readonly Utf8String ScreenshotTest = new("ScreenshotTest");
+		public static readonly Utf8String DepotPath = new("DepotPath");
+		public static readonly Utf8String Link = new("Link");
+		public static readonly Utf8String IgnoredLine = new("IgnoredLine");
+#pragma warning restore CS1591
+	}
+
+	/// <summary>
+	/// Attribute indicating that a type should be tagged in log output
+	/// </summary>
+	[AttributeUsage(AttributeTargets.Class | AttributeTargets.Struct)]
+	public sealed class LogValueTypeAttribute : Attribute
+	{
+		/// <summary>
+		/// Name to use for the type tag
+		/// </summary>
+		public string? Name { get; }
+
+		/// <summary>
+		/// Constructor
+		/// </summary>
+		public LogValueTypeAttribute(string? name = null) => Name = name;
+	}
+
+	/// <summary>
+	/// Information for a structured value for use in log events
+	/// </summary>
+	public sealed class LogValue
+	{
+		/// <summary>
+		/// Type of the event
+		/// </summary>
+		public Utf8String Type { get; set; }
+
+		/// <summary>
+		/// Rendering of the value
+		/// </summary>
+		public string Text { get; set; }
+
+		/// <summary>
+		/// Properties associated with the value
+		/// </summary>
+		public Dictionary<Utf8String, object?>? Properties { get; }
+
+		/// <summary>
+		/// Constructor
+		/// </summary>
+		/// <param name="type">Type of the value</param>
+		/// <param name="text">Rendering of the value as text</param>
+		/// <param name="properties">Additional properties for this value</param>
+		public LogValue(Utf8String type, string text, Dictionary<Utf8String, object?>? properties = null)
+		{
+			Type = type;
+			Text = text;
+			Properties = properties;
+		}
+
+		/// <summary>
+		/// Construct a source file log value with an overridden display text
+		/// </summary>
+		/// <param name="file">Source file to reference</param>
+		/// <param name="text">Display text for the fiel</param>
+		public static LogValue SourceFile(FileReference file, string text)
+		{
+			return new LogValue(LogValueType.SourceFile, text, new Dictionary<Utf8String, object?> { [LogEventPropertyName.File] = file.FullName });
+		}
+
+		/// <summary>
+		/// Constructs a value which contains a hyperlink to an external page
+		/// </summary>
+		/// <param name="target">Target URL</param>
+		public static LogValue Link(Uri target)
+		{
+			return Link(target, target.ToString());
+		}
+
+		/// <summary>
+		/// Constructs a value which contains a hyperlink to an external page
+		/// </summary>
+		/// <param name="target">Target URL</param>
+		/// <param name="text">Text to render for the link</param>
+		public static LogValue Link(Uri target, string text)
+		{
+			return new LogValue(LogValueType.Link, text, new Dictionary<Utf8String, object?> { [LogEventPropertyName.Target] = target.ToString() });
+		}
+
+		/// <summary>
+		/// Creates a LogValue from an object, overriding the type and display text for it
+		/// </summary>
+		/// <param name="obj">The object to construct from</param>
+		/// <returns></returns>
+		public static LogValue FromObject(object obj) => FromObject(LogValueType.Object, obj.ToString() ?? String.Empty, obj);
+
+		/// <summary>
+		/// Creates a LogValue from an object, overriding the type and display text for it
+		/// </summary>
+		/// <param name="type">Type of the object</param>
+		/// <param name="text">Rendered representation of the object in the output string</param>
+		/// <param name="obj">The object to construct from</param>
+		/// <returns></returns>
+		public static LogValue FromObject(Utf8String type, string text, object obj)
+		{
+			Type objType = obj.GetType();
+
+			Dictionary<Utf8String, object?>? properties = null;
+			foreach (PropertyInfo propertyInfo in objType.GetProperties())
+			{
+				if (propertyInfo.GetCustomAttribute<JsonIgnoreAttribute>() == null)
+				{
+					string name = propertyInfo.GetCustomAttribute<JsonPropertyNameAttribute>()?.Name ?? propertyInfo.Name;
+					properties ??= [];
+					properties[new Utf8String(name)] = propertyInfo.GetValue(obj)!;
+				}
+			}
+
+			return new LogValue(type, text, properties);
+		}
+
+		/// <inheritdoc/>
+		public override string ToString()
+		{
+			return Text;
+		}
+	}
+}
diff --git a/Engine/Source/Programs/Shared/EpicGames.Core/Logging/LogValueFormatter.cs b/Engine/Source/Programs/Shared/EpicGames.Core/Logging/LogValueFormatter.cs
new file mode 100644
--- /dev/null
+++ b/Engine/Source/Programs/Shared/EpicGames.Core/Logging/LogValueFormatter.cs
@@ -0,0 +1,409 @@
+// Copyright Epic Games, Inc. All Rights Reserved.
+
+using System;
+using System.Collections;
+using System.Collections.Concurrent;
+using System.Collections.Generic;
+using System.Diagnostics;
+using System.Diagnostics.CodeAnalysis;
+using System.Linq;
+using System.Reflection;
+using System.Text.Json;
+
+namespace EpicGames.Core
+{
+	/// <summary>
+	/// Formatter for objects of a given type
+	/// </summary>
+	public interface ILogValueFormatter
+	{
+		/// <summary>
+		/// Format an object to a JSON output stream
+		/// </summary>
+		/// <param name="value"></param>
+		/// <param name="writer"></param>
+		void Format(object value, Utf8JsonWriter writer);
+	}
+
+	/// <summary>
+	/// Specifies a type to use for serializing objects to a log event
+	/// </summary>
+	[AttributeUsage(AttributeTargets.Struct | AttributeTargets.Class)]
+	public sealed class LogValueFormatterAttribute : Attribute
+	{
+		/// <summary>
+		/// Type of the formatter
+		/// </summary>
+		public Type Type { get; }
+
+		/// <summary>
+		/// Constructor
+		/// </summary>
+		/// <param name="type"></param>
+		public LogValueFormatterAttribute(Type type)
+		{
+			Type = type;
+		}
+	}
+
+	/// <summary>
+	/// Utility methods for manipulating formatters
+	/// </summary>
+	public static class LogValueFormatter
+	{
+		class BoolFormatter : ILogValueFormatter
+		{
+			public void Format(object value, Utf8JsonWriter writer) => writer.WriteBooleanValue((bool)value);
+		}
+
+		class Int32Formatter : ILogValueFormatter
+		{
+			public void Format(object value, Utf8JsonWriter writer) => writer.WriteNumberValue((int)value);
+		}
+
+		class UInt32Formatter : ILogValueFormatter
+		{
+			public void Format(object value, Utf8JsonWriter writer) => writer.WriteNumberValue((uint)value);
+		}
+
+		class Int64Formatter : ILogValueFormatter
+		{
+			public void Format(object value, Utf8JsonWriter writer) => writer.WriteNumberValue((long)value);
+		}
+
+		class UInt64Formatter : ILogValueFormatter
+		{
+			public void Format(object value, Utf8JsonWriter writer) => writer.WriteNumberValue((ulong)value);
+		}
+
+		class FloatFormatter : ILogValueFormatter
+		{
+			public void Format(object value, Utf8JsonWriter writer) => writer.WriteNumberValue((float)value);
+		}
+
+		class DoubleFormatter : ILogValueFormatter
+		{
+			public void Format(object value, Utf8JsonWriter writer) => writer.WriteNumberValue((double)value);
+		}
+
+		class StringFormatter : ILogValueFormatter
+		{
+			public void Format(object value, Utf8JsonWriter writer) => writer.WriteStringValue(value.ToString());
+		}
+
+		class ActivityFormatter : ILogValueFormatter
+		{
+			public void Format(object value, Utf8JsonWriter writer)
+			{
+				writer.WriteStartArray();
+				Format((Activity)value, writer);
+				writer.WriteEndArray();
+			}
+
+			static void Format(Activity activity, Utf8JsonWriter writer)
+			{
+				if (activity.Parent != null)
+				{
+					Format(activity.Parent, writer);
+				}
+
+				writer.WriteStartObject();
+				writer.WriteString("name", activity.OperationName);
+				if (activity.Tags.Any())
+				{
+					writer.WriteStartObject("tags");
+					foreach (KeyValuePair<string, string?> pair in activity.Tags)
+					{
+						writer.WriteString(pair.Key, pair.Value);
+					}
+					writer.WriteEndObject();
+				}
+				writer.WriteEndObject();
+			}
+		}
+
+		class EnumerableFormatter : ILogValueFormatter
+		{
+			readonly ILogValueFormatter _elementFormatter;
+
+			public EnumerableFormatter(ILogValueFormatter elementFormatter) => _elementFormatter = elementFormatter;
+
+			public void Format(object value, Utf8JsonWriter writer)
+			{
+				writer.WriteStartArray();
+				foreach (object element in (IEnumerable)value)
+				{
+					_elementFormatter.Format(element, writer);
+				}
+				writer.WriteEndArray();
+			}
+		}
+
+		class StructuredLogValueFormatter : ILogValueFormatter
+		{
+			public void Format(object value, Utf8JsonWriter writer)
+			{
+				LogValue logValue = (LogValue)value;
+				writer.WriteStartObject();
+				writer.WriteString(LogEventPropertyName.Type, logValue.Type);
+				writer.WriteString(LogEventPropertyName.Text, logValue.Text);
+				if (logValue.Properties != null)
+				{
+					foreach (KeyValuePair<Utf8String, object?> pair in logValue.Properties)
+					{
+						writer.WritePropertyName(pair.Key);
+						WritePropertyValue(pair.Value, writer);
+					}
+				}
+				writer.WriteEndObject();
+			}
+
+			static void WritePropertyValue(object? value, Utf8JsonWriter writer)
+			{
+				if (value == null)
+				{
+					writer.WriteNullValue();
+				}
+				else
+				{
+					Type valueType = value.GetType();
+					switch (Type.GetTypeCode(valueType))
+					{
+						case TypeCode.Boolean:
+							writer.WriteBooleanValue((bool)value);
+							break;
+						case TypeCode.Byte:
+							writer.WriteNumberValue((byte)value);
+							break;
+						case TypeCode.SByte:
+							writer.WriteNumberValue((sbyte)value);
+							break;
+						case TypeCode.UInt16:
+							writer.WriteNumberValue((ushort)value);
+							break;
+						case TypeCode.UInt32:
+							writer.WriteNumberValue((uint)value);
+							break;
+						case TypeCode.UInt64:
+							writer.WriteNumberValue((long)value);
+							break;
+						case TypeCode.Int16:
+							writer.WriteNumberValue((short)value);
+							break;
+						case TypeCode.Int32:
+							writer.WriteNumberValue((int)value);
+							break;
+						case TypeCode.Int64:
+							writer.WriteNumberValue((long)value);
+							break;
+						case TypeCode.Decimal:
+							writer.WriteNumberValue((decimal)value);
+							break;
+						case TypeCode.Double:
+							writer.WriteNumberValue((double)value);
+							break;
+						case TypeCode.Single:
+							writer.WriteNumberValue((float)value);
+							break;
+						case TypeCode.String:
+							writer.WriteStringValue((string)value);
+							break;
+						default:
+							JsonSerializer.Serialize(writer, value);
+							break;
+					}
+				}
+			}
+		}
+
+		class AnnotateTypeFormatter : ILogValueFormatter
+		{
+			readonly Utf8String _type;
+
+			public AnnotateTypeFormatter(string type) => _type = new Utf8String(type);
+
+			public void Format(object value, Utf8JsonWriter writer)
+			{
+				writer.WriteStartObject();
+				writer.WriteString(LogEventPropertyName.Type, _type);
+				writer.WriteString(LogEventPropertyName.Text, value.ToString());
+				writer.WriteEndObject();
+			}
+		}
+
+		class DictionaryFormatter : ILogValueFormatter
+		{
+			public void Format(object value, Utf8JsonWriter writer)
+			{
+				Dictionary<string, object> dictionary = (Dictionary<string, object>)value;
+				writer.WriteStartObject();
+				foreach (KeyValuePair<string, object> kvp in dictionary)
+				{
+					writer.WritePropertyName(kvp.Key);
+					Format(kvp.Value, writer);
+				}
+				writer.WriteEndObject();
+			}
+		}
+
+		class Utf8DictionaryFormatter : ILogValueFormatter
+		{
+			public void Format(object value, Utf8JsonWriter writer)
+			{
+				Dictionary<Utf8String, object> dictionary = (Dictionary<Utf8String, object>)value;
+				writer.WriteStartObject();
+				foreach (KeyValuePair<Utf8String, object> kvp in dictionary)
+				{
+					writer.WritePropertyName(kvp.Key);
+					LogValueFormatter.Format(kvp.Value, writer);
+				}
+				writer.WriteEndObject();
+			}
+		}
+
+		class FileReferenceFormatter : ILogValueFormatter
+		{
+			public void Format(object value, Utf8JsonWriter writer)
+			{
+				FileReference file = (FileReference)value;
+				writer.WriteStartObject();
+				writer.WriteString(LogEventPropertyName.Type, (file.HasExtension(".uasset") || file.HasExtension(".umap")) ? LogValueType.Asset : LogValueType.SourceFile);
+				writer.WriteString(LogEventPropertyName.Text, file.FullName);
+				writer.WriteEndObject();
+			}
+		}
+
+		static readonly StringFormatter s_stringFormatter = new();
+
+		static readonly ConcurrentDictionary<Type, ILogValueFormatter> s_formatters = GetDefaultFormatters();
+
+		static ConcurrentDictionary<Type, ILogValueFormatter> GetDefaultFormatters()
+		{
+			ConcurrentDictionary<Type, ILogValueFormatter> formatters = new();
+			formatters.TryAdd(typeof(bool), new BoolFormatter());
+			formatters.TryAdd(typeof(int), new Int32Formatter());
+			formatters.TryAdd(typeof(uint), new UInt32Formatter());
+			formatters.TryAdd(typeof(long), new Int64Formatter());
+			formatters.TryAdd(typeof(ulong), new UInt64Formatter());
+			formatters.TryAdd(typeof(float), new FloatFormatter());
+			formatters.TryAdd(typeof(double), new DoubleFormatter());
+			formatters.TryAdd(typeof(string), s_stringFormatter);
+			formatters.TryAdd(typeof(Dictionary<string, object>), new DictionaryFormatter());
+			formatters.TryAdd(typeof(Dictionary<Utf8String, object>), new Utf8DictionaryFormatter());
+			formatters.TryAdd(typeof(LogValue), new StructuredLogValueFormatter());
+			formatters.TryAdd(typeof(FileReference), new FileReferenceFormatter());
+			formatters.TryAdd(typeof(Activity), new ActivityFormatter());
+			return formatters;
+		}
+
+		/// <summary>
+		/// Writes a typed 
+		/// </summary>
+		/// <param name="writer"></param>
+		/// <param name="type"></param>
+		/// <param name="text"></param>
+		public static void WriteTypedValue(Utf8JsonWriter writer, string type, string text)
+		{
+			writer.WriteStartObject();
+			writer.WriteString(LogEventPropertyName.Type, type);
+			writer.WriteString(LogEventPropertyName.Text, text);
+			writer.WriteEndObject();
+		}
+
+		/// <summary>
+		/// Registers a formatter for a particular type
+		/// </summary>
+		/// <param name="type"></param>
+		/// <param name="formatter"></param>
+		public static void RegisterFormatter(Type type, ILogValueFormatter formatter)
+		{
+			s_formatters.TryAdd(type, formatter);
+		}
+
+		/// <summary>
+		/// Registers a formatter for a particular type that adds a $type field to the serialized output
+		/// </summary>
+		/// <param name="type">Name of the type</param>
+		public static void RegisterTypeAnnotation<T>(string type)
+		{
+			s_formatters.TryAdd(typeof(T), new AnnotateTypeFormatter(type));
+		}
+
+		/// <summary>
+		/// Gets a formatter for the specified type
+		/// </summary>
+		/// <param name="type"></param>
+		/// <returns></returns>
+		public static ILogValueFormatter GetFormatter(Type type)
+		{
+			for (; ; )
+			{
+				ILogValueFormatter? formatter;
+				if (s_formatters.TryGetValue(type, out formatter))
+				{
+					return formatter;
+				}
+
+				formatter = GetFormatterUncached(type);
+				if (s_formatters.TryAdd(type, formatter))
+				{
+					return formatter;
+				}
+			}
+		}
+
+		static ILogValueFormatter GetFormatterUncached(Type type)
+		{
+			LogValueFormatterAttribute? formatterAttribute = type.GetCustomAttribute<LogValueFormatterAttribute>();
+			if (formatterAttribute != null)
+			{
+				return (ILogValueFormatter)Activator.CreateInstance(formatterAttribute.Type)!;
+			}
+
+			LogValueTypeAttribute? typeAttribute = type.GetCustomAttribute<LogValueTypeAttribute>();
+			if (typeAttribute != null)
+			{
+				return new AnnotateTypeFormatter(typeAttribute.Name ?? type.Name);
+			}
+
+			if (TryGetEnumerableType(type, out Type? elementType))
+			{
+				return new EnumerableFormatter(GetFormatter(elementType));
+			}
+
+			return s_stringFormatter;
+		}
+
+		static bool TryGetEnumerableType(Type type, [NotNullWhen(true)] out Type? elementType)
+		{
+			foreach (Type interfaceType in type.GetInterfaces())
+			{
+				if (interfaceType.IsGenericType && interfaceType.GetGenericTypeDefinition() == typeof(IEnumerable<>))
+				{
+					elementType = interfaceType.GetGenericArguments()[0];
+					return true;
+				}
+			}
+
+			elementType = null;
+			return false;
+		}
+
+		/// <summary>
+		/// Formats a value to a Json log stream
+		/// </summary>
+		/// <param name="value">Value to write</param>
+		/// <param name="writer">Json writer</param>
+		public static void Format(object? value, Utf8JsonWriter writer)
+		{
+			if (value == null)
+			{
+				writer.WriteNullValue();
+				return;
+			}
+
+			ILogValueFormatter formatter = GetFormatter(value.GetType());
+			formatter.Format(value, writer);
+		}
+	}
+}
diff --git a/Engine/Source/Programs/Shared/EpicGames.Core/Logging/MessageTemplate.cs b/Engine/Source/Programs/Shared/EpicGames.Core/Logging/MessageTemplate.cs
new file mode 100644
--- /dev/null
+++ b/Engine/Source/Programs/Shared/EpicGames.Core/Logging/MessageTemplate.cs
@@ -0,0 +1,297 @@
+// Copyright Epic Games, Inc. All Rights Reserved.
+
+using System;
+using System.Collections.Generic;
+using System.Globalization;
+using System.Text;
+using System.Text.Json;
+
+namespace EpicGames.Core
+{
+	/// <summary>
+	/// Utility class for dealing with message templates
+	/// </summary>
+	public static class MessageTemplate
+	{
+		/// <summary>
+		/// The default property name for the message template format string in an enumerable log state parameter
+		/// </summary>
+		public const string FormatPropertyName = "{OriginalFormat}";
+
+		/// <summary>
+		/// Renders a format string
+		/// </summary>
+		/// <param name="format">The format string</param>
+		/// <param name="properties">Property values to embed</param>
+		/// <returns>The rendered string</returns>
+		public static string Render(string format, IEnumerable<KeyValuePair<string, object?>>? properties)
+		{
+			StringBuilder result = new();
+			Render(format, properties, result);
+			return result.ToString();
+		}
+
+		/// <summary>
+		/// Renders a format string to the end of a string builder
+		/// </summary>
+		/// <param name="format">The format string to render</param>
+		/// <param name="properties">Sequence of key/value properties</param>
+		/// <param name="result">Buffer to append the rendered string to</param>
+		public static void Render(string format, IEnumerable<KeyValuePair<string, object?>>? properties, StringBuilder result)
+		{
+			int nextOffset = 0;
+
+			List<(int, int)>? names = ParsePropertyNames(format);
+			if (names != null && properties != null)
+			{
+				foreach ((int offset, int length) in names)
+				{
+					ReadOnlySpan<char> argument = format.AsSpan(offset, length);
+
+					string? formatSpecifier = null;
+
+					// Parse the format specifier
+					int formatSpecifierIdx = argument.IndexOf(':');
+					if (formatSpecifierIdx != -1)
+					{
+						formatSpecifier = argument.Slice(formatSpecifierIdx + 1).ToString();
+						argument = argument.Slice(0, formatSpecifierIdx);
+					}
+
+					// Parse the property width
+					int alignment = 0;
+
+					int alignmentIdx = argument.IndexOf(',');
+					if (alignmentIdx != -1 && Int32.TryParse(argument.Slice(alignmentIdx + 1), NumberStyles.Integer | NumberStyles.AllowLeadingSign, null, out alignment))
+					{
+						argument = argument.Slice(0, alignmentIdx);
+					}
+
+					// Try to get the property value
+					object? value;
+					if (TryGetPropertyValue(argument, properties, out value))
+					{
+						// Append the text up to this argument
+						int startOffset = offset - 1;
+						if (format[startOffset] == '@' || format[startOffset] == '$')
+						{
+							startOffset--;
+						}
+						Unescape(format.AsSpan(nextOffset, startOffset - nextOffset), result);
+
+						// Render the property
+						string? rendered;
+						if (format[offset] == '@')
+						{
+							rendered = JsonSerializer.Serialize(value, value?.GetType() ?? typeof(object), new JsonSerializerOptions { PropertyNamingPolicy = JsonNamingPolicy.CamelCase });
+						}
+						else if (formatSpecifier != null && value is IFormattable formattable)
+						{
+							rendered = formattable.ToString(formatSpecifier, null);
+						}
+						else
+						{
+							rendered = value?.ToString() ?? "null";
+						}
+
+						// Apply a postive width adjustment (right align)
+						int leftPad = Math.Max(alignment - rendered.Length, 0);
+						if (leftPad > 0)
+						{
+							result.Append(' ', leftPad);
+						}
+
+						// Append the argument
+						result.Append(rendered);
+
+						// Apply a negative width adjustment (left align)
+						int rightPad = Math.Max(-alignment - rendered.Length, 0);
+						if (rightPad > 0)
+						{
+							result.Append(' ', rightPad);
+						}
+
+						// Start the next plain-text run after the closing brace
+						nextOffset = offset + length + 1;
+					}
+				}
+			}
+
+			Unescape(format.AsSpan(nextOffset, format.Length - nextOffset), result);
+		}
+
+		/// <summary>
+		/// Escapes a string for use in a message template
+		/// </summary>
+		/// <param name="text">Text to escape</param>
+		/// <returns>The escaped string</returns>
+		public static string Escape(string text)
+		{
+			StringBuilder result = new();
+			Escape(text, result);
+			return result.ToString();
+		}
+
+		/// <summary>
+		/// Escapes a span of characters and appends the result to a string
+		/// </summary>
+		/// <param name="text">Span of characters to escape</param>
+		/// <param name="result">Buffer to receive the escaped string</param>
+		public static void Escape(ReadOnlySpan<char> text, StringBuilder result)
+		{
+			foreach (char character in text)
+			{
+				result.Append(character);
+				if (character == '{' || character == '}')
+				{
+					result.Append(character);
+				}
+			}
+		}
+
+		/// <summary>
+		/// Unescapes a string from a message template
+		/// </summary>
+		/// <param name="text">The text to unescape</param>
+		/// <returns>The unescaped text</returns>
+		public static string Unescape(string text)
+		{
+			StringBuilder result = new();
+			Unescape(text.AsSpan(), result);
+			return result.ToString();
+		}
+
+		/// <summary>
+		/// Unescape a string and append the result to a string builder
+		/// </summary>
+		/// <param name="text">Text to unescape</param>
+		/// <param name="result">Receives the unescaped text</param>
+		public static void Unescape(ReadOnlySpan<char> text, StringBuilder result)
+		{
+			char lastChar = '\0';
+			foreach (char character in text)
+			{
+				if ((character != '{' && character != '}') || character != lastChar)
+				{
+					result.Append(character);
+				}
+				lastChar = character;
+			}
+		}
+
+		/// <summary>
+		/// Finds locations of property names from the given format string
+		/// </summary>
+		/// <param name="format">The format string to parse</param>
+		/// <returns>List of offset, length pairs for property names. Null if the string does not contain any property references.</returns>
+		public static List<(int, int)>? ParsePropertyNames(string format)
+		{
+			List<(int, int)>? names = null;
+			for (int idx = 0; idx < format.Length - 1; idx++)
+			{
+				if (format[idx] == '{')
+				{
+					if (format[idx + 1] == '{')
+					{
+						idx++;
+					}
+					else
+					{
+						int startIdx = idx + 1;
+
+						idx = format.IndexOf('}', startIdx);
+						if (idx == -1)
+						{
+							break;
+						}
+						names ??= [];
+
+						names.Add((startIdx, idx - startIdx));
+					}
+				}
+			}
+			return names;
+		}
+
+		/// <summary>
+		/// Parse the ordered arguments into a dictionary of named properties
+		/// </summary>
+		/// <param name="format">Format string</param>
+		/// <param name="args">Argument list to parse</param>
+		/// <param name="properties"></param>
+		/// <returns></returns>
+		public static void ParsePropertyValues(string format, object[] args, Dictionary<string, object?> properties)
+		{
+			List<(int, int)>? offsets = ParsePropertyNames(format);
+			if (offsets != null)
+			{
+				for (int idx = 0; idx < offsets.Count; idx++)
+				{
+					string name = format.Substring(offsets[idx].Item1, offsets[idx].Item2);
+
+					// Parse for format specifier
+					int specifierIdx = name.IndexOf(':', 0);
+					if (specifierIdx != -1)
+					{
+						name = name.Substring(0, specifierIdx);
+					}
+
+					int number;
+					if (Int32.TryParse(name, out number))
+					{
+						if (number >= 0 && number < args.Length)
+						{
+							properties[name] = args[number];
+						}
+					}
+					else
+					{
+						if (idx < args.Length)
+						{
+							properties[name] = args[idx];
+						}
+					}
+				}
+			}
+		}
+
+		/// <summary>
+		/// Attempts to get a named property value from the given dictionary
+		/// </summary>
+		/// <param name="name">Name of the property</param>
+		/// <param name="properties">Sequence of property name/value pairs</param>
+		/// <param name="value">On success, receives the property value</param>
+		/// <returns>True if the property was found, false otherwise</returns>
+		public static bool TryGetPropertyValue(ReadOnlySpan<char> name, IEnumerable<KeyValuePair<string, object?>> properties, out object? value)
+		{
+			int number;
+			if (Int32.TryParse(name, NumberStyles.Integer, null, out number))
+			{
+				foreach (KeyValuePair<string, object?> property in properties)
+				{
+					if (number == 0)
+					{
+						value = property.Value;
+						return true;
+					}
+					number--;
+				}
+			}
+			else
+			{
+				foreach (KeyValuePair<string, object?> property in properties)
+				{
+					ReadOnlySpan<char> parameterName = property.Key.AsSpan();
+					if (name.Equals(parameterName, StringComparison.Ordinal))
+					{
+						value = property.Value;
+						return true;
+					}
+				}
+			}
+
+			value = null;
+			return false;
+		}
+	}
+}
diff --git a/Engine/Source/Programs/Shared/EpicGames.Core/Logging/PrefixLogger.cs b/Engine/Source/Programs/Shared/EpicGames.Core/Logging/PrefixLogger.cs
new file mode 100644
--- /dev/null
+++ b/Engine/Source/Programs/Shared/EpicGames.Core/Logging/PrefixLogger.cs
@@ -0,0 +1,53 @@
+// Copyright Epic Games, Inc. All Rights Reserved.
+
+using System;
+using System.Collections.Generic;
+using Microsoft.Extensions.Logging;
+
+namespace EpicGames.Core
+{
+	/// <summary>
+	/// ILogger implementation that adds a prefix to the start of each line
+	/// </summary>
+	public class PrefixLogger : ILogger
+	{
+		readonly string _prefix;
+		readonly ILogger _inner;
+
+		/// <summary>
+		/// Constructor
+		/// </summary>
+		/// <param name="prefix">Prefix to insert</param>
+		/// <param name="inner">Logger to write to</param>
+		public PrefixLogger(string prefix, ILogger inner)
+		{
+			_prefix = prefix;
+			_inner = inner;
+		}
+
+		/// <inheritdoc/>
+		public IDisposable? BeginScope<TState>(TState state) where TState : notnull => _inner.BeginScope<TState>(state);
+
+		/// <inheritdoc/>
+		public bool IsEnabled(LogLevel logLevel) => _inner.IsEnabled(logLevel);
+
+		/// <inheritdoc/>
+		public void Log<TState>(LogLevel logLevel, EventId eventId, TState state, Exception? exception, Func<TState, Exception?, string> formatter)
+		{
+			if (state is IEnumerable<KeyValuePair<string, object>> enumerable)
+			{
+				List<KeyValuePair<string, object>> copy = [.. enumerable];
+
+				int idx = copy.FindIndex(x => x.Key.Equals("{OriginalFormat}", StringComparison.OrdinalIgnoreCase));
+				if (idx != -1 && copy[idx].Value is string format)
+				{
+					copy[idx] = new KeyValuePair<string, object>(copy[idx].Key, "{_tag} " + format);
+					copy.Add(new KeyValuePair<string, object>("_tag", _prefix));
+					_inner.Log(logLevel, eventId, copy, exception, (s, e) => $"{_prefix} {formatter(state, exception)}");
+					return;
+				}
+			}
+			_inner.Log(logLevel, eventId, state, exception, formatter);
+		}
+	}
+}
diff --git a/Engine/Source/Programs/Shared/EpicGames.Core/MessageTemplate.cs b/Engine/Source/Programs/Shared/EpicGames.Core/MessageTemplate.cs
deleted file mode 100644
diff --git a/Engine/Source/Programs/Shared/EpicGames.Core/PrefixLogger.cs b/Engine/Source/Programs/Shared/EpicGames.Core/PrefixLogger.cs
deleted file mode 100644
