BuildGraph V2 fixes and improvements

- Added support for attaching annotations to nodes
- You can now supply a node name when instantiating a node
- RunEarly is now supported on nodes
- Fix for weak/after/fences attribute on nodes
- Fix for dependency order of node outputs

#Synced-CL 52380225

diff --git a/Engine/Source/Programs/AutomationTool/BuildGraph/BgScriptReader.cs b/Engine/Source/Programs/AutomationTool/BuildGraph/BgScriptReader.cs
--- a/Engine/Source/Programs/AutomationTool/BuildGraph/BgScriptReader.cs
+++ b/Engine/Source/Programs/AutomationTool/BuildGraph/BgScriptReader.cs
@@ -1571,7 +1571,16 @@ namespace AutomationTool
 					};
 					foreach ((string key, string value) in annotations)
 					{
-						newNode.Annotations[key] = value;
+						BgAnnotationDef? annotationDef = newNode.Annotations.FirstOrDefault(a => a.Name.Equals(key, StringComparison.OrdinalIgnoreCase));
+						if (annotationDef == null)
+						{
+							annotationDef = new BgAnnotationDef(key, value);
+							newNode.Annotations.Add(annotationDef);
+						}
+						else
+						{
+							annotationDef.Value = value;
+						}
 					}
 					_graph.NameToNode.Add(name, newNode);
 
@@ -2018,7 +2027,17 @@ namespace AutomationTool
 				{
 					foreach ((string key, string value) in annotations)
 					{
-						node.Annotations[key] = value;
+						BgAnnotationDef? annotationDef = node.Annotations.FirstOrDefault(a => a.Name.Equals(key, StringComparison.OrdinalIgnoreCase));
+						if (annotationDef == null)
+						{
+							annotationDef = new BgAnnotationDef(key, value);
+							node.Annotations.Add(annotationDef);
+						}
+						else
+						{
+							annotationDef.Value = value;
+						}
+
 					}
 				}
 			}
diff --git a/Engine/Source/Programs/Shared/EpicGames.BuildGraph/BgAnnotationDef.cs b/Engine/Source/Programs/Shared/EpicGames.BuildGraph/BgAnnotationDef.cs
new file mode 100644
--- /dev/null
+++ b/Engine/Source/Programs/Shared/EpicGames.BuildGraph/BgAnnotationDef.cs
@@ -0,0 +1,37 @@
+﻿// Copyright Epic Games, Inc. All Rights Reserved.
+
+namespace EpicGames.BuildGraph
+{
+	/// <summary>
+	/// Defines an annotation key/value pair for BuildGraph nodes
+	/// </summary>
+	public class BgAnnotationDef
+	{
+		/// <summary>
+		/// Name, or key, of the annotation
+		/// </summary>
+		public string Name { get; set; }
+		
+		/// <summary>
+		/// Value to be paired with the annotation name
+		/// </summary>
+		public string Value { get; set; }
+
+		/// <summary>
+		/// Constructor
+		/// </summary>
+		public BgAnnotationDef(string name, string value)
+		{
+			Name = name;
+			Value = value;
+		}
+		
+		/// <summary>
+		/// Simply the annotation name.
+		/// </summary>
+		public override string ToString()
+		{
+			return Name;
+		}
+	}
+}
\ No newline at end of file
diff --git a/Engine/Source/Programs/Shared/EpicGames.BuildGraph/BgGraphDef.cs b/Engine/Source/Programs/Shared/EpicGames.BuildGraph/BgGraphDef.cs
--- a/Engine/Source/Programs/Shared/EpicGames.BuildGraph/BgGraphDef.cs
+++ b/Engine/Source/Programs/Shared/EpicGames.BuildGraph/BgGraphDef.cs
@@ -234,6 +234,12 @@ namespace EpicGames.BuildGraph
 				{
 					AddDependencies(inputDependency, retainNodes);
 				}
+				
+				foreach (BgNodeDef orderDependency in node.OrderDependencies)
+				{
+					AddDependencies(orderDependency, retainNodes);
+				}
+
 				foreach (BgNodeDef optionalInputDependency in node.OptionalInputDependencies)
 				{
 					AddDependencies(optionalInputDependency, retainNodes);
@@ -558,9 +564,9 @@ namespace EpicGames.BuildGraph
 						jsonWriter.WriteArrayEnd();
 
 						jsonWriter.WriteObjectStart("Annotations");
-						foreach ((string key, string value) in node.Annotations)
+						foreach (BgAnnotationDef annotation in node.Annotations)
 						{
-							jsonWriter.WriteValue(key, value);
+							jsonWriter.WriteValue(annotation.Name, annotation.Value);
 						}
 						jsonWriter.WriteObjectEnd();
 
diff --git a/Engine/Source/Programs/Shared/EpicGames.BuildGraph/BgNodeDef.cs b/Engine/Source/Programs/Shared/EpicGames.BuildGraph/BgNodeDef.cs
--- a/Engine/Source/Programs/Shared/EpicGames.BuildGraph/BgNodeDef.cs
+++ b/Engine/Source/Programs/Shared/EpicGames.BuildGraph/BgNodeDef.cs
@@ -173,9 +173,9 @@ namespace EpicGames.BuildGraph
 		public bool NotifyOnWarnings { get; set; } = true;
 
 		/// <summary>
-		/// Custom annotations for this node
+		/// Annotations for this node
 		/// </summary>
-		public Dictionary<string, string> Annotations { get; } = new Dictionary<string, string>(StringComparer.OrdinalIgnoreCase);
+		public List<BgAnnotationDef> Annotations { get; } = [];
 
 		/// <summary>
 		/// Ignore modified files matching the patterns provided
@@ -313,6 +313,11 @@ namespace EpicGames.BuildGraph
 		/// </summary>
 		public List<BgNodeOutputExprDef> OptionalInputExprs { get; } = [];
 
+		/// <summary>
+		/// Order expressions
+		/// </summary>
+		public List<BgNodeExpressionDef> OrderExprs { get; } = [];
+		
 		/// <summary>
 		/// Number of outputs from this node
 		/// </summary>
@@ -354,6 +359,8 @@ namespace EpicGames.BuildGraph
 
 			inputDependencies.ExceptWith(node.InputDependencies);
 			node.InputDependencies.AddRange(inputDependencies);
+			
+			node.OrderDependencies.AddRange(node.OrderExprs);
 
 			HashSet<BgNodeDef> optionalInputDependencies = [];
 			foreach (BgNodeOutput optionalInput in node.OptionalInputs)
diff --git a/Engine/Source/Programs/Shared/EpicGames.BuildGraph/Expressions/BgAnnotation.cs b/Engine/Source/Programs/Shared/EpicGames.BuildGraph/Expressions/BgAnnotation.cs
new file mode 100644
--- /dev/null
+++ b/Engine/Source/Programs/Shared/EpicGames.BuildGraph/Expressions/BgAnnotation.cs
@@ -0,0 +1,39 @@
+﻿// Copyright Epic Games, Inc. All Rights Reserved.
+
+namespace EpicGames.BuildGraph.Expressions;
+
+/// <summary>
+/// Specification for a annotation
+/// </summary>
+public class BgAnnotation : BgExpr
+{
+	/// <summary>
+	/// Name, or key, of the annotation
+	/// </summary>
+	public BgString Name { get; }
+
+	/// <summary>
+	/// Value to be paired with the annotation name
+	/// </summary>
+	public BgString Value { get; }
+
+	/// <inheritdoc/>
+	public BgAnnotation(BgString name, BgString value)
+		: base(BgExprFlags.ForceFragment)
+	{
+		Name = name;
+		Value = value;
+	}
+
+	/// <inheritdoc/>
+	public override void Write(BgBytecodeWriter writer)
+	{
+		BgObject<BgAnnotationDef> obj = BgObject<BgAnnotationDef>.Empty;
+		obj = obj.Set(x => x.Name, Name);
+		obj = obj.Set(x => x.Value, Value);
+		writer.WriteExpr(obj);
+	}
+
+	/// <inheritdoc/>
+	public override BgString ToBgString() => "{Annotation}";
+}
\ No newline at end of file
diff --git a/Engine/Source/Programs/Shared/EpicGames.BuildGraph/Expressions/BgNode.cs b/Engine/Source/Programs/Shared/EpicGames.BuildGraph/Expressions/BgNode.cs
--- a/Engine/Source/Programs/Shared/EpicGames.BuildGraph/Expressions/BgNode.cs
+++ b/Engine/Source/Programs/Shared/EpicGames.BuildGraph/Expressions/BgNode.cs
@@ -100,6 +100,11 @@ namespace EpicGames.BuildGraph.Expressions
 		/// Labels that this node contributes to
 		/// </summary>
 		public BgList<BgLabel> Labels { get; private set; } = BgList.Empty<BgLabel>();
+		
+		/// <summary>
+		/// Custom annotations for this node
+		/// </summary>
+		public BgList<BgAnnotation> Annotations { get; private set; } = BgList.Empty<BgAnnotation>();
 
 		/// <summary>
 		/// Constructor
@@ -129,6 +134,7 @@ namespace EpicGames.BuildGraph.Expressions
 			Fences = node.Fences;
 			RunEarly = node.RunEarly;
 			Labels = node.Labels;
+			Annotations = node.Annotations;
 		}
 
 		/// <inheritdoc/>
@@ -141,9 +147,11 @@ namespace EpicGames.BuildGraph.Expressions
 			obj = obj.Set(x => x.OutputCount, (BgInt)OutputCount);
 			obj = obj.Set(x => x.InputExprs, Inputs);
 			obj = obj.Set(x => x.OptionalInputExprs, OptionalInputs);
-			obj = obj.Set(x => x.OrderDependencies, Fences);
+			obj = obj.Set(x => x.OrderExprs, Fences);
+			obj = obj.Set(x => x.OptionalInputExprs, OptionalInputs);
 			obj = obj.Set(x => x.RunEarly, RunEarly);
 			obj = obj.Set(x => x.Labels, Labels);
+			obj = obj.Set(x => x.Annotations, Annotations);
 			writer.WriteExpr(obj);
 		}
 
@@ -155,8 +163,16 @@ namespace EpicGames.BuildGraph.Expressions
 		/// <param name="fences"></param>
 		/// <param name="runEarly"></param>
 		/// <param name="labels"></param>
+		/// <param name="annotations"></param>
 		/// <returns></returns>
-		internal BgNode Modify(BgList<BgFileSet>? inputs = null, BgList<BgFileSet>? optionalInputs = null, BgList<BgNode>? fences = null, BgBool? runEarly = null, BgList<BgLabel>? labels = null)
+		internal BgNode Modify(
+			BgList<BgFileSet>? inputs = null,
+			BgList<BgFileSet>? optionalInputs = null,
+			BgList<BgNode>? fences = null,
+			BgBool? runEarly = null,
+			BgList<BgLabel>? labels = null,
+			BgList<BgAnnotation>? annotations = null
+		)
 		{
 			BgNode node = Clone();
 			if (inputs is not null)
@@ -179,6 +195,10 @@ namespace EpicGames.BuildGraph.Expressions
 			{
 				node.Labels = labels;
 			}
+			if (annotations is not null)
+			{
+				node.Annotations = annotations;
+			}
 			return node;
 		}
 
@@ -407,7 +427,14 @@ namespace EpicGames.BuildGraph.Expressions
 		{
 			if (type == typeof(BgFileSet))
 			{
-				return new BgFileSetFromNodeOutputExpr(this, index);
+				// When BgNodeDef is loaded it automatically inserts another
+				// output into index 0 that represents all outputs. The actual
+				// outputs from the node are inserted after this automatically
+				// generated output. That means that what is index 0 here will
+				// map to the automatically generated output. To compensate for
+				// this we need to increment all output indexes by 1 to match
+				// the index they will actually be once BgNodeDef loads them.
+				return new BgFileSetFromNodeOutputExpr(this, index + 1);
 			}
 			else
 			{
@@ -446,6 +473,7 @@ namespace EpicGames.BuildGraph.Expressions
 			return new BgNode(thunk, agent);
 		}
 		
+
 		/// <summary>
 		/// Creates a node builder for the given agent
 		/// </summary>
@@ -572,5 +600,24 @@ namespace EpicGames.BuildGraph.Expressions
 		{
 			return (T)node.Modify(labels: node.Labels.Add(label));
 		}
+		
+		/// <summary>
+		/// Allows the node to run early.
+		/// </summary>
+		public static T RunEarly<T>(this T node) where T : BgNode
+		{
+			return (T)node.Modify(runEarly: true);
+		}
+		
+		/// <summary>
+		/// Add an annotation to the node.
+		/// </summary>
+		/// <param name="node">The node to modify</param>
+		/// <param name="annotation"></param>
+		/// <returns>The current node spec, to allow chaining calls</returns>
+		public static T AddAnnotation<T>(this T node, BgAnnotation annotation) where T : BgNode
+		{
+			return (T)node.Modify(annotations: node.Annotations.Add(annotation));
+		}
 	}
 }
