Add support for UStructs in namespaces within RigUnit

#Synced-CL 52358354

diff --git a/Engine/Source/Programs/Shared/EpicGames.UHT/Exporters/CodeGen/UhtRigVMStringBuilderExtensions.cs b/Engine/Source/Programs/Shared/EpicGames.UHT/Exporters/CodeGen/UhtRigVMStringBuilderExtensions.cs
--- a/Engine/Source/Programs/Shared/EpicGames.UHT/Exporters/CodeGen/UhtRigVMStringBuilderExtensions.cs
+++ b/Engine/Source/Programs/Shared/EpicGames.UHT/Exporters/CodeGen/UhtRigVMStringBuilderExtensions.cs
@@ -44,9 +44,15 @@ namespace EpicGames.UHT.Exporters.CodeGen
 				parameterTypeSpan = parameterTypeSpan[0..^1];
 			}
 			builder.Append(parameterTypeSpan);
-			if (parameterTypeSpan.Length > 0 && (parameterTypeSpan[0] == 'T' || parameterTypeSpan[0] == 'F'))
+
+			if (parameterTypeSpan.Length > 0)
 			{
-				builder.Append('&');
+				int namespaceEndIndex = parameterTypeSpan.LastIndexOf(':');
+				int firstTypeCharIndex = namespaceEndIndex >= 0 ? (namespaceEndIndex + 1) : 0;
+				if (parameterTypeSpan[firstTypeCharIndex] == 'T' || parameterTypeSpan[firstTypeCharIndex] == 'F')
+				{
+					builder.Append('&');
+				}
 			}
 			return builder;
 		}
diff --git a/Engine/Source/Programs/Shared/EpicGames.UHT/Types/UhtScriptStruct.cs b/Engine/Source/Programs/Shared/EpicGames.UHT/Types/UhtScriptStruct.cs
--- a/Engine/Source/Programs/Shared/EpicGames.UHT/Types/UhtScriptStruct.cs
+++ b/Engine/Source/Programs/Shared/EpicGames.UHT/Types/UhtScriptStruct.cs
@@ -430,7 +430,10 @@ namespace EpicGames.UHT.Types
 		public string TypeConstRef(bool castType = false)
 		{
 			string typeNoRef = TypeNoRef(castType);
-			if (typeNoRef.StartsWith('T') || typeNoRef.StartsWith('F'))
+			int namespaceEndIndex = typeNoRef.LastIndexOf(':');
+			int firstTypeCharIndex = namespaceEndIndex >= 0 ? (namespaceEndIndex + 1) : 0;
+
+			if (typeNoRef.Length > 0 && (typeNoRef[firstTypeCharIndex] == 'T' || typeNoRef[firstTypeCharIndex] == 'F'))
 			{
 				return $"const {typeNoRef}&";
 			}
