Cleaning up baseline generated header includes. IsUEnumClass.h and ReflectedTypeAccessors.h includes should be added even when MinimizeGeneratedIncludes is false; we should not rely on transitive includes of those headers through ScriptMacros.h.

#Synced-CL 52214067

diff --git a/Engine/Source/Programs/Shared/EpicGames.UHT/Exporters/CodeGen/UhtHeaderCodeGeneratorHFile.cs b/Engine/Source/Programs/Shared/EpicGames.UHT/Exporters/CodeGen/UhtHeaderCodeGeneratorHFile.cs
--- a/Engine/Source/Programs/Shared/EpicGames.UHT/Exporters/CodeGen/UhtHeaderCodeGeneratorHFile.cs
+++ b/Engine/Source/Programs/Shared/EpicGames.UHT/Exporters/CodeGen/UhtHeaderCodeGeneratorHFile.cs
@@ -54,35 +54,27 @@ namespace EpicGames.UHT.Exporters.CodeGen
 				// Attempt to limit the headers included. This is needed in the lower level engine code
 				// to get around circular header include issues.
 				builder.Append("\r\n");
-				if (Module.Module.MinimizeGeneratedIncludes)
+
+				builder.Append("#include \"UObject/ObjectMacros.h\"\r\n");
+
+				if (HeaderFile.References.ExportTypes.Any(x => x is UhtEnum or UhtClass or UhtScriptStruct))
+				{
+					// We need to specialize StaticClass/StaticStruct/StaticEnum template 
+					builder.Append("#include \"UObject/ReflectedTypeAccessors.h\"\r\n");
+				}
+				if (HeaderFile.References.ExportTypes.Any(x => x is UhtEnum))
+				{
+					// We need to specialize TIsUEnumClass
+					builder.Append("#include \"Templates/IsUEnumClass.h\"\r\n");
+				}
+
+				if (Module.Module.MinimizeGeneratedIncludes || (HeaderFile.References.ExportTypes.Count > 0 && HeaderFile.References.ExportTypes.Find(x => x is not UhtEnum) == null))
 				{
-					builder.Append("#include \"UObject/ObjectMacros.h\"\r\n");
 					builder.Append("#include \"Templates/NoDestroy.h\"\r\n");
-					if (HeaderFile.References.ExportTypes.Any(x => x is UhtEnum or UhtClass or UhtScriptStruct))
-					{
-						// We need to specialize StaticClass/StaticStruct/StaticEnum template 
-						builder.Append("#include \"UObject/ReflectedTypeAccessors.h\"\r\n");
-					}
-					if (HeaderFile.References.ExportTypes.Any(x => x is UhtEnum))
-					{
-						// We need to specialize TIsUEnumClass
-						builder.Append("#include \"Templates/IsUEnumClass.h\"\r\n");
-					}
 				}
 				else
 				{
-					if (HeaderFile.References.ExportTypes.Count > 0 && HeaderFile.References.ExportTypes.Find(x => x is not UhtEnum) == null)
-					{
-						builder.Append("#include \"Templates/IsUEnumClass.h\"\r\n");
-						builder.Append("#include \"UObject/ObjectMacros.h\"\r\n");
-						builder.Append("#include \"UObject/ReflectedTypeAccessors.h\"\r\n");
-						builder.Append("#include \"Templates/NoDestroy.h\"\r\n");
-					}
-					else
-					{
-						builder.Append("#include \"UObject/ObjectMacros.h\"\r\n");
-						builder.Append("#include \"UObject/ScriptMacros.h\"\r\n");
-					}
+					builder.Append("#include \"UObject/ScriptMacros.h\"\r\n");
 				}
 
 				if (headerInfo.NeedsPushModelHeaders)
