Update test nuget packages

#rnx
#Synced-CL 52406569

diff --git a/Engine/Source/Programs/Shared/EpicGames.Perforce.Fixture/EpicGames.Perforce.Fixture.csproj b/Engine/Source/Programs/Shared/EpicGames.Perforce.Fixture/EpicGames.Perforce.Fixture.csproj
--- a/Engine/Source/Programs/Shared/EpicGames.Perforce.Fixture/EpicGames.Perforce.Fixture.csproj
+++ b/Engine/Source/Programs/Shared/EpicGames.Perforce.Fixture/EpicGames.Perforce.Fixture.csproj
@@ -11,15 +11,6 @@
 
   <ItemGroup>
     <PackageReference Include="Microsoft.Extensions.Logging.Console" Version="8.0.1" />
-    <PackageReference Include="Microsoft.NET.Test.Sdk" Version="17.14.1" />
-    <PackageReference Include="MSTest.TestAdapter" Version="3.11.1">
-      <TreatAsUsed>true</TreatAsUsed>
-    </PackageReference>
-    <PackageReference Include="MSTest.TestFramework" Version="3.11.1" />
-    <PackageReference Include="coverlet.collector" Version="6.0.4">
-      <PrivateAssets>all</PrivateAssets>
-      <IncludeAssets>runtime; build; native; contentfiles; analyzers; buildtransitive</IncludeAssets>
-    </PackageReference>
   </ItemGroup>
 
   <ItemGroup>
