Add support for UStructs in namespaces for TInstancedStruct properties

#Synced-CL 52402383

diff --git a/Engine/Source/Programs/Shared/EpicGames.UHT/Types/Properties/UhtStructProperty.cs b/Engine/Source/Programs/Shared/EpicGames.UHT/Types/Properties/UhtStructProperty.cs
--- a/Engine/Source/Programs/Shared/EpicGames.UHT/Types/Properties/UhtStructProperty.cs
+++ b/Engine/Source/Programs/Shared/EpicGames.UHT/Types/Properties/UhtStructProperty.cs
@@ -586,9 +586,12 @@ namespace EpicGames.UHT.Types
 		{
 			switch (textType)
 			{
-				default:
+				case UhtPropertyTextType.GetterSetterArg:
 					builder.Append(TemplateWrapperName).Append('<').Append(TemplateArgumentStruct.SourceName).Append('>');
 					break;
+				default:
+					builder.Append(TemplateWrapperName).Append('<').Append(TemplateArgumentStruct.FullyQualifiedSourceName).Append('>');
+					break;
 			}
 			return builder;
 		}
