Fall back to PATH lookup for p4 on non-Windows platforms

#rnx
#Synced-CL 52444173

diff --git a/Engine/Source/Programs/Shared/EpicGames.Perforce/PerforceChildProcess.cs b/Engine/Source/Programs/Shared/EpicGames.Perforce/PerforceChildProcess.cs
--- a/Engine/Source/Programs/Shared/EpicGames.Perforce/PerforceChildProcess.cs
+++ b/Engine/Source/Programs/Shared/EpicGames.Perforce/PerforceChildProcess.cs
@@ -107,6 +107,10 @@ namespace EpicGames.Perforce
 						break;
 					}
 				}
+				if (perforceFileName == "p4.exe")
+				{
+					perforceFileName = "p4"; // Fall back to PATH lookup
+				}
 			}
 			return perforceFileName;
 		}
