EpicGames.Core: Backup JSON log files too if log file backup is enabled.

Additionally, refactor code so that less behaviour is exclusive to the static Log methods and the logical flow is the same between the static Log methods and DefaultLogger. Narrow the API to just what is actually required to help prevent mistakes.

#Synced-CL 52502778

diff --git a/Engine/Source/Programs/AutomationTool/AutomationUtils/Automation.cs b/Engine/Source/Programs/AutomationTool/AutomationUtils/Automation.cs
--- a/Engine/Source/Programs/AutomationTool/AutomationUtils/Automation.cs
+++ b/Engine/Source/Programs/AutomationTool/AutomationUtils/Automation.cs
@@ -57,9 +57,8 @@ namespace AutomationTool
 		/// Main method.
 		/// </summary>
 		/// <param name="AutomationToolCommandLine">Command line</param>
-		/// <param name="StartupListener"></param>
 		/// <param name="ScriptModuleAssemblies"></param>
-		public static async Task<ExitCode> ProcessAsync(ParsedCommandLine AutomationToolCommandLine, StartupTraceListener StartupListener, HashSet<FileReference> ScriptModuleAssemblies)
+		public static async Task<ExitCode> ProcessAsync(ParsedCommandLine AutomationToolCommandLine, HashSet<FileReference> ScriptModuleAssemblies)
 		{
 			GlobalCommandLine.Initialize(AutomationToolCommandLine);
 
diff --git a/Engine/Source/Programs/AutomationTool/Program.cs b/Engine/Source/Programs/AutomationTool/Program.cs
--- a/Engine/Source/Programs/AutomationTool/Program.cs
+++ b/Engine/Source/Programs/AutomationTool/Program.cs
@@ -366,10 +366,7 @@ namespace AutomationToolDriver
 			}
 		}
 
-
 		static ParsedCommandLine AutomationToolCommandLine;
-		static StartupTraceListener StartupListener = new StartupTraceListener();
-
 
 		/// <summary>
 		/// Main entry point
@@ -380,7 +377,7 @@ namespace AutomationToolDriver
 			ILogger Logger = Log.Logger;
 
 			// Initialize the log system, buffering the output until we can create the log file
-			Log.AddTraceListener(StartupListener);
+			Log.AddStartupTraceListener();
 			Logger.LogInformation("Starting AutomationTool...");
 
 			// Populate AutomationToolCommandLine and CommandsToExecute
@@ -395,7 +392,7 @@ namespace AutomationToolDriver
 			}
 
 			// Wait for a debugger to be attached
-			if (AutomationToolCommandLine.IsSetGlobal("-WaitForDebugger"))	
+			if (AutomationToolCommandLine.IsSetGlobal("-WaitForDebugger"))
 			{
 				Console.WriteLine("Waiting for debugger to be attached...");
 				while (Debugger.IsAttached == false)
@@ -406,12 +403,12 @@ namespace AutomationToolDriver
 			}
 
 			Stopwatch Timer = Stopwatch.StartNew();
-			
+
 			// Ensure UTF8Output flag is respected, since we are initializing logging early in the program.
 			if (AutomationToolCommandLine.IsSetGlobal("-UTF8Output"))
-            {
-                Console.OutputEncoding = new System.Text.UTF8Encoding(false, false);
-            }
+			{
+				Console.OutputEncoding = new System.Text.UTF8Encoding(false, false);
+			}
 
 			// Parse the log level argument
 			if (AutomationToolCommandLine.IsSetGlobal("-Verbose"))
@@ -471,11 +468,11 @@ namespace AutomationToolDriver
 				ReturnCode = await ProcessSingleton.RunSingleInstanceAsync(MainProc, bWaitForUATMutex, Log.Logger);
 			}
 			catch (Exception Ex)
-            {
+			{
 				Logger.LogError(Ex, "Unhandled exception: {Message}", ExceptionUtils.FormatException(Ex));
-            }
-            finally
-            {
+			}
+			finally
+			{
 				// If there was an error before a log writer was attached, do so now
 				if (!Log.HasFileWriter() && ReturnCode != ExitCode.Success)
 				{
@@ -485,14 +482,14 @@ namespace AutomationToolDriver
 				}
 
 				// Write the exit code
-                Logger.LogInformation("AutomationTool executed for {Time}", Timer.Elapsed.ToString("h'h 'm'm 's's'"));
-                Logger.LogInformation("AutomationTool exiting with ExitCode={ExitCode} ({ExitReason})", (int)ReturnCode, ReturnCode);
+				Logger.LogInformation("AutomationTool executed for {Time}", Timer.Elapsed.ToString("h'h 'm'm 's's'"));
+				Logger.LogInformation("AutomationTool exiting with ExitCode={ExitCode} ({ExitReason})", (int)ReturnCode, ReturnCode);
 
-                // Can't use NoThrow here because the code logs exceptions. We're shutting down logging!
-                Trace.Close();
-            }
-            return (int)ReturnCode;
-        }
+				// Can't use NoThrow here because the code logs exceptions. We're shutting down logging!
+				Trace.Close();
+			}
+			return (int)ReturnCode;
+		}
 
 		static async Task<ExitCode> MainProc()
 		{
@@ -545,8 +542,7 @@ namespace AutomationToolDriver
 			Type AutomationTools_Automation = AutomationUtilsAssembly.GetType("AutomationTool.Automation");
 			MethodInfo Automation_Process = AutomationTools_Automation.GetMethod("ProcessAsync");
 			Logger.LogInformation("Total script module initialization time: {InitTime:0.00} s.", Timer.Elapsed.TotalSeconds);
-			return await (Task<ExitCode>) Automation_Process.Invoke(null,
-				new object[] {AutomationToolCommandLine, StartupListener, ScriptModuleAssemblyPaths});
+			return await (Task<ExitCode>) Automation_Process.Invoke(null, [AutomationToolCommandLine, ScriptModuleAssemblyPaths]);
 		}
 	}
 }
diff --git a/Engine/Source/Programs/Shared/EpicGames.Core/Logging/Log.cs b/Engine/Source/Programs/Shared/EpicGames.Core/Logging/Log.cs
--- a/Engine/Source/Programs/Shared/EpicGames.Core/Logging/Log.cs
+++ b/Engine/Source/Programs/Shared/EpicGames.Core/Logging/Log.cs
@@ -3,6 +3,7 @@
 using System;
 using System.Collections.Concurrent;
 using System.Collections.Generic;
+using System.Collections.Immutable;
 using System.Diagnostics;
 using System.IO;
 using System.Linq;
@@ -165,7 +166,7 @@ namespace EpicGames.Core
 		}
 
 		/// <summary>
-		/// When true, create a backup of any log file that would be overwritten by a new log
+		/// When true, create a backup of any text log file that would be overwritten by a new log
 		/// Log.txt will be backed up with its UTC creation time in the name e.g.
 		/// Log-backup-2021.10.29-19.53.17.txt
 		/// </summary>
@@ -222,7 +223,7 @@ namespace EpicGames.Core
 		}
 
 		/// <summary>
-		/// Backup an existing log file if it already exists at the outputpath
+		/// Backup an existing log file if it already exists at the output path
 		/// </summary>
 		/// <param name="outputFile">The file to back up</param>
 		public static void BackupLogFile(FileReference outputFile)
@@ -282,10 +283,7 @@ namespace EpicGames.Core
 		/// <param name="outputFile">The file to write to</param>
 		public static void AddFileWriter(string name, FileReference outputFile)
 		{
-			Logger.LogInformation("Log file: {OutputFile}", outputFile);
-
-			BackupLogFile(outputFile);
-			AddFileWriterWithoutBackup(name, outputFile);
+			s_defaultLogger.AddFileWriter(name, outputFile);
 		}
 
 		/// <summary>
@@ -301,21 +299,11 @@ namespace EpicGames.Core
 		}
 
 		/// <summary>
-		/// Adds a <see cref="TraceListener"/> to the collection in a safe manner.
+		/// Adds a <see cref="StartupTraceListener"/> to the collection in a safe manner.
 		/// </summary>
-		/// <param name="traceListener">The <see cref="TraceListener"/> to add.</param>
-		public static void AddTraceListener(TraceListener traceListener)
+		public static void AddStartupTraceListener()
 		{
-			s_defaultLogger.AddTraceListener(traceListener);
-		}
-
-		/// <summary>
-		/// Removes a <see cref="TraceListener"/> from the collection in a safe manner.
-		/// </summary>
-		/// <param name="traceListener">The <see cref="TraceListener"/> to remove.</param>
-		public static void RemoveTraceListener(TraceListener traceListener)
-		{
-			s_defaultLogger.RemoveTraceListener(traceListener);
+			s_defaultLogger.AddStartupTraceListener();
 		}
 
 		/// <summary>
@@ -1038,7 +1026,7 @@ namespace EpicGames.Core
 		/// <summary>
 		/// Output streams for structured log data
 		/// </summary>
-		private IReadOnlyList<FileStream> _jsonStreams = [];
+		private ImmutableArray<FileStream> _jsonStreams = [];
 
 		/// <summary>
 		/// Constructor
@@ -1050,7 +1038,7 @@ namespace EpicGames.Core
 			IncludeStartingTimestamp = true;
 			EventParser = new LogEventParser(this);
 
-			bool bJsonStdOut = Environment.CommandLine.Contains("JsonStdOut", StringComparison.InvariantCulture);
+			bool bJsonStdOut = Environment.CommandLine.Contains("JsonStdOut", StringComparison.Ordinal);
 			string? envVar = Environment.GetEnvironmentVariable("UE_LOG_JSON_TO_STDOUT");
 			if (bJsonStdOut || envVar != null && Int32.TryParse(envVar, out int value) && value != 0)
 			{
@@ -1117,7 +1105,7 @@ namespace EpicGames.Core
 		/// <param name="name">Listener name</param>
 		/// <param name="outputFile">The file to write to</param>
 		/// <returns>The created trace listener</returns>
-		public TextWriterTraceListener AddFileWriter(string name, FileReference outputFile)
+		private TextWriterTraceListener AddTextWriter(string name, FileReference outputFile)
 		{
 			try
 			{
@@ -1128,30 +1116,48 @@ namespace EpicGames.Core
 				{
 					Trace.Listeners.Add(logTraceListener);
 					WriteInitialTimestamp();
-
-					List<FileStream> newJsonStreams = [.. _jsonStreams];
-					newJsonStreams.Add(FileReference.Open(outputFile.ChangeExtension(".json"), FileMode.Create, FileAccess.Write, FileShare.Read | FileShare.Delete));
-					_jsonStreams = newJsonStreams;
 				}
 				return logTraceListener;
 			}
 			catch (Exception ex)
 			{
-				throw new Exception($"Error while creating log file \"{outputFile}\"", ex);
+				ex.Data["LogFile"] = outputFile;
+				throw;
+			}
+		}
+
+		/// <summary>
+		/// Adds a JSON stream to the collection
+		/// </summary>
+		/// <param name="outputFile">The file to write to</param>
+		private void AddJsonWriter(FileReference outputFile)
+		{
+			try
+			{
+				DirectoryReference.CreateDirectory(outputFile.Directory);
+				lock (_syncObject)
+				{
+					FileStream jsonStream = FileReference.Open(outputFile, FileMode.Create, FileAccess.Write, FileShare.Read | FileShare.Delete);
+					_jsonStreams = _jsonStreams.Add(jsonStream);
+				}
+			}
+			catch (Exception ex)
+			{
+				ex.Data["LogFile"] = outputFile;
+				throw;
 			}
 		}
 
 		/// <summary>
-		/// Adds a <see cref="TraceListener"/> to the collection in a safe manner.
+		/// Adds a <see cref="StartupTraceListener"/> to the collection in a safe manner.
 		/// </summary>
-		/// <param name="traceListener">The <see cref="TraceListener"/> to add.</param>
-		public void AddTraceListener(TraceListener traceListener)
+		public void AddStartupTraceListener()
 		{
 			lock (_syncObject)
 			{
-				if (!Trace.Listeners.Contains(traceListener))
+				if (!Trace.Listeners.OfType<StartupTraceListener>().Any())
 				{
-					Trace.Listeners.Add(traceListener);
+					Trace.Listeners.Add(new StartupTraceListener());
 					WriteInitialTimestamp();
 				}
 			}
@@ -1171,33 +1177,28 @@ namespace EpicGames.Core
 		}
 
 		/// <summary>
-		/// Removes a <see cref="TraceListener"/> from the collection in a safe manner.
+		/// Removes a <see cref="StartupTraceListener"/> from the collection in a safe manner.
 		/// </summary>
-		/// <param name="traceListener">The <see cref="TraceListener"/> to remove.</param>
-		public void RemoveTraceListener(TraceListener traceListener)
+		public StartupTraceListener? RemoveStartupTraceListener()
 		{
 			lock (_syncObject)
 			{
-				if (Trace.Listeners.Contains(traceListener))
+				StartupTraceListener? startupListener = Trace.Listeners.OfType<StartupTraceListener>().SingleOrDefault();
+				if (startupListener is not null)
 				{
-					Trace.Listeners.Remove(traceListener);
+					Trace.Listeners.Remove(startupListener);
 				}
+
+				return startupListener;
 			}
 		}
 
-		/// <summary>
-		/// Removes <see cref="StartupTraceListener"/> from the collection in a safe manner.
-		/// </summary>
-		public void RemoveStartupTraceListener()
+		public void AddFileWriter(string name, FileReference outputFile)
 		{
-			lock (_syncObject)
-			{
-				List<StartupTraceListener> startupListeners = [.. Trace.Listeners.OfType<StartupTraceListener>()];
-				if (startupListeners.Count != 0)
-				{
-					Trace.Listeners.Remove(startupListeners.First());
-				}
-			}
+			this.LogInformation("Log file: {OutputFile}", outputFile);
+			Core.Log.BackupLogFile(outputFile);
+			Core.Log.BackupLogFile(outputFile.ChangeExtension(".json"));
+			AddFileWriterWithoutBackup(name, outputFile);
 		}
 
 		/// <summary>
@@ -1209,18 +1210,14 @@ namespace EpicGames.Core
 		/// <param name="outputFile">The file to write to</param>
 		public void AddFileWriterWithoutBackup(string name, FileReference outputFile)
 		{
-			TextWriterTraceListener firstTextWriter = AddFileWriter(name, outputFile);
+			TextWriterTraceListener textWriter = AddTextWriter(name, outputFile);
+			// Note that we can't correctly copy the startup listener to the JSON writer as the output has already been formatted.
+			AddJsonWriter(outputFile.ChangeExtension(".json"));
 
 			lock (_syncObject)
 			{
-				// find the StartupTraceListener in the listeners that was added early on
-				List<StartupTraceListener> startupListeners = [.. Trace.Listeners.OfType<StartupTraceListener>()];
-				if (startupListeners.Count != 0)
-				{
-					StartupTraceListener startupListener = startupListeners.First();
-					startupListener.CopyTo(firstTextWriter);
-					Trace.Listeners.Remove(startupListeners.First());
-				}
+				StartupTraceListener? startupListener = RemoveStartupTraceListener();
+				startupListener?.CopyTo(textWriter);
 			}
 		}
 
diff --git a/Engine/Source/Programs/Shared/EpicGames.Core/Logging/StartupTraceListener.cs b/Engine/Source/Programs/Shared/EpicGames.Core/Logging/StartupTraceListener.cs
new file mode 100644
--- /dev/null
+++ b/Engine/Source/Programs/Shared/EpicGames.Core/Logging/StartupTraceListener.cs
@@ -0,0 +1,50 @@
+// Copyright Epic Games, Inc. All Rights Reserved.
+
+using System.Diagnostics;
+using System.Text;
+
+namespace EpicGames.Core
+{
+	/// <summary>
+	/// Captures all log output during startup until a log file writer has been created
+	/// </summary>
+	internal class StartupTraceListener : TraceListener
+	{
+		readonly StringBuilder _buffer = new();
+
+		/// <summary>
+		/// Copy the contents of the buffered output to another trace listener
+		/// </summary>
+		/// <param name="other">The trace listener to receive the buffered output</param>
+		public void CopyTo(TraceListener other)
+		{
+			foreach (string line in _buffer.ToString().Split("\n"))
+			{
+				other.WriteLine(line);
+			}
+		}
+
+		/// <summary>
+		/// Write a message to the buffer
+		/// </summary>
+		/// <param name="message">The message to write</param>
+		public override void Write(string? message)
+		{
+			if (NeedIndent)
+			{
+				WriteIndent();
+			}
+			_buffer.Append(message);
+		}
+
+		/// <summary>
+		/// Write a message to the buffer, followed by a newline
+		/// </summary>
+		/// <param name="message">The message to write</param>
+		public override void WriteLine(string? message)
+		{
+			Write(message);
+			_buffer.Append('\n');
+		}
+	}
+}
diff --git a/Engine/Source/Programs/Shared/EpicGames.Core/StartupTraceListener.cs b/Engine/Source/Programs/Shared/EpicGames.Core/StartupTraceListener.cs
deleted file mode 100644
diff --git a/Engine/Source/Programs/UnrealBuildTool/UnrealBuildTool.cs b/Engine/Source/Programs/UnrealBuildTool/UnrealBuildTool.cs
--- a/Engine/Source/Programs/UnrealBuildTool/UnrealBuildTool.cs
+++ b/Engine/Source/Programs/UnrealBuildTool/UnrealBuildTool.cs
@@ -277,8 +277,7 @@ namespace UnrealBuildTool
 				}
 
 				// Always start capturing logs as early as possible to later copy to a log file if the ToolMode desires it (we have to start capturing before we get the ToolModeOptions below)
-				StartupTraceListener StartupTrace = new StartupTraceListener();
-				Log.AddTraceListener(StartupTrace);
+				Log.AddStartupTraceListener();
 
 				// Enable threading on logger since it can stall quite a lot when logging happens from multiple threads
 				Log.EnableThreaded();
@@ -323,7 +322,7 @@ namespace UnrealBuildTool
 				// if we don't care about the trace listener, toss it now
 				if ((ModeOptions & ToolModeOptions.UseStartupTraceListener) == 0)
 				{
-					Log.RemoveTraceListener(StartupTrace);
+					Log.RemoveStartupTraceListener();
 				}
 
 				// Read the XML configuration files
