formatting fixup

diff --git a/Engine/Source/Programs/Shared/EpicGames.Core/CondensedGraphBuilder.cs b/Engine/Source/Programs/Shared/EpicGames.Core/CondensedGraphBuilder.cs
--- a/Engine/Source/Programs/Shared/EpicGames.Core/CondensedGraphBuilder.cs
+++ b/Engine/Source/Programs/Shared/EpicGames.Core/CondensedGraphBuilder.cs
@@ -2,7 +2,6 @@
 
 using System;
 using System.Collections.Generic;
-using Microsoft.Extensions.Logging;
 
 namespace EpicGames.Core
 {
@@ -23,54 +22,51 @@ namespace EpicGames.Core
 	/// <typeparam name="T">Type of node objects</typeparam>
 	public class CondensedGraphBuilder<T> where T : notnull
 	{
-		/// <summary>
-		/// UBT's logger object. Used for diagnostic output.
-		/// </summary>
-		public ILogger? Logger = null;
-
 		/// <summary>
 		/// Functor returning a node's name. Used for diagnostic logging.
 		/// </summary>
-		public Func<T, string>? NodeToString = null;
+		private Func<T, string>? _nodeToString = null;
 
 		/// <summary>
 		/// List of all graph nodes in insertion order.
 		/// </summary>
-		private readonly List<GraphNode> Nodes = new List<GraphNode>();
+		private readonly List<GraphNode> _nodes = [];
 
 		/// <summary>
 		/// Maps each user-provided node to its internal GraphNode.
 		/// </summary>
-		private readonly Dictionary<T, GraphNode> NodeMap = new Dictionary<T, GraphNode>();
+		private readonly Dictionary<T, GraphNode> _nodeMap = [];
 
 		/// <summary>
 		/// Output: list of strongly connected components. Each SCC is an array of original nodes.
 		/// Ordered root-to-leaf: condensed edges always point from lower to higher SCC index.
 		/// Always populated after Build().
 		/// </summary>
-		private readonly List<T[]> ResultSCCs = new List<T[]>();
+		private readonly List<T[]> _resultSCCs = [];
 
 		/// <summary>
 		/// Output: maps each original node to the index of its SCC in ResultSCCs.
 		/// Always populated after Build().
 		/// </summary>
-		private readonly Dictionary<T, int> ResultNodeToSCCIndex = new Dictionary<T, int>();
+		private readonly Dictionary<T, int> _resultNodeToSCCIndex = [];
 
 		/// <summary>
 		/// Output: adjacency list of the condensed graph. ResultCondensedEdges[i] holds the
 		/// SCC indices that SCC i has outgoing edges to. Only populated when Build() returns true.
 		/// </summary>
-		private readonly List<int[]> ResultCondensedEdges = new List<int[]>();
+		private readonly List<int[]> _resultCondensedEdges = [];
+
+		public Func<T, string>? NodeToString { get => _nodeToString; set => _nodeToString = value; }
 
 		/// <summary>
 		/// Initializes the builder and constructs the internal graph from the provided edges.
 		/// All nodes referenced by edges (both source and destination) are included.
 		/// Self-edges are discarded.
 		/// </summary>
-		/// <param name="Edges">Directed edges as (source, destination) tuples</param>
-		public CondensedGraphBuilder(List<Tuple<T, T>> Edges)
+		/// <param name="edges">Directed edges as (source, destination) tuples</param>
+		public CondensedGraphBuilder(List<Tuple<T, T>> edges)
 		{
-			CreateGraph(Edges);
+			CreateGraph(edges);
 		}
 
 		/// <summary>
@@ -82,12 +78,12 @@ namespace EpicGames.Core
 		/// </returns>
 		public bool Build()
 		{
-			ResultSCCs.Clear();
-			ResultNodeToSCCIndex.Clear();
-			ResultCondensedEdges.Clear();
+			_resultSCCs.Clear();
+			_resultNodeToSCCIndex.Clear();
+			_resultCondensedEdges.Clear();
 
-			int NumVertices = Nodes.Count;
-			if (NumVertices == 0)
+			int numVertices = _nodes.Count;
+			if (numVertices == 0)
 			{
 				return false;
 			}
@@ -95,60 +91,60 @@ namespace EpicGames.Core
 			// Phase 1 — Iterative DFS on the original graph.
 			// Records each vertex in post-order (after all outgoing edges are exhausted).
 			// Any back-edge to an InProgress vertex (other than a self-edge) indicates a cycle.
-			List<int> PostOrder = new List<int>(NumVertices);
-			bool HasCycle = false;
+			List<int> postOrder = new(numVertices);
+			bool hasCycle = false;
 
 			{
-				VisitStatus[] Status = new VisitStatus[NumVertices]; // default: NotVisited
-				List<(int Vertex, int NextEdge)> Stack = new List<(int, int)>(NumVertices);
+				VisitStatus[] status = new VisitStatus[numVertices]; // default: NotVisited
+				List<(int Vertex, int NextEdge)> stack = new(numVertices);
 
-				for (int Root = 0; Root < NumVertices; ++Root)
+				for (int root = 0; root < numVertices; ++root)
 				{
-					if (Status[Root] != VisitStatus.NotVisited)
+					if (status[root] != VisitStatus.NotVisited)
 					{
 						continue;
 					}
-					Status[Root] = VisitStatus.InProgress;
-					Stack.Add((Root, 0));
+					status[root] = VisitStatus.InProgress;
+					stack.Add((root, 0));
 
-					while (Stack.Count > 0)
+					while (stack.Count > 0)
 					{
-						(int Vertex, int NextEdge) = Stack[Stack.Count - 1];
-						List<GraphNode> Edges = Nodes[Vertex].Links;
-						bool Pushed = false;
+						(int vertex, int nextEdge) = stack[stack.Count - 1];
+						List<GraphNode> edges = _nodes[vertex]._links;
+						bool pushed = false;
 
-						while (NextEdge < Edges.Count)
+						while (nextEdge < edges.Count)
 						{
-							int EdgeVertex = Edges[NextEdge].Index;
-							++NextEdge;
+							int edgeVertex = edges[nextEdge]._index;
+							++nextEdge;
 
-							if (Status[EdgeVertex] == VisitStatus.Visited)
+							if (status[edgeVertex] == VisitStatus.Visited)
 							{
 								// Already fully processed; nothing to do.
 							}
-							else if (Status[EdgeVertex] == VisitStatus.NotVisited)
+							else if (status[edgeVertex] == VisitStatus.NotVisited)
 							{
-								Status[EdgeVertex] = VisitStatus.InProgress;
-								Stack[Stack.Count - 1] = (Vertex, NextEdge);
-								Stack.Add((EdgeVertex, 0));
-								Pushed = true;
+								status[edgeVertex] = VisitStatus.InProgress;
+								stack[stack.Count - 1] = (vertex, nextEdge);
+								stack.Add((edgeVertex, 0));
+								pushed = true;
 								break;
 							}
 							else // InProgress — back edge detected
 							{
 								// Self-edges are not counted as cycles; they are discarded during CreateGraph.
-								if (EdgeVertex != Vertex)
+								if (edgeVertex != vertex)
 								{
-									HasCycle = true;
+									hasCycle = true;
 								}
 							}
 						}
 
-						if (!Pushed)
+						if (!pushed)
 						{
-							PostOrder.Add(Vertex);
-							Status[Vertex] = VisitStatus.Visited;
-							Stack.RemoveAt(Stack.Count - 1);
+							postOrder.Add(vertex);
+							status[vertex] = VisitStatus.Visited;
+							stack.RemoveAt(stack.Count - 1);
 						}
 					}
 				}
@@ -157,16 +153,16 @@ namespace EpicGames.Core
 			// No cycles found: input is already a DAG.
 			// Populate SCCs as an identity mapping in reverse post-order (root-to-leaf).
 			// ResultCondensedEdges is left empty per the C++ TryConstructcondensedGraph contract.
-			if (!HasCycle)
+			if (!hasCycle)
 			{
-				ResultSCCs.Capacity = NumVertices;
-				ResultNodeToSCCIndex.EnsureCapacity(NumVertices);
+				_resultSCCs.Capacity = numVertices;
+				_resultNodeToSCCIndex.EnsureCapacity(numVertices);
 
-				for (int SCCIndex = 0; SCCIndex < NumVertices; ++SCCIndex)
+				for (int sCCIndex = 0; sCCIndex < numVertices; ++sCCIndex)
 				{
-					T NodeData = Nodes[PostOrder[NumVertices - 1 - SCCIndex]].Data;
-					ResultSCCs.Add(new T[] { NodeData });
-					ResultNodeToSCCIndex[NodeData] = SCCIndex;
+					T nodeData = _nodes[postOrder[numVertices - 1 - sCCIndex]]._data;
+					_resultSCCs.Add([nodeData]);
+					_resultNodeToSCCIndex[nodeData] = sCCIndex;
 				}
 
 				return false;
@@ -177,59 +173,59 @@ namespace EpicGames.Core
 			// (i.e., all vertices that can reach the seed in the original graph) join the same SCC.
 			// Processing in reverse post-order ensures SCCs are discovered root-to-leaf.
 			{
-				bool[] Assigned = new bool[NumVertices];
-				List<(int Vertex, int NextEdge)> Stack = new List<(int, int)>(NumVertices);
+				bool[] assigned = new bool[numVertices];
+				List<(int Vertex, int NextEdge)> stack = new(numVertices);
 
-				for (int i = NumVertices - 1; i >= 0; --i)
+				for (int i = numVertices - 1; i >= 0; --i)
 				{
-					int Root = PostOrder[i];
-					if (Assigned[Root])
+					int root = postOrder[i];
+					if (assigned[root])
 					{
 						continue;
 					}
 
-					List<int> SCCVertices = new List<int>();
-					Assigned[Root] = true;
-					SCCVertices.Add(Root);
-					Stack.Add((Root, 0));
+					List<int> sCCVertices = [];
+					assigned[root] = true;
+					sCCVertices.Add(root);
+					stack.Add((root, 0));
 
-					while (Stack.Count > 0)
+					while (stack.Count > 0)
 					{
-						(int Vertex, int NextEdge) = Stack[Stack.Count - 1];
-						List<GraphNode> TransposeEdges = Nodes[Vertex].TransposeLinks;
-						bool Pushed = false;
+						(int vertex, int nextEdge) = stack[stack.Count - 1];
+						List<GraphNode> transposeEdges = _nodes[vertex]._transposeLinks;
+						bool pushed = false;
 
-						while (NextEdge < TransposeEdges.Count)
+						while (nextEdge < transposeEdges.Count)
 						{
-							int EdgeVertex = TransposeEdges[NextEdge].Index;
-							++NextEdge;
+							int edgeVertex = transposeEdges[nextEdge]._index;
+							++nextEdge;
 
-							if (!Assigned[EdgeVertex])
+							if (!assigned[edgeVertex])
 							{
-								Assigned[EdgeVertex] = true;
-								SCCVertices.Add(EdgeVertex);
-								Stack[Stack.Count - 1] = (Vertex, NextEdge);
-								Stack.Add((EdgeVertex, 0));
-								Pushed = true;
+								assigned[edgeVertex] = true;
+								sCCVertices.Add(edgeVertex);
+								stack[stack.Count - 1] = (vertex, nextEdge);
+								stack.Add((edgeVertex, 0));
+								pushed = true;
 								break;
 							}
 						}
 
-						if (!Pushed)
+						if (!pushed)
 						{
-							Stack.RemoveAt(Stack.Count - 1);
+							stack.RemoveAt(stack.Count - 1);
 						}
 					}
 
-					int SCCIndex = ResultSCCs.Count;
-					T[] SCCData = new T[SCCVertices.Count];
-					for (int j = 0; j < SCCVertices.Count; ++j)
+					int sCCIndex = _resultSCCs.Count;
+					T[] sCCData = new T[sCCVertices.Count];
+					for (int j = 0; j < sCCVertices.Count; ++j)
 					{
-						T NodeData = Nodes[SCCVertices[j]].Data;
-						SCCData[j] = NodeData;
-						ResultNodeToSCCIndex[NodeData] = SCCIndex;
+						T nodeData = _nodes[sCCVertices[j]]._data;
+						sCCData[j] = nodeData;
+						_resultNodeToSCCIndex[nodeData] = sCCIndex;
 					}
-					ResultSCCs.Add(SCCData);
+					_resultSCCs.Add(sCCData);
 				}
 			}
 
@@ -238,34 +234,34 @@ namespace EpicGames.Core
 			// edge list, mapping each target node to its SCC index and deduplicating with a
 			// boolean scratch array that is spot-cleared per iteration.
 			{
-				int NumSCCs = ResultSCCs.Count;
-				bool[] EdgeAdded = new bool[NumSCCs];
+				int numSCCs = _resultSCCs.Count;
+				bool[] edgeAdded = new bool[numSCCs];
 
-				for (int SCCIndex = 0; SCCIndex < NumSCCs; ++SCCIndex)
+				for (int sCCIndex = 0; sCCIndex < numSCCs; ++sCCIndex)
 				{
-					List<int> SCCEdges = new List<int>();
+					List<int> sCCEdges = [];
 
-					foreach (T InNode in ResultSCCs[SCCIndex])
+					foreach (T inNode in _resultSCCs[sCCIndex])
 					{
-						foreach (GraphNode Neighbor in NodeMap[InNode].Links)
+						foreach (GraphNode neighbor in _nodeMap[inNode]._links)
 						{
-							int NeighborSCC = ResultNodeToSCCIndex[Neighbor.Data];
+							int neighborSCC = _resultNodeToSCCIndex[neighbor._data];
 							// NeighborSCC != SCCIndex also removes condensed self-edges
-							if (NeighborSCC != SCCIndex && !EdgeAdded[NeighborSCC])
+							if (neighborSCC != sCCIndex && !edgeAdded[neighborSCC])
 							{
-								EdgeAdded[NeighborSCC] = true;
-								SCCEdges.Add(NeighborSCC);
+								edgeAdded[neighborSCC] = true;
+								sCCEdges.Add(neighborSCC);
 							}
 						}
 					}
 
 					// Clear only the entries we touched — faster than resetting the whole array.
-					foreach (int EdgeSCC in SCCEdges)
+					foreach (int edgeSCC in sCCEdges)
 					{
-						EdgeAdded[EdgeSCC] = false;
+						edgeAdded[edgeSCC] = false;
 					}
 
-					ResultCondensedEdges.Add(SCCEdges.ToArray());
+					_resultCondensedEdges.Add([.. sCCEdges]);
 				}
 			}
 
@@ -277,55 +273,55 @@ namespace EpicGames.Core
 		/// Each SCC is an array of original nodes. Ordered root-to-leaf in the condensed graph.
 		/// Always populated, even when Build() returns false.
 		/// </summary>
-		public List<T[]> GetSCCs() => ResultSCCs;
+		public List<T[]> GetSCCs() => _resultSCCs;
 
 		/// <summary>
 		/// Returns a mapping from each original node to the index of its SCC in GetSCCs()
 		/// (valid after calling Build()). Always populated, even when Build() returns false.
 		/// </summary>
-		public IReadOnlyDictionary<T, int> GetNodeToSCCIndex() => ResultNodeToSCCIndex;
+		public IReadOnlyDictionary<T, int> GetNodeToSCCIndex() => _resultNodeToSCCIndex;
 
 		/// <summary>
 		/// Returns the adjacency list of the condensed graph (valid after calling Build()).
 		/// GetCondensedEdges()[i] lists the SCC indices that SCC i has outgoing edges to.
 		/// Only populated when Build() returns true; empty when Build() returns false.
 		/// </summary>
-		public List<int[]> GetCondensedEdges() => ResultCondensedEdges;
+		public List<int[]> GetCondensedEdges() => _resultCondensedEdges;
 
 		/// <summary>
 		/// Constructs the internal graph from the provided edge list.
 		/// Nodes are added in the order they first appear. Self-edges are discarded.
 		/// </summary>
-		private void CreateGraph(List<Tuple<T, T>> Edges)
+		private void CreateGraph(List<Tuple<T, T>> edges)
 		{
 			// First pass: register all nodes in insertion order.
-			foreach (Tuple<T, T> Edge in Edges)
+			foreach (Tuple<T, T> edge in edges)
 			{
-				if (!NodeMap.ContainsKey(Edge.Item1))
+				if (!_nodeMap.ContainsKey(edge.Item1))
 				{
-					GraphNode Node = new GraphNode(Edge.Item1, Nodes.Count);
-					Nodes.Add(Node);
-					NodeMap[Edge.Item1] = Node;
+					GraphNode node = new(edge.Item1, _nodes.Count);
+					_nodes.Add(node);
+					_nodeMap[edge.Item1] = node;
 				}
-				if (!NodeMap.ContainsKey(Edge.Item2))
+				if (!_nodeMap.ContainsKey(edge.Item2))
 				{
-					GraphNode Node = new GraphNode(Edge.Item2, Nodes.Count);
-					Nodes.Add(Node);
-					NodeMap[Edge.Item2] = Node;
+					GraphNode node = new(edge.Item2, _nodes.Count);
+					_nodes.Add(node);
+					_nodeMap[edge.Item2] = node;
 				}
 			}
 
 			// Second pass: add edges, discarding self-edges.
-			foreach (Tuple<T, T> Edge in Edges)
+			foreach (Tuple<T, T> edge in edges)
 			{
-				GraphNode Src = NodeMap[Edge.Item1];
-				GraphNode Dst = NodeMap[Edge.Item2];
-				if (Src == Dst)
+				GraphNode src = _nodeMap[edge.Item1];
+				GraphNode dst = _nodeMap[edge.Item2];
+				if (src == dst)
 				{
 					continue;
 				}
-				Src.Links.Add(Dst);
-				Dst.TransposeLinks.Add(Src);
+				src._links.Add(dst);
+				dst._transposeLinks.Add(src);
 			}
 		}
 
@@ -339,21 +335,21 @@ namespace EpicGames.Core
 		private class GraphNode
 		{
 			/// <summary>User-provided node object.</summary>
-			public T Data;
+			public T _data;
 
 			/// <summary>Position in the Nodes list; used as an array index during Build().</summary>
-			public int Index;
+			public int _index;
 
 			/// <summary>Outgoing edges in the original graph.</summary>
-			public List<GraphNode> Links = new List<GraphNode>();
+			public List<GraphNode> _links = [];
 
 			/// <summary>Incoming edges in the original graph (outgoing edges in the transpose graph).</summary>
-			public List<GraphNode> TransposeLinks = new List<GraphNode>();
+			public List<GraphNode> _transposeLinks = [];
 
-			public GraphNode(T Data, int Index)
+			public GraphNode(T data, int index)
 			{
-				this.Data  = Data;
-				this.Index = Index;
+				_data  = data;
+				_index = index;
 			}
 		}
 	}
