Additional format fix for missing XML comment

diff --git a/Engine/Source/Programs/Shared/EpicGames.Core/CondensedGraphBuilder.cs b/Engine/Source/Programs/Shared/EpicGames.Core/CondensedGraphBuilder.cs
--- a/Engine/Source/Programs/Shared/EpicGames.Core/CondensedGraphBuilder.cs
+++ b/Engine/Source/Programs/Shared/EpicGames.Core/CondensedGraphBuilder.cs
@@ -56,6 +56,9 @@ namespace EpicGames.Core
 		/// </summary>
 		private readonly List<int[]> _resultCondensedEdges = [];
 
+		/// <summary>
+		/// Function to go from a node to the string output. Can be used for logging.
+		/// </summary>
 		public Func<T, string>? NodeToString { get => _nodeToString; set => _nodeToString = value; }
 
 		/// <summary>
