Scrub use of --ignore-missing-attachments when importing oplogs/snapshots.

#Synced-CL 52440464

diff --git a/Engine/Source/Programs/AutomationTool/BuildGraph/Tasks/ZenImportSnapshotTask.cs b/Engine/Source/Programs/AutomationTool/BuildGraph/Tasks/ZenImportSnapshotTask.cs
--- a/Engine/Source/Programs/AutomationTool/BuildGraph/Tasks/ZenImportSnapshotTask.cs
+++ b/Engine/Source/Programs/AutomationTool/BuildGraph/Tasks/ZenImportSnapshotTask.cs
@@ -459,7 +459,7 @@ namespace AutomationTool.Tasks
 				oplogCreateCommandLine.AppendFormat(" --hosturl {0}", host);
 
 				StringBuilder oplogImportCommandLine = new StringBuilder();
-				oplogImportCommandLine.AppendFormat("oplog-import {0} {1} {2} --ignore-missing-attachments --clean", projectId, oplogName, projectStoreFile.FullName);
+				oplogImportCommandLine.AppendFormat("oplog-import {0} {1} {2} --clean", projectId, oplogName, projectStoreFile.FullName);
 				oplogImportCommandLine.AppendFormat(" --hosturl {0}", host);
 
 				if (_parameters.ForceImport)
diff --git a/Engine/Source/Programs/Shared/EpicGames.ProjectStore/CloudStorage/CloudStorage.cs b/Engine/Source/Programs/Shared/EpicGames.ProjectStore/CloudStorage/CloudStorage.cs
--- a/Engine/Source/Programs/Shared/EpicGames.ProjectStore/CloudStorage/CloudStorage.cs
+++ b/Engine/Source/Programs/Shared/EpicGames.ProjectStore/CloudStorage/CloudStorage.cs
@@ -535,7 +535,7 @@ namespace EpicGames.ProjectStore
 			oplogCreateCommandLine.AppendFormat("oplog-create --force-update {0} {1}", projectId, oplogName);
 
 			StringBuilder oplogImportCommandLine = new StringBuilder();
-			oplogImportCommandLine.AppendFormat("oplog-import {0} {1} {2} --ignore-missing-attachments --clean", projectId, oplogName, projectStoreFile.FullName);
+			oplogImportCommandLine.AppendFormat("oplog-import {0} {1} {2} --clean", projectId, oplogName, projectStoreFile.FullName);
 
 			if (forceImport)
 			{
diff --git a/Engine/Source/Programs/Shared/EpicGames.ProjectStore/ProjectStoreTypes.cs b/Engine/Source/Programs/Shared/EpicGames.ProjectStore/ProjectStoreTypes.cs
--- a/Engine/Source/Programs/Shared/EpicGames.ProjectStore/ProjectStoreTypes.cs
+++ b/Engine/Source/Programs/Shared/EpicGames.ProjectStore/ProjectStoreTypes.cs
@@ -576,7 +576,7 @@ namespace EpicGames.ProjectStore
 				oplogCreateCommandLine.AppendFormat("oplog-create --force-update {0} {1}", projectId, oplogName);
 
 				StringBuilder oplogImportCommandLine = new();
-				oplogImportCommandLine.AppendFormat("oplog-import {0} {1} {2} --ignore-missing-attachments --clean", projectId, oplogName, projectStoreFile.FullName);
+				oplogImportCommandLine.AppendFormat("oplog-import {0} {1} {2} --clean", projectId, oplogName, projectStoreFile.FullName);
 
 				if (forceImport)
 				{
@@ -674,7 +674,7 @@ namespace EpicGames.ProjectStore
 			oplogCreateCommandLine.AppendFormat("oplog-create --force-update {0} {1}", projectId, oplogName);
 
 			StringBuilder oplogImportCommandLine = new();
-			oplogImportCommandLine.AppendFormat("oplog-import {0} {1} {2} --ignore-missing-attachments --clean", projectId, oplogName, projectStoreFile.FullName);
+			oplogImportCommandLine.AppendFormat("oplog-import {0} {1} {2} --clean", projectId, oplogName, projectStoreFile.FullName);
 
 			if (forceImport)
 			{
