Consolidate GlobalSingleInstanceMutex and SingleInstanceMutex into EpicGames.Core, adding more tests for the additional APIs. Fix up usages of GlobalSingleInstanceMutex.

#Synced-CL 52511844

diff --git a/Engine/Source/Programs/Shared/EpicGames.Build.Tests/GlobalSingleInstanceMutexTests.cs b/Engine/Source/Programs/Shared/EpicGames.Build.Tests/GlobalSingleInstanceMutexTests.cs
deleted file mode 100644
diff --git a/Engine/Source/Programs/Shared/EpicGames.Build/System/GlobalSingleInstanceMutex.cs b/Engine/Source/Programs/Shared/EpicGames.Build/System/GlobalSingleInstanceMutex.cs
--- a/Engine/Source/Programs/Shared/EpicGames.Build/System/GlobalSingleInstanceMutex.cs
+++ b/Engine/Source/Programs/Shared/EpicGames.Build/System/GlobalSingleInstanceMutex.cs
@@ -1,6 +1,8 @@
 // Copyright Epic Games, Inc. All Rights Reserved.
 
 using System;
+using System.Diagnostics.Metrics;
+using System.Runtime.InteropServices;
 using System.Threading;
 using System.Threading.Tasks;
 using EpicGames.Core;
@@ -8,40 +10,10 @@ using EpicGames.Core;
 namespace UnrealBuildBase
 {
 	/// <summary>
-	/// System-wide mutex allowing only one instance of the program to run at a time
+	/// Helper methods for single instance mutexes.
 	/// </summary>
-	public class GlobalSingleInstanceMutex : IDisposable
+	public static class GlobalSingleInstanceMutex
 	{
-		/// <summary>
-		/// The global mutex instance
-		/// </summary>
-		Mutex? _globalMutex;
-
-		/// <summary>
-		/// A single-threaded scheduler to run the mutex acquire / release on, to make this class safe for use within an async context.
-		/// </summary>
-		SingleThreadTaskScheduler _scheduler = new();
-
-		/// <summary>
-		/// Constructor. Attempts to acquire the global mutex
-		/// </summary>
-		/// <param name="mutexName">Name of the mutex to acquire</param>
-		/// <param name="bWaitMutex">Allow waiting for the mutex to be acquired</param>
-		public GlobalSingleInstanceMutex(string mutexName, bool bWaitMutex)
-		{
-			try
-			{
-				RunActionOnScheduler(() => GetMutex(mutexName, bWaitMutex));
-			}
-			catch
-			{
-				// If we can't construct the object fully then we should dispose our resources otherwise they would leak.
-				_globalMutex?.Dispose();
-				_scheduler.Dispose();
-				throw;
-			}
-		}
-
 		/// <summary>
 		/// Gets the name of a mutex unique for the given path
 		/// </summary>
@@ -61,58 +33,18 @@ namespace UnrealBuildBase
 		public static string GetUniqueMutexForPath(string name, FileSystemReference? uniquePath) => GetUniqueMutexForPath(name, uniquePath?.FullName ?? String.Empty);
 
 		/// <summary>
-		/// Release the mutex and dispose of the object
+		/// Tries to acquire a global mutex with the given name if it's available, and throws if it is not.
 		/// </summary>
-		public void Dispose()
+		/// <param name="name">Name of the mutex to acquire.</param>
+		/// <returns>The acquired mutex.</returns>
+		/// <exception cref="BuildLogEventException">Thrown when the mutex cannot be acquired.</exception>
+		public static IDisposable AcquireNowOrThrow(string name)
 		{
-			Dispose(true);
-			GC.SuppressFinalize(this);
-		}
-
-		/// <summary>
-		/// Overridable dispose method
-		/// </summary>
-		/// <param name="disposing">Whether the object should be disposed</param>
-		protected virtual void Dispose(bool disposing)
-		{
-			if (disposing)
-			{
-				RunActionOnScheduler(() => _globalMutex?.ReleaseMutex());
-				_globalMutex?.Dispose();
-				_globalMutex = null;
-				_scheduler.Dispose();
-			}
-		}
-
-		private void RunActionOnScheduler(Action action)
-		{
-			Task getMutexTask = Task.Factory.StartNew(action, CancellationToken.None, TaskCreationOptions.None, _scheduler);
-			getMutexTask.GetAwaiter().GetResult();
-		}
-
-		private void GetMutex(string mutexName, bool bWaitMutex)
-		{
-			// Try to create the mutex, with it initially locked
-			_globalMutex = new Mutex(true, mutexName, out bool bCreatedMutex);
-
-			// If we didn't create the mutex, we can wait for it or fail immediately
-			if (!bCreatedMutex)
+			if (!SingleInstanceMutex.TryAcquire(name, out IDisposable? mutex))
 			{
-				if (bWaitMutex)
-				{
-					try
-					{
-						_globalMutex.WaitOne();
-					}
-					catch (AbandonedMutexException)
-					{
-					}
-				}
-				else
-				{
-					throw new BuildLogEventException(new CompilationResultException(CompilationResult.ConflictingInstance), "A conflicting instance of {Mutex} is already running.", mutexName);
-				}
+				throw new BuildLogEventException(new CompilationResultException(CompilationResult.ConflictingInstance), "A conflicting instance of {Mutex} is already running.", name);
 			}
+			return mutex;
 		}
 	}
 }
diff --git a/Engine/Source/Programs/Shared/EpicGames.Core.Tests/SingleInstanceMutexTests.cs b/Engine/Source/Programs/Shared/EpicGames.Core.Tests/SingleInstanceMutexTests.cs
new file mode 100644
--- /dev/null
+++ b/Engine/Source/Programs/Shared/EpicGames.Core.Tests/SingleInstanceMutexTests.cs
@@ -0,0 +1,147 @@
+// Copyright Epic Games, Inc. All Rights Reserved.
+
+using System;
+using System.Threading;
+using System.Threading.Tasks;
+using Microsoft.VisualStudio.TestTools.UnitTesting;
+
+namespace EpicGames.Core.Tests;
+
+[TestClass]
+public class SingleInstanceMutexTests
+{
+#pragma warning disable CA1849 // Intentionally testing the synchronous code paths.
+	[TestMethod]
+	public async Task AcquiringGlobalMutexSynchronouslyBlocksCorrectlyAsync()
+	{
+		const string MutexName = "EpicGames.Build.Tests.UnitTestMutex0";
+		IDisposable? mutex1 = SingleInstanceMutex.Acquire(MutexName);
+
+		try
+		{
+			Task mutex2Task = Task.Run(() =>
+			{
+				using IDisposable mutex2 = SingleInstanceMutex.Acquire(MutexName);
+			}, TestContext.CancellationToken);
+
+			// We should get the delay task back because the mutex 2 task should be blocked acquiring the mutex.
+			Task delayTask = Task.Delay(200, TestContext.CancellationToken);
+			Task finishedTask = await Task.WhenAny(mutex2Task, delayTask);
+
+			Assert.AreSame(delayTask, finishedTask);
+
+			// Dispose of the first mutex so the second one can acquire it.
+			mutex1.Dispose();
+			mutex1 = null;
+
+			delayTask = Task.Delay(200, TestContext.CancellationToken);
+			finishedTask = await Task.WhenAny(mutex2Task, delayTask);
+			await finishedTask;
+
+			// We should get the mutex 2 task as it should have been able to acquire the mutex and free it again.
+			Assert.AreSame(mutex2Task, finishedTask);
+		}
+		finally
+		{
+			mutex1?.Dispose();
+		}
+	}
+#pragma warning restore CA1849
+
+	[TestMethod]
+	public async Task AcquiringGlobalMutexAsynchronouslyBlocksCorrectlyAsync()
+	{
+		const string MutexName = "EpicGames.Build.Tests.UnitTestMutex1";
+		IDisposable? mutex1 = await SingleInstanceMutex.AcquireAsync(MutexName, TestContext.CancellationToken);
+
+		try
+		{
+			Task<IDisposable> mutex2Task = SingleInstanceMutex.AcquireAsync(MutexName, TestContext.CancellationToken);
+
+			// We should get the delay task back because the mutex 2 task should be blocked acquiring the mutex.
+			Task delayTask = Task.Delay(200, TestContext.CancellationToken);
+			Task finishedTask = await Task.WhenAny(mutex2Task, delayTask);
+
+			Assert.AreSame(delayTask, finishedTask);
+
+			// Dispose of the first mutex so the second one can acquire it.
+			mutex1.Dispose();
+			mutex1 = null;
+
+			delayTask = Task.Delay(200, TestContext.CancellationToken);
+			finishedTask = await Task.WhenAny(mutex2Task, delayTask);
+			await finishedTask;
+
+			// We should get the mutex 2 task as it should have been able to acquire the mutex and free it again.
+			Assert.AreSame(mutex2Task, finishedTask);
+
+			(await mutex2Task).Dispose();
+		}
+		finally
+		{
+			mutex1?.Dispose();
+		}
+	}
+
+	[TestMethod]
+	public async Task AcquiringGlobalMutexAsynchronouslyCancelsCorrectlyAsync()
+	{
+		const string MutexName = "EpicGames.Build.Tests.UnitTestMutex2";
+		IDisposable? mutex1 = await SingleInstanceMutex.AcquireAsync(MutexName, TestContext.CancellationToken);
+
+		try
+		{
+			using (CancellationTokenSource cancellationTokenSource = new())
+			{
+				Task mutex2Task = SingleInstanceMutex.AcquireAsync(MutexName, cancellationTokenSource.Token);
+
+				// We should get the delay task back because the mutex 2 task should be blocked acquiring the mutex.
+				Task delayTask = Task.Delay(200, TestContext.CancellationToken);
+				Task finishedTask = await Task.WhenAny(mutex2Task, delayTask);
+
+				Assert.AreSame(delayTask, finishedTask);
+
+				await cancellationTokenSource.CancelAsync();
+
+				await Assert.ThrowsAsync<OperationCanceledException>(async () => await mutex2Task);
+
+				// Dispose of the first mutex.
+				mutex1.Dispose();
+				mutex1 = null;
+
+				// Prove that the task was properly cancelled and we can acquire the mutex now.
+				mutex1 = await SingleInstanceMutex.AcquireAsync(MutexName, TestContext.CancellationToken);
+			}
+		}
+		finally
+		{
+			mutex1?.Dispose();
+		}
+	}
+
+	[TestMethod]
+	public void TryAcquiringGlobalMutex()
+	{
+		const string MutexName = "EpicGames.Build.Tests.UnitTestMutex3";
+		IDisposable? mutex1 = null;
+		IDisposable? mutex2 = null;
+
+		try
+		{
+			bool didAcquire = SingleInstanceMutex.TryAcquire(MutexName, out mutex1);
+			Assert.IsTrue(didAcquire);
+			Assert.IsNotNull(mutex1);
+
+			didAcquire = SingleInstanceMutex.TryAcquire(MutexName, out mutex2);
+			Assert.IsFalse(didAcquire);
+			Assert.IsNull(mutex2);
+		}
+		finally
+		{
+			mutex1?.Dispose();
+			mutex2?.Dispose();
+		}
+	}
+
+	public TestContext TestContext { get; set; }
+}
diff --git a/Engine/Source/Programs/Shared/EpicGames.Core/SingleInstanceMutex.cs b/Engine/Source/Programs/Shared/EpicGames.Core/SingleInstanceMutex.cs
--- a/Engine/Source/Programs/Shared/EpicGames.Core/SingleInstanceMutex.cs
+++ b/Engine/Source/Programs/Shared/EpicGames.Core/SingleInstanceMutex.cs
@@ -1,6 +1,7 @@
 // Copyright Epic Games, Inc. All Rights Reserved.
 
 using System;
+using System.Diagnostics.CodeAnalysis;
 using System.Runtime.InteropServices;
 using System.Threading;
 using System.Threading.Tasks;
@@ -8,105 +9,162 @@ using System.Threading.Tasks;
 namespace EpicGames.Core
 {
 	/// <summary>
-	/// Allows taking a lock on a global mutex from async code.
+	/// System-wide (global) mutex that is thread- and async-safe.
 	/// </summary>
-	public static class SingleInstanceMutex
+	public class SingleInstanceMutex : IDisposable
 	{
-		class SingleInstanceMutexImpl : IDisposable
-		{
-			readonly string? _name;
-			readonly TaskCompletionSource _readyTcs = new(TaskCreationOptions.RunContinuationsAsynchronously);
-			readonly ManualResetEvent _disposing = new(false);
-			Thread? _thread;
+		/// <summary>
+		/// The global mutex instance
+		/// </summary>
+		private Mutex? _globalMutex;
 
-			public SingleInstanceMutexImpl(string name)
-			{
-				_name = name;
+		/// <summary>
+		/// A single-threaded scheduler to run the mutex acquire / release on, to make this class safe for use within an async context.
+		/// </summary>
+		private readonly SingleThreadTaskScheduler _scheduler = new();
 
-				_thread = new Thread(() => ThreadFunc());
-				_thread.Start();
-			}
+		/// <summary>
+		/// Release the mutex and dispose of the object
+		/// </summary>
+		public void Dispose()
+		{
+			Dispose(true);
+			GC.SuppressFinalize(this);
+		}
 
-			public async ValueTask WaitAsync(CancellationToken cancellationToken)
+		/// <summary>
+		/// Overridable dispose method
+		/// </summary>
+		/// <param name="disposing">Whether the object should be disposed</param>
+		protected virtual void Dispose(bool disposing)
+		{
+			if (disposing)
 			{
-				using (CancellationTokenRegistration registration = cancellationToken.Register(() => _readyTcs.TrySetCanceled()))
+				RunActionOnScheduler(() =>
 				{
-					await _readyTcs.Task;
-				}
+					_globalMutex?.ReleaseMutex();
+					return true;
+				});
+				_globalMutex?.Dispose();
+				_globalMutex = null;
+				_scheduler.Dispose();
 			}
+		}
 
-			public void Dispose()
-			{
-				if (_thread != null)
-				{
-					_disposing.Set();
-
-					_thread.Join();
-					_thread = null;
-				}
+		private Task<bool> RunActionOnScheduler(Func<bool> action, CancellationToken cancellationToken = default)
+		{
+			return Task.Factory.StartNew(action, cancellationToken, TaskCreationOptions.None, _scheduler);
+		}
 
-				_disposing.Dispose();
-			}
+		private bool GetMutex(string mutexName, bool wait, CancellationToken cancellationToken = default)
+		{
+			// Try to create the mutex, with it initially locked
+			_globalMutex = new Mutex(true, mutexName, out bool createdMutex);
 
-			void ThreadFunc()
+			// If we didn't create the mutex, we can wait for it or fail immediately
+			if (!createdMutex)
 			{
-				using Mutex mutex = new(false, _name);
-				try
+				if (wait)
 				{
-					// Waiting for multiple handles is only supported on Windows
-					if (RuntimeInformation.IsOSPlatform(OSPlatform.Windows))
+					try
 					{
-						WaitHandle[] handles = [_disposing, mutex];
-						if (WaitHandle.WaitAny(handles) == 0)
+						// Waiting for multiple handles is only supported on Windows
+						if (RuntimeInformation.IsOSPlatform(OSPlatform.Windows))
 						{
-							return;
+							WaitHandle[] handles = [cancellationToken.WaitHandle, _globalMutex];
+							if (WaitHandle.WaitAny(handles) == 0)
+							{
+								cancellationToken.ThrowIfCancellationRequested();
+							}
 						}
-					}
-					else
-					{
-						while (!mutex.WaitOne(100))
+						else
 						{
-							if (_disposing.WaitOne(0))
+							while (!_globalMutex.WaitOne(100))
 							{
-								return;
+								if (cancellationToken.IsCancellationRequested)
+								{
+									cancellationToken.ThrowIfCancellationRequested();
+								}
 							}
 						}
 					}
+					catch (AbandonedMutexException)
+					{
+					}
 				}
-				catch (AbandonedMutexException)
+				else
 				{
+					return false;
 				}
+			}
 
-				_readyTcs.TrySetResult();
-				_disposing.WaitOne();
+			return true;
+		}
 
-				mutex.ReleaseMutex();
-			}
+		/// <summary>
+		/// Acquires a global mutex with the given name.
+		/// </summary>
+		/// <param name="name">Name of the mutex to acquire.</param>
+		/// <returns>Handle to the mutex. Must be disposed when complete.</returns>
+		public static IDisposable Acquire(string name)
+		{
+#pragma warning disable CA2000 // Don't need to dispose of tasks.
+			return AcquireAsync(name).GetAwaiter().GetResult();
+#pragma warning restore CA2000
 		}
 
 		/// <summary>
 		/// Acquires a global mutex with the given name.
 		/// </summary>
-		/// <param name="name">Name of the mutex to acquire</param>
-		/// <param name="cancellationToken">Cancellation token for the wait</param>
+		/// <param name="name">Name of the mutex to acquire.</param>
+		/// <param name="cancellationToken">Cancellation token for the wait.</param>
 		/// <returns>Handle to the mutex. Must be disposed when complete.</returns>
 		public static async Task<IDisposable> AcquireAsync(string name, CancellationToken cancellationToken = default)
 		{
-			SingleInstanceMutexImpl? impl = new(name);
+			SingleInstanceMutex mutex = new();
 			try
 			{
-				await impl.WaitAsync(cancellationToken);
-				cancellationToken.ThrowIfCancellationRequested();
+				await mutex.RunActionOnScheduler(() => mutex.GetMutex(name, true, cancellationToken), cancellationToken).ConfigureAwait(false);
+			}
+			catch
+			{
+				// Calling dispose would try to release the mutex, so don't do that - instead just dispose the bits we need to.
+				mutex._globalMutex?.Dispose();
+				mutex._scheduler.Dispose();
+				throw;
+			}
+			return mutex;
+		}
+
+		/// <summary>
+		/// Tries to acquire a global mutex with the given name if it's available.
+		/// </summary>
+		/// <param name="name">Name of the mutex to acquire.</param>
+		/// <param name="mutex">The mutex if acquired, or <see langword="null"/> if not.</param>
+		/// <returns>Whether the mutex was acquired.</returns>
+		public static bool TryAcquire(string name, [NotNullWhen(true)] out IDisposable? mutex)
+		{
+			SingleInstanceMutex tempMutex = new();
+			try
+			{
+				bool didAcquire = tempMutex.RunActionOnScheduler(() => tempMutex.GetMutex(name, false)).GetAwaiter().GetResult();
+				if (!didAcquire)
+				{
+					mutex = null;
+					tempMutex._globalMutex?.Dispose();
+					tempMutex._scheduler.Dispose();
+					return false;
+				}
 
-				SingleInstanceMutexImpl result = impl;
-				impl = null;
-				return result;
+				mutex = tempMutex;
+				return true;
 			}
-			finally
+			catch
 			{
-#pragma warning disable CA1508 // Avoid dead conditional code
-				impl?.Dispose();
-#pragma warning restore CA1508 // Avoid dead conditional code
+				mutex = null;
+				tempMutex._globalMutex?.Dispose();
+				tempMutex._scheduler.Dispose();
+				return false;
 			}
 		}
 	}
diff --git a/Engine/Source/Programs/UnrealBuildTool/Configuration/Rules/RulesCompiler.cs b/Engine/Source/Programs/UnrealBuildTool/Configuration/Rules/RulesCompiler.cs
--- a/Engine/Source/Programs/UnrealBuildTool/Configuration/Rules/RulesCompiler.cs
+++ b/Engine/Source/Programs/UnrealBuildTool/Configuration/Rules/RulesCompiler.cs
@@ -73,7 +73,7 @@ namespace UnrealBuildTool
 				if (EngineRulesAssembly == null)
 				{
 					// Prevent multiple conflicting processes building rules assembly at the same time
-					using GlobalSingleInstanceMutex _ = new(GlobalSingleInstanceMutex.GetUniqueMutexForPath("UnrealBuildTool_CreateTargetRulesAssembly", Unreal.EngineDirectory), true);
+					using IDisposable _ = SingleInstanceMutex.Acquire(GlobalSingleInstanceMutex.GetUniqueMutexForPath("UnrealBuildTool_CreateTargetRulesAssembly", Unreal.EngineDirectory));
 					{
 						List<PluginInfo> EnginePlugins = [];
 						List<PluginInfo> MarketplacePlugins = [];
@@ -291,7 +291,7 @@ namespace UnrealBuildTool
 			return LoadedAssemblyMap.GetOrAdd(ProjectFileName, new Lazy<RulesAssembly>(() =>
 			{
 				// Prevent multiple conflicting processes building rules assembly at the same time
-				using GlobalSingleInstanceMutex _ = new(GlobalSingleInstanceMutex.GetUniqueMutexForPath("UnrealBuildTool_CreateTargetRulesAssembly", ProjectFileName), true);
+				using IDisposable _ = SingleInstanceMutex.Acquire(GlobalSingleInstanceMutex.GetUniqueMutexForPath("UnrealBuildTool_CreateTargetRulesAssembly", ProjectFileName));
 
 				Logger.LogTrace("Creating project rules assembly for {Project}...", ProjectFileName);
 
@@ -421,7 +421,7 @@ namespace UnrealBuildTool
 			return LoadedAssemblyMap.GetOrAdd(PluginFileName, new Lazy<RulesAssembly>(() =>
 			{
 				// Prevent multiple conflicting processes building rules assembly at the same time
-				using GlobalSingleInstanceMutex _ = new(GlobalSingleInstanceMutex.GetUniqueMutexForPath("UnrealBuildTool_CreateTargetRulesAssembly", PluginFileName), true);
+				using IDisposable _ = SingleInstanceMutex.Acquire(GlobalSingleInstanceMutex.GetUniqueMutexForPath("UnrealBuildTool_CreateTargetRulesAssembly", PluginFileName));
 
 				Logger.LogTrace("Creating plugin rules assembly for {Plugin}...", PluginFileName);
 
diff --git a/Engine/Source/Programs/UnrealBuildTool/Configuration/Xml/XmlConfig.cs b/Engine/Source/Programs/UnrealBuildTool/Configuration/Xml/XmlConfig.cs
--- a/Engine/Source/Programs/UnrealBuildTool/Configuration/Xml/XmlConfig.cs
+++ b/Engine/Source/Programs/UnrealBuildTool/Configuration/Xml/XmlConfig.cs
@@ -74,7 +74,7 @@ namespace UnrealBuildTool
 		public static void ReadConfigFiles(FileReference? overrideCacheFile, DirectoryReference? projectRootDirectory, ILogger logger)
 		{
 			string XmlConfigMutexName = GlobalSingleInstanceMutex.GetUniqueMutexForPath("UnrealBuildTool_Mutex_XmlConfig", Unreal.EngineDirectory);
-			using GlobalSingleInstanceMutex XmlConfigMutex = new GlobalSingleInstanceMutex(XmlConfigMutexName, true);
+			using IDisposable XmlConfigMutex = SingleInstanceMutex.Acquire(XmlConfigMutexName);
 
 			// Find all the configurable types
 			List<Type> configTypes = FindConfigurableTypes();
diff --git a/Engine/Source/Programs/UnrealBuildTool/Modes/QueryMode.cs b/Engine/Source/Programs/UnrealBuildTool/Modes/QueryMode.cs
--- a/Engine/Source/Programs/UnrealBuildTool/Modes/QueryMode.cs
+++ b/Engine/Source/Programs/UnrealBuildTool/Modes/QueryMode.cs
@@ -277,7 +277,7 @@ namespace UnrealBuildTool
 
 								// Prevent multiple conflicting processes building TargetRules at the same time
 								string MutexName = GlobalSingleInstanceMutex.GetUniqueMutexForPath("UnrealBuildTool_QueryMode_UEBuildTarget-Create", Unreal.RootDirectory);
-								using (new GlobalSingleInstanceMutex(MutexName, true))
+								using (await SingleInstanceMutex.AcquireAsync(MutexName))
 								{
 									CurrentTarget = UEBuildTarget.Create(TargetDescriptors[0], false, false, bUsePrecompiled, TargetDescriptors[0].IntermediateEnvironment, NullLogger.Instance);
 								}
@@ -374,7 +374,7 @@ namespace UnrealBuildTool
 
 					// Prevent multiple conflicting processes building TargetRules at the same time
 					string MutexName = GlobalSingleInstanceMutex.GetUniqueMutexForPath("UnrealBuildTool_QueryMode_UEBuildTarget-Create", Unreal.RootDirectory);
-					using (new GlobalSingleInstanceMutex(MutexName, true))
+					using (await SingleInstanceMutex.AcquireAsync(MutexName))
 					{
 						CurrentTarget = UEBuildTarget.Create(TargetDescriptors[0], false, false, bUsePrecompiled, TargetDescriptors[0].IntermediateEnvironment, Logger);
 					}
@@ -556,7 +556,7 @@ namespace UnrealBuildTool
 
 					// Prevent multiple conflicting processes building TargetRules at the same time
 					string MutexName = GlobalSingleInstanceMutex.GetUniqueMutexForPath("UnrealBuildTool_QueryMode_UEBuildTarget-Create", Unreal.RootDirectory);
-					using (new GlobalSingleInstanceMutex(MutexName, true))
+					using (await SingleInstanceMutex.AcquireAsync(MutexName))
 					{
 						CurrentTarget = UEBuildTarget.Create(TargetDescriptors[0], false, false, bUsePrecompiled, TargetDescriptors[0].IntermediateEnvironment, Logger);
 					}
diff --git a/Engine/Source/Programs/UnrealBuildTool/Modes/QueryTargetsMode.cs b/Engine/Source/Programs/UnrealBuildTool/Modes/QueryTargetsMode.cs
--- a/Engine/Source/Programs/UnrealBuildTool/Modes/QueryTargetsMode.cs
+++ b/Engine/Source/Programs/UnrealBuildTool/Modes/QueryTargetsMode.cs
@@ -58,7 +58,7 @@ namespace UnrealBuildTool
 		/// <param name="arguments">Command line arguments</param>
 		/// <returns></returns>
 		/// <param name="logger"></param>
-		public Task<int> ExecuteAsync(CommandLineArguments arguments, ILogger logger)
+		public async Task<int> ExecuteAsync(CommandLineArguments arguments, ILogger logger)
 		{
 			arguments.ApplyTo(this);
 
@@ -84,7 +84,7 @@ namespace UnrealBuildTool
 			// Create the rules assembly
 			RulesAssembly assembly;
 			string mutexName = GlobalSingleInstanceMutex.GetUniqueMutexForPath("UnrealBuildTool_QueryMode_UEBuildTarget-Create", Unreal.RootDirectory);
-			using (new GlobalSingleInstanceMutex(mutexName, true))
+			using (await SingleInstanceMutex.AcquireAsync(mutexName))
 			{
 				assembly = ProjectFile != null
 					? RulesCompiler.CreateProjectRulesAssembly(ProjectFile, buildConfiguration.bUsePrecompiled, buildConfiguration.bSkipRulesCompile, buildConfiguration.bForceRulesCompile, logger)
@@ -94,7 +94,7 @@ namespace UnrealBuildTool
 			// Write information about these targets
 			WriteTargetInfo(ProjectFile, assembly, OutputFile, arguments, logger, bIncludeAllTargets, bIncludeParentAssembly, bIncludeProgramTargets);
 			logger.LogInformation("Written {OutputFile}", OutputFile);
-			return Task.FromResult(0);
+			return 0;
 		}
 
 		/// <summary>
diff --git a/Engine/Source/Programs/UnrealBuildTool/Modes/WriteMetadataMode.cs b/Engine/Source/Programs/UnrealBuildTool/Modes/WriteMetadataMode.cs
--- a/Engine/Source/Programs/UnrealBuildTool/Modes/WriteMetadataMode.cs
+++ b/Engine/Source/Programs/UnrealBuildTool/Modes/WriteMetadataMode.cs
@@ -65,7 +65,7 @@ namespace UnrealBuildTool
 			// Acquire a different mutex to the regular UBT instance, since this mode will be called as part of a build. We need the mutex to ensure that building two modular configurations 
 			// in parallel don't clash over writing shared *.modules files (eg. DebugGame and Development editors).
 			string MutexName = GlobalSingleInstanceMutex.GetUniqueMutexForPath("UnrealBuildTool_WriteMetadata", Unreal.RootDirectory);
-			using (new GlobalSingleInstanceMutex(MutexName, true))
+			using (await SingleInstanceMutex.AcquireAsync(MutexName))
 			{
 				return await ExecuteInternal(Arguments, Logger);
 			}
diff --git a/Engine/Source/Programs/UnrealBuildTool/Platform/Apple/AppleToolChain.cs b/Engine/Source/Programs/UnrealBuildTool/Platform/Apple/AppleToolChain.cs
--- a/Engine/Source/Programs/UnrealBuildTool/Platform/Apple/AppleToolChain.cs
+++ b/Engine/Source/Programs/UnrealBuildTool/Platform/Apple/AppleToolChain.cs
@@ -1078,7 +1078,7 @@ namespace UnrealBuildTool
 			// Acquire a different mutex to the regular UBT instance, since this mode will be called as part of a build. We need the mutex to ensure that building two modular configurations 
 			// in parallel don't clash over writing shared *.modules files (eg. DebugGame and Development editors).
 			string MutexName = GlobalSingleInstanceMutex.GetUniqueMutexForPath("UnrealBuildTool_XcodeBuild", Unreal.RootDirectory);
-			using (new GlobalSingleInstanceMutex(MutexName, true))
+			using (SingleInstanceMutex.Acquire(MutexName))
 			{
 				for (int Pass = 0; Pass < 2; Pass++)
 				{
diff --git a/Engine/Source/Programs/UnrealBuildTool/Platform/IOS/UEDeployIOS.cs b/Engine/Source/Programs/UnrealBuildTool/Platform/IOS/UEDeployIOS.cs
--- a/Engine/Source/Programs/UnrealBuildTool/Platform/IOS/UEDeployIOS.cs
+++ b/Engine/Source/Programs/UnrealBuildTool/Platform/IOS/UEDeployIOS.cs
@@ -495,7 +495,7 @@ namespace UnrealBuildTool
 				{
 					// need to make sure touching this plist is in the same lock as FinalizeAppWithXcode()
 					string MutexName = GlobalSingleInstanceMutex.GetUniqueMutexForPath("UnrealBuildTool_XcodeBuild", Unreal.RootDirectory);
-					using (new GlobalSingleInstanceMutex(MutexName, true))
+					using (SingleInstanceMutex.Acquire(MutexName))
 					{
 						FileReference FinalPlistFile;
 						FinalPlistFile = new FileReference($"{ProjectDirectory}/Build/IOS/UBTGenerated/Info.Template.plist");
diff --git a/Engine/Source/Programs/UnrealBuildTool/UnrealBuildTool.cs b/Engine/Source/Programs/UnrealBuildTool/UnrealBuildTool.cs
--- a/Engine/Source/Programs/UnrealBuildTool/UnrealBuildTool.cs
+++ b/Engine/Source/Programs/UnrealBuildTool/UnrealBuildTool.cs
@@ -178,7 +178,6 @@ namespace UnrealBuildTool
 
 		/// <summary>
 		/// Main entry point. Parses any global options and initializes the logging system, then invokes the appropriate command.
-		/// NB: That the entry point is deliberately NOT async, since we have a single-instance mutex that cannot be disposed from a different thread.
 		/// </summary>
 		/// <param name="ArgumentsArray">Command line arguments</param>
 		/// <returns>Zero on success, non-zero on error</returns>
@@ -187,7 +186,7 @@ namespace UnrealBuildTool
 			ToolModeOptions ModeOptions = ToolModeOptions.None;
 			FileReference? RunFile = null;
 			DirectoryReference? TempDirectory = null;
-			GlobalSingleInstanceMutex? Mutex = null;
+			IDisposable? Mutex = null;
 
 			ILogger Logger = Log.Logger;
 
@@ -449,7 +448,7 @@ namespace UnrealBuildTool
 					using (GlobalTracer.Instance.BuildSpan("SingleInstanceMutex.Acquire()").StartActive())
 					{
 						string MutexName = GlobalSingleInstanceMutex.GetUniqueMutexForPath("UnrealBuildTool_Mutex", FileReference.FromString(Assembly.GetExecutingAssembly().Location));
-						Mutex = new GlobalSingleInstanceMutex(MutexName, Options.WaitMutex);
+						Mutex = Options.WaitMutex ? SingleInstanceMutex.Acquire(MutexName) : GlobalSingleInstanceMutex.AcquireNowOrThrow(MutexName);
 					}
 				}
 
