EpicGames.Core: Fix a potential exception when searching for single letter names in a FileSystemEntry that happen to match the Windows drive letter.

#12831

#Synced-CL 52569840

diff --git a/Engine/Source/Programs/Shared/EpicGames.Core/FileSystemReference.cs b/Engine/Source/Programs/Shared/EpicGames.Core/FileSystemReference.cs
--- a/Engine/Source/Programs/Shared/EpicGames.Core/FileSystemReference.cs
+++ b/Engine/Source/Programs/Shared/EpicGames.Core/FileSystemReference.cs
@@ -216,6 +216,15 @@ namespace EpicGames.Core
 				return false;
 			}
 
+			// Adjust the offset to ignore drive letters on Windows. This will either set the start to `:` on Windows, or past the initial `/` on Unix.
+			// This is fine in both cases, as `:` is a forbidden filename character on Windows, and `/` is a forbidden filename character on Windows and Unix.
+			// Also works for UNC paths by skipping the first `\`, which is a forbidden filename character on Windows.
+			// As such, moving the offset to a minimum of 1 cannot possibly affect the results of a valid query.
+			// This prevents a search for "D" from finding the drive letter and the below access of `FullName[matchIdx - 1]` being out of bounds.
+			int originalOffset = offset;
+			offset = Math.Max(offset, 1);
+			length -= offset - originalOffset;
+
 			// Find each occurrence of the name within the remaining string, then test whether it's surrounded by directory separators
 			int matchIdx = offset;
 			for (; ; )
