UnrealBuildTool / EpicGames.Shared: Now that we're on dotnet 10, use the Lock type rather than object for synchronization.

#Synced-CL 52569840

diff --git a/Engine/Source/Programs/Shared/.editorconfig b/Engine/Source/Programs/Shared/.editorconfig
--- a/Engine/Source/Programs/Shared/.editorconfig
+++ b/Engine/Source/Programs/Shared/.editorconfig
@@ -168,6 +168,9 @@ dotnet_diagnostic.IDE0301.severity = warning
 # IDE0305: Use collection expression for fluent
 dotnet_diagnostic.IDE0305.severity = warning
 
+# IDE0330: Prefer 'System.Threading.Lock'
+dotnet_diagnostic.IDE0330.severity = warning
+
 # IDE0011: Add braces
 csharp_prefer_braces = when_multiline:warning
 # NOTE: We need the below severity entry for Add Braces due to [REDACTED_PROFILE_URL]/roslyn/issues/44201
diff --git a/Engine/Source/Programs/Shared/EpicGames.Analytics.Telemetry/ColumnValidator.cs b/Engine/Source/Programs/Shared/EpicGames.Analytics.Telemetry/ColumnValidator.cs
--- a/Engine/Source/Programs/Shared/EpicGames.Analytics.Telemetry/ColumnValidator.cs
+++ b/Engine/Source/Programs/Shared/EpicGames.Analytics.Telemetry/ColumnValidator.cs
@@ -5,6 +5,7 @@ using System.Collections;
 using System.Collections.Generic;
 using System.ComponentModel.DataAnnotations.Schema;
 using System.Reflection;
+using System.Threading;
 
 namespace EpicGames.Analytics
 {
@@ -14,7 +15,7 @@ namespace EpicGames.Analytics
 	/// <typeparam name="T">The record type to validate against.</typeparam>
 	public static class ColumnValidator<T>
 	{
-		private static readonly object s_initLock = new();
+		private static readonly Lock s_initLock = new();
 		private static HashSet<string>? s_columnKeys = null;
 		private static HashSet<string>? s_propertyKeys = null;
 
diff --git a/Engine/Source/Programs/Shared/EpicGames.Build/System/NativeProjectsBase.cs b/Engine/Source/Programs/Shared/EpicGames.Build/System/NativeProjectsBase.cs
--- a/Engine/Source/Programs/Shared/EpicGames.Build/System/NativeProjectsBase.cs
+++ b/Engine/Source/Programs/Shared/EpicGames.Build/System/NativeProjectsBase.cs
@@ -4,6 +4,7 @@ using System;
 using System.Collections.Generic;
 using System.IO;
 using System.Linq;
+using System.Threading;
 using EpicGames.Core;
 using Microsoft.Extensions.Logging;
 
@@ -17,7 +18,7 @@ namespace UnrealBuildBase
 		/// <summary>
 		/// Object used for synchronizing access to static fields
 		/// </summary>
-		static protected object LockObject = new();
+		static protected Lock LockObject = new();
 
 		/// <summary>
 		/// The native project base directories
diff --git a/Engine/Source/Programs/Shared/EpicGames.Core/AsyncThreadPoolWorkQueue.cs b/Engine/Source/Programs/Shared/EpicGames.Core/AsyncThreadPoolWorkQueue.cs
--- a/Engine/Source/Programs/Shared/EpicGames.Core/AsyncThreadPoolWorkQueue.cs
+++ b/Engine/Source/Programs/Shared/EpicGames.Core/AsyncThreadPoolWorkQueue.cs
@@ -17,7 +17,7 @@ public sealed class AsyncThreadPoolWorkQueue : IDisposable
 	/// <summary>
 	/// Object used for controlling access to updating _lockObject
 	/// </summary>
-	private readonly object _lockObject = new();
+	private readonly Lock _lockObject = new();
 
 	/// <summary>
 	/// Event which indicates whether the queue is empty.
diff --git a/Engine/Source/Programs/Shared/EpicGames.Core/IMemoryAllocator.cs b/Engine/Source/Programs/Shared/EpicGames.Core/IMemoryAllocator.cs
--- a/Engine/Source/Programs/Shared/EpicGames.Core/IMemoryAllocator.cs
+++ b/Engine/Source/Programs/Shared/EpicGames.Core/IMemoryAllocator.cs
@@ -254,7 +254,7 @@ namespace EpicGames.Core
 				=> $"{_tag ?? "???"} ({_allocation?.Memory.Length ?? 0} bytes)";
 		}
 
-		readonly object _lockObject = new();
+		readonly Lock _lockObject = new();
 		readonly IMemoryAllocator<byte> _inner;
 		long _allocatedSize;
 		readonly LinkedList<MemoryAllocation> _allocations = new();
diff --git a/Engine/Source/Programs/Shared/EpicGames.Core/JsonTracer.cs b/Engine/Source/Programs/Shared/EpicGames.Core/JsonTracer.cs
--- a/Engine/Source/Programs/Shared/EpicGames.Core/JsonTracer.cs
+++ b/Engine/Source/Programs/Shared/EpicGames.Core/JsonTracer.cs
@@ -96,7 +96,7 @@ namespace EpicGames.Core
 		/// <seealso cref="MockSpan.References"/>
 		public string? ParentId { get; }
 
-		private readonly object _lock = new();
+		private readonly Lock _lock = new();
 
 		public JsonTracerSpan(JsonTracer jsonTracer, string operationName, DateTimeOffset startTimestamp, Dictionary<string, object>? initialTags, List<Reference>? references)
 		{
@@ -404,7 +404,7 @@ namespace EpicGames.Core
 		public IScopeManager ScopeManager { get; }
 		public ISpan? ActiveSpan => ScopeManager.Active?.Span;
 
-		private readonly object _lock = new();
+		private readonly Lock _lock = new();
 		private readonly List<JsonTracerSpan> _finishedSpans = [];
 		private readonly DirectoryReference? _telemetryDir;
 
diff --git a/Engine/Source/Programs/Shared/EpicGames.Core/Logging/Log.cs b/Engine/Source/Programs/Shared/EpicGames.Core/Logging/Log.cs
--- a/Engine/Source/Programs/Shared/EpicGames.Core/Logging/Log.cs
+++ b/Engine/Source/Programs/Shared/EpicGames.Core/Logging/Log.cs
@@ -10,6 +10,7 @@ using System.Linq;
 using System.Reflection;
 using System.Text;
 using System.Text.RegularExpressions;
+using System.Threading;
 using System.Threading.Channels;
 using System.Threading.Tasks;
 using JetBrains.Annotations;
@@ -932,7 +933,7 @@ namespace EpicGames.Core
 		}
 
 		private readonly LoggerScopeCollection _scopeCollection = new();
-		private readonly object _syncObject = new();
+		private readonly Lock _syncObject = new();
 
 		/// <summary>
 		/// Minimum level for outputting messages
diff --git a/Engine/Source/Programs/Shared/EpicGames.Core/ManagedProcess.cs b/Engine/Source/Programs/Shared/EpicGames.Core/ManagedProcess.cs
--- a/Engine/Source/Programs/Shared/EpicGames.Core/ManagedProcess.cs
+++ b/Engine/Source/Programs/Shared/EpicGames.Core/ManagedProcess.cs
@@ -478,7 +478,7 @@ namespace EpicGames.Core
 		/// Static lock object. This is used to synchronize the creation of child processes - in particular, the inheritance of stdout/stderr write pipes. If processes
 		/// inherit pipes meant for other processes, they won't be closed until both terminate.
 		/// </summary>
-		static readonly object s_lockObject = new();
+		static readonly Lock s_lockObject = new();
 
 		/// <summary>
 		/// Used to perform CPU usage and resource accounting of all children process involved in a single compilation unit.
diff --git a/Engine/Source/Programs/Shared/EpicGames.Core/MemoryMappedFileCache.cs b/Engine/Source/Programs/Shared/EpicGames.Core/MemoryMappedFileCache.cs
--- a/Engine/Source/Programs/Shared/EpicGames.Core/MemoryMappedFileCache.cs
+++ b/Engine/Source/Programs/Shared/EpicGames.Core/MemoryMappedFileCache.cs
@@ -148,7 +148,7 @@ namespace EpicGames.Core
 			}
 		}
 
-		readonly object _lockObject = new();
+		readonly Lock _lockObject = new();
 		readonly Dictionary<FileReference, MappedFile> _pathToMappedFile = [];
 		readonly LinkedList<MappedFile> _mappedFiles = new();
 
diff --git a/Engine/Source/Programs/Shared/EpicGames.Core/ThreadPoolWorkQueue.cs b/Engine/Source/Programs/Shared/EpicGames.Core/ThreadPoolWorkQueue.cs
--- a/Engine/Source/Programs/Shared/EpicGames.Core/ThreadPoolWorkQueue.cs
+++ b/Engine/Source/Programs/Shared/EpicGames.Core/ThreadPoolWorkQueue.cs
@@ -15,7 +15,7 @@ namespace EpicGames.Core
 		/// <summary>
 		/// Object used for controlling access to updating _lockObject
 		/// </summary>
-		readonly object _lockObject = new();
+		readonly Lock _lockObject = new();
 
 		/// <summary>
 		/// Number of jobs remaining in the queue. This is updated in an atomic way.
diff --git a/Engine/Source/Programs/Shared/EpicGames.Horde/Compute/ComputeSocket.cs b/Engine/Source/Programs/Shared/EpicGames.Horde/Compute/ComputeSocket.cs
--- a/Engine/Source/Programs/Shared/EpicGames.Horde/Compute/ComputeSocket.cs
+++ b/Engine/Source/Programs/Shared/EpicGames.Horde/Compute/ComputeSocket.cs
@@ -374,7 +374,7 @@ namespace EpicGames.Horde.Compute
 			}
 		}
 
-		readonly object _lockObject = new();
+		readonly Lock _lockObject = new();
 
 		bool _complete;
 		bool _isClosing;
diff --git a/Engine/Source/Programs/Shared/EpicGames.Horde/Compute/RequestIdAllocator.cs b/Engine/Source/Programs/Shared/EpicGames.Horde/Compute/RequestIdAllocator.cs
--- a/Engine/Source/Programs/Shared/EpicGames.Horde/Compute/RequestIdAllocator.cs
+++ b/Engine/Source/Programs/Shared/EpicGames.Horde/Compute/RequestIdAllocator.cs
@@ -3,6 +3,7 @@
 using System;
 using System.Collections.Generic;
 using System.Security.Cryptography;
+using System.Threading;
 
 namespace EpicGames.Horde.Compute;
 
@@ -11,7 +12,7 @@ namespace EpicGames.Horde.Compute;
 /// </summary>
 public class RequestIdAllocator
 {
-	private readonly object _lock = new();
+	private readonly Lock _lock = new();
 	private readonly List<string> _pendingIds = [];
 	private readonly List<string> _batchIds = [];
 	private readonly string _allocatorId;
diff --git a/Engine/Source/Programs/Shared/EpicGames.Horde/Logs/LogNode.cs b/Engine/Source/Programs/Shared/EpicGames.Horde/Logs/LogNode.cs
--- a/Engine/Source/Programs/Shared/EpicGames.Horde/Logs/LogNode.cs
+++ b/Engine/Source/Programs/Shared/EpicGames.Horde/Logs/LogNode.cs
@@ -149,7 +149,7 @@ namespace EpicGames.Horde.Logs
 		readonly LogChunkSequenceBuilder _indexTextBuilder;
 
 		// Lock object for access to the fields above
-		readonly object _lockObject = new();
+		readonly Lock _lockObject = new();
 
 		// Inner log device; used for messaging encoding errors.
 		readonly ILogger _logger;
diff --git a/Engine/Source/Programs/Shared/EpicGames.Horde/Storage/Bundles/V1/BundleReader.cs b/Engine/Source/Programs/Shared/EpicGames.Horde/Storage/Bundles/V1/BundleReader.cs
--- a/Engine/Source/Programs/Shared/EpicGames.Horde/Storage/Bundles/V1/BundleReader.cs
+++ b/Engine/Source/Programs/Shared/EpicGames.Horde/Storage/Bundles/V1/BundleReader.cs
@@ -129,7 +129,7 @@ namespace EpicGames.Horde.Storage.Bundles.V1
 		readonly BundleCache _cache;
 		readonly ILogger _logger;
 
-		readonly object _queueLock = new();
+		readonly Lock _queueLock = new();
 		readonly List<QueuedBundle> _queuedBundles = [];
 		readonly Dictionary<string, Task<ReadOnlyMemory<byte>>> _decodeTasks = new(StringComparer.Ordinal);
 
diff --git a/Engine/Source/Programs/Shared/EpicGames.Horde/Storage/Bundles/V1/BundleWriter.cs b/Engine/Source/Programs/Shared/EpicGames.Horde/Storage/Bundles/V1/BundleWriter.cs
--- a/Engine/Source/Programs/Shared/EpicGames.Horde/Storage/Bundles/V1/BundleWriter.cs
+++ b/Engine/Source/Programs/Shared/EpicGames.Horde/Storage/Bundles/V1/BundleWriter.cs
@@ -214,7 +214,7 @@ namespace EpicGames.Horde.Storage.Bundles.V1
 			readonly int _maxPacketSize;
 			readonly BundleCompressionFormat _compressionFormat;
 
-			readonly object _lockObject = new();
+			readonly Lock _lockObject = new();
 
 			// The current packet being built
 			IMemoryOwner<byte>? _currentPacket;
diff --git a/Engine/Source/Programs/Shared/EpicGames.Horde/Storage/Bundles/V2/BundleWriter.cs b/Engine/Source/Programs/Shared/EpicGames.Horde/Storage/Bundles/V2/BundleWriter.cs
--- a/Engine/Source/Programs/Shared/EpicGames.Horde/Storage/Bundles/V2/BundleWriter.cs
+++ b/Engine/Source/Programs/Shared/EpicGames.Horde/Storage/Bundles/V2/BundleWriter.cs
@@ -55,7 +55,7 @@ namespace EpicGames.Horde.Storage.Bundles.V2
 		// Bundle that has not yet been written to storage.
 		internal sealed class PendingBundleHandle : BundleHandle, IDisposable
 		{
-			public readonly object LockObject = new();
+			public readonly Lock LockObject = new();
 
 			readonly BundleWriter _outer;
 			readonly BundleStorageNamespace _storageNamespace;
@@ -256,7 +256,9 @@ namespace EpicGames.Horde.Storage.Bundles.V2
 				{
 					lock (LockObject)
 					{
+#pragma warning disable CA1508 // Avoid dead conditional code - need to re-check after locking.
 						if (_flushedHandle == null)
+#pragma warning restore CA1508
 						{
 							Debug.Assert(_encodedPacketWriter != null);
 							int fetchLength = length ?? (_encodedPacketWriter.Length - offset);
@@ -283,7 +285,7 @@ namespace EpicGames.Horde.Storage.Bundles.V2
 
 		PendingBundleHandle _currentBundle;
 
-		readonly object _lockObject = new();
+		readonly Lock _lockObject = new();
 		Task _flushTask = Task.CompletedTask;
 
 		/// <summary>
diff --git a/Engine/Source/Programs/Shared/EpicGames.Horde/Storage/Bundles/V2/PacketWriter.cs b/Engine/Source/Programs/Shared/EpicGames.Horde/Storage/Bundles/V2/PacketWriter.cs
--- a/Engine/Source/Programs/Shared/EpicGames.Horde/Storage/Bundles/V2/PacketWriter.cs
+++ b/Engine/Source/Programs/Shared/EpicGames.Horde/Storage/Bundles/V2/PacketWriter.cs
@@ -6,6 +6,7 @@ using System.Collections.Generic;
 using System.Diagnostics;
 using System.Linq;
 using System.Runtime.InteropServices;
+using System.Threading;
 using EpicGames.Core;
 
 namespace EpicGames.Horde.Storage.Bundles.V2
@@ -20,7 +21,7 @@ namespace EpicGames.Horde.Storage.Bundles.V2
 		readonly BundleHandle _bundleHandle;
 		readonly PacketHandle _packetHandle;
 		readonly IMemoryAllocator<byte> _allocator;
-		readonly object _lockObject;
+		readonly Lock _lockObject;
 
 		IRefCountedHandle<Memory<byte>> _bufferHandle;
 		Memory<byte> _buffer;
@@ -48,7 +49,7 @@ namespace EpicGames.Horde.Storage.Bundles.V2
 		/// <summary>
 		/// Constructor
 		/// </summary>
-		public PacketWriter(BundleHandle bundleHandle, PacketHandle packetHandle, IMemoryAllocator<byte> allocator, object lockObject)
+		public PacketWriter(BundleHandle bundleHandle, PacketHandle packetHandle, IMemoryAllocator<byte> allocator, Lock lockObject)
 		{
 			_bundleHandle = bundleHandle;
 			_packetHandle = packetHandle;
@@ -157,8 +158,6 @@ namespace EpicGames.Horde.Storage.Bundles.V2
 				PacketExport export = new(_buffer.Slice(offset, length));
 				PacketExportHeader header = export.GetHeader();
 
-				BlobType type = _types[header.TypeIdx];
-
 				IBlobRef[] imports = new IBlobRef[header.Imports.Length];
 				for (int idx = 0; idx < header.Imports.Length; idx++)
 				{
diff --git a/Engine/Source/Programs/Shared/EpicGames.Horde/Storage/ContentChunker.cs b/Engine/Source/Programs/Shared/EpicGames.Horde/Storage/ContentChunker.cs
--- a/Engine/Source/Programs/Shared/EpicGames.Horde/Storage/ContentChunker.cs
+++ b/Engine/Source/Programs/Shared/EpicGames.Horde/Storage/ContentChunker.cs
@@ -312,7 +312,7 @@ namespace EpicGames.Horde.Storage
 
 		class RingBuffer
 		{
-			readonly object _lockObject = new();
+			readonly Lock _lockObject = new();
 
 			int _head;
 			int _size;
diff --git a/Engine/Source/Programs/Shared/EpicGames.Horde/Storage/Nodes/DirectoryNode.Update.cs b/Engine/Source/Programs/Shared/EpicGames.Horde/Storage/Nodes/DirectoryNode.Update.cs
--- a/Engine/Source/Programs/Shared/EpicGames.Horde/Storage/Nodes/DirectoryNode.Update.cs
+++ b/Engine/Source/Programs/Shared/EpicGames.Horde/Storage/Nodes/DirectoryNode.Update.cs
@@ -40,7 +40,7 @@ namespace EpicGames.Horde.Storage.Nodes
 	/// </summary>
 	class UpdateStats : IUpdateStats
 	{
-		readonly object _lockObject = new();
+		readonly Lock _lockObject = new();
 		readonly Stopwatch _timer = Stopwatch.StartNew();
 		readonly IProgress<IUpdateStats>? _progress;
 		long _lastTotalSize;
diff --git a/Engine/Source/Programs/Shared/EpicGames.Horde/Storage/StorageBackendCache.cs b/Engine/Source/Programs/Shared/EpicGames.Horde/Storage/StorageBackendCache.cs
--- a/Engine/Source/Programs/Shared/EpicGames.Horde/Storage/StorageBackendCache.cs
+++ b/Engine/Source/Programs/Shared/EpicGames.Horde/Storage/StorageBackendCache.cs
@@ -169,7 +169,7 @@ namespace EpicGames.Horde.Storage
 				=> _inner.UpdateMetadataAsync(request, cancellationToken);
 		}
 
-		object LockObject => _items;
+		private readonly Lock _lockObject = new();
 
 		readonly MemoryMappedFileCache _memoryMappedFileCache;
 		readonly FileObjectStore _objectStore;
@@ -281,7 +281,7 @@ namespace EpicGames.Horde.Storage
 			for (; ; )
 			{
 				PendingItem? pendingItem;
-				lock (LockObject)
+				lock (_lockObject)
 				{
 					Item? item;
 					if (_keyToItem.TryGetValue(key, out item))
@@ -320,7 +320,7 @@ namespace EpicGames.Horde.Storage
 			Interlocked.Add(ref _fetchTimeTicks, openFinishTicks - openStartTicks);
 
 			long cleanStartTicks = openFinishTicks;
-			lock (LockObject)
+			lock (_lockObject)
 			{
 				_size += totalLength;
 
@@ -352,7 +352,7 @@ namespace EpicGames.Horde.Storage
 
 			long writeStartTicks = cleanFinishTicks;
 			await _objectStore.WriteAsync(key, stream, cancellationToken: cancellationToken);
-			lock (LockObject)
+			lock (_lockObject)
 			{
 				Item item = new(key, totalLength);
 				_items.AddFirst(item.ListNode);
diff --git a/Engine/Source/Programs/Shared/EpicGames.Horde/Storage/Workspace.cs b/Engine/Source/Programs/Shared/EpicGames.Horde/Storage/Workspace.cs
--- a/Engine/Source/Programs/Shared/EpicGames.Horde/Storage/Workspace.cs
+++ b/Engine/Source/Programs/Shared/EpicGames.Horde/Storage/Workspace.cs
@@ -1990,7 +1990,7 @@ namespace EpicGames.Horde.Storage
 		{
 			public bool AnyChanges { get; private set; }
 
-			readonly object _lockObject = new();
+			readonly Lock _lockObject = new();
 
 			public override void DirectoryAdded(string path, DirectoryState parentDir, DirectoryInfo subDirInfo)
 				=> WriteLine($"{FormatPath(path)}/ (added)", ConsoleColor.Green);
diff --git a/Engine/Source/Programs/Shared/EpicGames.MsBuild/MsBuildSupport.cs b/Engine/Source/Programs/Shared/EpicGames.MsBuild/MsBuildSupport.cs
--- a/Engine/Source/Programs/Shared/EpicGames.MsBuild/MsBuildSupport.cs
+++ b/Engine/Source/Programs/Shared/EpicGames.MsBuild/MsBuildSupport.cs
@@ -3,6 +3,7 @@
 using System;
 using System.Collections.Generic;
 using System.Diagnostics;
+using System.Threading;
 using EpicGames.Core;
 using Microsoft.Build.Locator;
 
@@ -14,7 +15,7 @@ namespace EpicGames.MsBuild;
 public static class MsBuildSupport
 {
 	private static bool s_hasRegisteredMsBuildPath = false;
-	private static readonly object s_lock = new();
+	private static readonly Lock s_lock = new();
 
 	/// <summary>
 	/// Register our bundled dotnet installation to be used by Microsoft.Build
diff --git a/Engine/Source/Programs/Shared/EpicGames.OIDC/OidcAuthState.cs b/Engine/Source/Programs/Shared/EpicGames.OIDC/OidcAuthState.cs
--- a/Engine/Source/Programs/Shared/EpicGames.OIDC/OidcAuthState.cs
+++ b/Engine/Source/Programs/Shared/EpicGames.OIDC/OidcAuthState.cs
@@ -182,7 +182,7 @@ public sealed class OidcAuthState(
 		return authState?.TokenInfo?.AccessToken;
 	}
 
-	private readonly object _lockObject = new();
+	private readonly Lock _lockObject = new();
 	private readonly IOidcTokenManager _oidcTokenManager = oidcTokenManager;
 	private readonly IStaticAccessTokenProvider _staticAccessTokenProvider = staticAccessTokenProvider;
 	private readonly IClock _clock = clock;
diff --git a/Engine/Source/Programs/Shared/EpicGames.OIDC/OidcTokenManager.cs b/Engine/Source/Programs/Shared/EpicGames.OIDC/OidcTokenManager.cs
--- a/Engine/Source/Programs/Shared/EpicGames.OIDC/OidcTokenManager.cs
+++ b/Engine/Source/Programs/Shared/EpicGames.OIDC/OidcTokenManager.cs
@@ -46,7 +46,7 @@ namespace EpicGames.OIDC
 
 	public class OidcTokenManager : IOidcTokenManager
 	{
-		private readonly object _lockObject = new();
+		private readonly Lock _lockObject = new();
 		private readonly Dictionary<string, OidcTokenClient> _tokenClients = new(StringComparer.OrdinalIgnoreCase);
 		private readonly ITokenStore _tokenStore;
 
diff --git a/Engine/Source/Programs/Shared/EpicGames.Serialization/CbConverter.cs b/Engine/Source/Programs/Shared/EpicGames.Serialization/CbConverter.cs
--- a/Engine/Source/Programs/Shared/EpicGames.Serialization/CbConverter.cs
+++ b/Engine/Source/Programs/Shared/EpicGames.Serialization/CbConverter.cs
@@ -4,6 +4,7 @@ using System;
 using System.Collections.Generic;
 using System.Linq.Expressions;
 using System.Reflection;
+using System.Threading;
 using EpicGames.Core;
 using EpicGames.Serialization.Converters;
 
@@ -201,7 +202,7 @@ namespace EpicGames.Serialization
 		/// <summary>
 		/// Object used for locking access to shared objects
 		/// </summary>
-		static readonly object s_lockObject = new();
+		static readonly Lock s_lockObject = new();
 
 		/// <summary>
 		/// Cache of property to converter
diff --git a/Engine/Source/Programs/Shared/EpicGames.Tracing/UnrealInsights/Writer.cs b/Engine/Source/Programs/Shared/EpicGames.Tracing/UnrealInsights/Writer.cs
--- a/Engine/Source/Programs/Shared/EpicGames.Tracing/UnrealInsights/Writer.cs
+++ b/Engine/Source/Programs/Shared/EpicGames.Tracing/UnrealInsights/Writer.cs
@@ -4,6 +4,7 @@ using System;
 using System.Collections.Generic;
 using System.IO;
 using System.Linq;
+using System.Threading;
 using EpicGames.Tracing.UnrealInsights.Events;
 
 namespace EpicGames.Tracing.UnrealInsights
@@ -15,7 +16,7 @@ namespace EpicGames.Tracing.UnrealInsights
 		readonly Dictionary<ushort, List<ITraceEvent>> _threadToEvents = [];
 
 		private ushort _uidCounter = PredefinedEventUid.WellKnownNum; // IDs below 16 are reserved for well-known events
-		private readonly object _lock = new();
+		private readonly Lock _lock = new();
 
 		public UnrealInsightsWriter()
 		{
diff --git a/Engine/Source/Programs/Shared/EpicGames.UBA/Impl/LoggerImpl.cs b/Engine/Source/Programs/Shared/EpicGames.UBA/Impl/LoggerImpl.cs
--- a/Engine/Source/Programs/Shared/EpicGames.UBA/Impl/LoggerImpl.cs
+++ b/Engine/Source/Programs/Shared/EpicGames.UBA/Impl/LoggerImpl.cs
@@ -22,7 +22,7 @@ namespace EpicGames.UBA.Impl
 		readonly BeginScopeCallback _beginScopeCallbackDelegate;
 		readonly EndScopeCallback _endScopeCallbackDelegate;
 		readonly LogCallback _logCallbackDelegate;
-		readonly object _lock = new();
+		readonly Lock _lock = new();
 
 		#region DllImport
 		[DllImport("UbaHost", CharSet = CharSet.Auto)]
@@ -82,10 +82,10 @@ namespace EpicGames.UBA.Impl
 		public nint GetHandle() => _handle;
 
 		/// <inheritdoc/>
-		public void BeginScope() => Monitor.Enter(_lock);
+		public void BeginScope() => _lock.Enter();
 
 		/// <inheritdoc/>
-		public void EndScope() => Monitor.Exit(_lock);
+		public void EndScope() => _lock.Exit();
 
 		/// <inheritdoc/>
 		public void Log(LogEntryType type, string message)
diff --git a/Engine/Source/Programs/UnrealBuildTool/.globalconfig b/Engine/Source/Programs/UnrealBuildTool/.globalconfig
--- a/Engine/Source/Programs/UnrealBuildTool/.globalconfig
+++ b/Engine/Source/Programs/UnrealBuildTool/.globalconfig
@@ -157,6 +157,9 @@ dotnet_diagnostic.IDE0005.severity = warning
 # IDE0270: Null check can be simplified (if null check)
 dotnet_diagnostic.IDE0270.severity = silent
 
+# IDE0330: Prefer 'System.Threading.Lock'
+dotnet_diagnostic.IDE0330.severity = warning
+
 # IDE0090: 'new' expression can be simplified
 dotnet_diagnostic.IDE0090.severity = silent
 
diff --git a/Engine/Source/Programs/UnrealBuildTool/Actions/ActionHistory.cs b/Engine/Source/Programs/UnrealBuildTool/Actions/ActionHistory.cs
--- a/Engine/Source/Programs/UnrealBuildTool/Actions/ActionHistory.cs
+++ b/Engine/Source/Programs/UnrealBuildTool/Actions/ActionHistory.cs
@@ -7,6 +7,7 @@ using System.Diagnostics;
 using System.IO;
 using System.Linq;
 using System.Text;
+using System.Threading;
 using EpicGames.Core;
 using Microsoft.Extensions.Logging;
 using UnrealBuildBase;
@@ -220,7 +221,7 @@ namespace UnrealBuildTool
 		/// <summary>
 		/// Used to ensure exclusive access to the layers list
 		/// </summary>
-		readonly object LockObject = new object();
+		readonly Lock LockObject = new();
 
 		/// <summary>
 		/// Map of filename to layer
@@ -360,7 +361,7 @@ namespace UnrealBuildTool
 		/// <summary>
 		/// The lock object for this history
 		/// </summary>
-		readonly object LockObject = new object();
+		readonly Lock LockObject = new();
 
 		/// <summary>
 		/// List of partitions
diff --git a/Engine/Source/Programs/UnrealBuildTool/Configuration/Rules/RulesCompiler.cs b/Engine/Source/Programs/UnrealBuildTool/Configuration/Rules/RulesCompiler.cs
--- a/Engine/Source/Programs/UnrealBuildTool/Configuration/Rules/RulesCompiler.cs
+++ b/Engine/Source/Programs/UnrealBuildTool/Configuration/Rules/RulesCompiler.cs
@@ -5,6 +5,7 @@ using System.Collections.Concurrent;
 using System.Collections.Generic;
 using System.IO;
 using System.Linq;
+using System.Threading;
 using EpicGames.Core;
 using Microsoft.Extensions.Logging;
 using OpenTracing.Util;
@@ -45,7 +46,7 @@ namespace UnrealBuildTool
 		/// <summary>
 		/// Lock for creating engine rule assembly
 		/// </summary>
-		private static readonly object EngineRulesAssemblyLock = new();
+		private static readonly Lock EngineRulesAssemblyLock = new();
 
 		/// <summary>
 		/// The cached rules assembly for engine modules and targets.
diff --git a/Engine/Source/Programs/UnrealBuildTool/Executors/UnrealBuildAccelerator/UBAAgentCoordinatorHorde.cs b/Engine/Source/Programs/UnrealBuildTool/Executors/UnrealBuildAccelerator/UBAAgentCoordinatorHorde.cs
--- a/Engine/Source/Programs/UnrealBuildTool/Executors/UnrealBuildAccelerator/UBAAgentCoordinatorHorde.cs
+++ b/Engine/Source/Programs/UnrealBuildTool/Executors/UnrealBuildAccelerator/UBAAgentCoordinatorHorde.cs
@@ -474,7 +474,7 @@ namespace UnrealBuildTool
 		readonly ManualResetEventSlim _initDoneEvent = new();
 
 		bool _initDone;
-		readonly object _startLock = new();
+		readonly Lock _startLock = new();
 		IUbaAgentCoordinatorScheduler? _scheduler;
 
 		ServiceProvider? _serviceProvider;
diff --git a/Engine/Source/Programs/UnrealBuildTool/Executors/UnrealBuildAccelerator/UBAExecutor.cs b/Engine/Source/Programs/UnrealBuildTool/Executors/UnrealBuildAccelerator/UBAExecutor.cs
--- a/Engine/Source/Programs/UnrealBuildTool/Executors/UnrealBuildAccelerator/UBAExecutor.cs
+++ b/Engine/Source/Programs/UnrealBuildTool/Executors/UnrealBuildAccelerator/UBAExecutor.cs
@@ -40,7 +40,7 @@ namespace UnrealBuildTool
 #pragma warning restore CA2213
 		string? _sessionHyperlink;
 
-		readonly object _sessionLock = new();
+		readonly Lock _sessionLock = new();
 
 		readonly Dictionary<UnrealBuildAcceleratorCacheConfig, ICacheClient> _cacheClients = [];
 		readonly List<IUBAAgentCoordinator> _agentCoordinators = new();
@@ -806,7 +806,7 @@ namespace UnrealBuildTool
 			internal IScheduler? _scheduler;
 			internal ISessionServer? _session;
 			internal IServer? _server;
-			private readonly object _lockObject = new();
+			private readonly Lock _lockObject = new();
 
 			public void Reset()
 			{
diff --git a/Engine/Source/Programs/UnrealBuildTool/Modes/AggregateParsedTimingInfo.cs b/Engine/Source/Programs/UnrealBuildTool/Modes/AggregateParsedTimingInfo.cs
--- a/Engine/Source/Programs/UnrealBuildTool/Modes/AggregateParsedTimingInfo.cs
+++ b/Engine/Source/Programs/UnrealBuildTool/Modes/AggregateParsedTimingInfo.cs
@@ -29,7 +29,7 @@ namespace UnrealBuildTool
 		private readonly ConcurrentBag<TimingData> AggregateIncludes = new ConcurrentBag<TimingData>();
 		private readonly ConcurrentBag<TimingData> AggregateClasses = new ConcurrentBag<TimingData>();
 		private readonly ConcurrentBag<TimingData> AggregateFunctions = new ConcurrentBag<TimingData>();
-		private readonly object AddFileLock = new object();
+		private readonly Lock AddFileLock = new();
 
 		public async Task<int> ExecuteAsync(CommandLineArguments Arguments, ILogger Logger)
 		{
diff --git a/Engine/Source/Programs/UnrealBuildTool/Platform/Android/AndroidToolChain.cs b/Engine/Source/Programs/UnrealBuildTool/Platform/Android/AndroidToolChain.cs
--- a/Engine/Source/Programs/UnrealBuildTool/Platform/Android/AndroidToolChain.cs
+++ b/Engine/Source/Programs/UnrealBuildTool/Platform/Android/AndroidToolChain.cs
@@ -7,6 +7,7 @@ using System.Diagnostics;
 using System.IO;
 using System.Linq;
 using System.Text.RegularExpressions;
+using System.Threading;
 using EpicGames.Core;
 using Microsoft.Extensions.Logging;
 using UnrealBuildBase;
@@ -1060,7 +1061,7 @@ namespace UnrealBuildTool
 		}
 
 		// cache the location of NDK tools
-		protected static object CtorLock = new();
+		protected static Lock CtorLock = new();
 		protected static string? ClangPath;
 		protected static string ToolchainParamsArm64 = "";
 		protected static string ToolchainParamsx64 = "";
diff --git a/Engine/Source/Programs/UnrealBuildTool/Platform/IOS/UEDeployIOS.cs b/Engine/Source/Programs/UnrealBuildTool/Platform/IOS/UEDeployIOS.cs
--- a/Engine/Source/Programs/UnrealBuildTool/Platform/IOS/UEDeployIOS.cs
+++ b/Engine/Source/Programs/UnrealBuildTool/Platform/IOS/UEDeployIOS.cs
@@ -5,6 +5,7 @@ using System.Collections.Generic;
 using System.IO;
 using System.Linq;
 using System.Text;
+using System.Threading;
 using System.Xml.Linq;
 using EpicGames.Core;
 using Microsoft.Extensions.Logging;
@@ -465,7 +466,7 @@ namespace UnrealBuildTool
 			File.WriteAllText(PlistFile.FullName, Text.ToString());
 		}
 
-		private static readonly object GenerateIOSPListLock = new();
+		private static readonly Lock GenerateIOSPListLock = new();
 
 		public static bool GenerateIOSPList(FileReference? ProjectFile, UnrealTargetConfiguration Config, string ProjectDirectory, bool bIsUnrealGame, string GameName, bool bIsClient, string ProjectName, string InEngineDir, string AppDirectory, UnrealPluginLanguage? UPL, string? BundleID, bool bBuildAsFramework, ILogger Logger)
 		{
diff --git a/Engine/Source/Programs/UnrealBuildTool/Platform/Windows/MicrosoftPlatformSDK.cs b/Engine/Source/Programs/UnrealBuildTool/Platform/Windows/MicrosoftPlatformSDK.cs
--- a/Engine/Source/Programs/UnrealBuildTool/Platform/Windows/MicrosoftPlatformSDK.cs
+++ b/Engine/Source/Programs/UnrealBuildTool/Platform/Windows/MicrosoftPlatformSDK.cs
@@ -11,6 +11,7 @@ using System.Linq;
 using System.Runtime.InteropServices;
 using System.Runtime.Versioning;
 using System.Text.RegularExpressions;
+using System.Threading;
 using EpicGames.Core;
 using Microsoft.Extensions.Logging;
 using Microsoft.VisualStudio.Setup.Configuration;
@@ -107,7 +108,7 @@ namespace UnrealBuildTool
 		/// </summary>
 		private static IReadOnlyList<VisualStudioInstallation>? CachedVisualStudioInstallations;
 
-		private static readonly object CachedVisualStudioInstallationsLock = new();
+		private static readonly Lock CachedVisualStudioInstallationsLock = new();
 
 		/// <summary>
 		/// Cache of DIA SDK installation directories
diff --git a/Engine/Source/Programs/UnrealBuildTool/Platform/Windows/VCPreludeCache.cs b/Engine/Source/Programs/UnrealBuildTool/Platform/Windows/VCPreludeCache.cs
--- a/Engine/Source/Programs/UnrealBuildTool/Platform/Windows/VCPreludeCache.cs
+++ b/Engine/Source/Programs/UnrealBuildTool/Platform/Windows/VCPreludeCache.cs
@@ -5,6 +5,7 @@ using System.Collections.Generic;
 using System.Diagnostics;
 using System.IO;
 using System.Linq;
+using System.Threading;
 using EpicGames.Core;
 using UnrealBuildBase;
 
@@ -31,7 +32,7 @@ namespace UnrealBuildTool
 		/// <summary>
 		/// Object used for ensuring only one thread can create items in the cache at once
 		/// </summary>
-		static readonly object LockObject = new object();
+		static readonly Lock LockObject = new();
 
 		/// <summary>
 		/// File containing a list of macros to be used to generate a prelude file
diff --git a/Engine/Source/Programs/UnrealBuildTool/Platform/Windows/VCToolChain.cs b/Engine/Source/Programs/UnrealBuildTool/Platform/Windows/VCToolChain.cs
--- a/Engine/Source/Programs/UnrealBuildTool/Platform/Windows/VCToolChain.cs
+++ b/Engine/Source/Programs/UnrealBuildTool/Platform/Windows/VCToolChain.cs
@@ -7,6 +7,7 @@ using System.IO;
 using System.Linq;
 using System.Runtime.Versioning;
 using System.Text;
+using System.Threading;
 using System.Xml.Linq;
 using EpicGames.Core;
 using Microsoft.CodeAnalysis;
@@ -2680,7 +2681,7 @@ namespace UnrealBuildTool
 			OptionalUpdateSxSManifest(Target, MakefileBuilder, ExeDir);
 		}
 
-		private static readonly ConcurrentDictionary<FileReference, object> ManifestLocks = [];
+		private static readonly ConcurrentDictionary<FileReference, Lock> ManifestLocks = [];
 
 		protected static void OptionalUpdateSxSManifest(ReadOnlyTargetRules Target, TargetMakefileBuilder MakefileBuilder, DirectoryReference ExeDir)
 		{
@@ -2703,7 +2704,7 @@ namespace UnrealBuildTool
 				FileReference IntermediateManifest = GetSxSIntermediateManifest(Target, ExeDir);
 				FileReference TemplateManifest = FileReference.Combine(Unreal.EngineDirectory, "Build", "Windows", "Resources", AssemblyName, AssemblyName + ".manifest");
 
-				object Lock = ManifestLocks.GetOrAdd(IntermediateManifest, _ => new object());
+				Lock Lock = ManifestLocks.GetOrAdd(IntermediateManifest, _ => new());
 				lock (Lock)
 				{
 					XDocument XmlManifest;
diff --git a/Engine/Source/Programs/UnrealBuildTool/System/DataDrivenPlatformInfo.cs b/Engine/Source/Programs/UnrealBuildTool/System/DataDrivenPlatformInfo.cs
--- a/Engine/Source/Programs/UnrealBuildTool/System/DataDrivenPlatformInfo.cs
+++ b/Engine/Source/Programs/UnrealBuildTool/System/DataDrivenPlatformInfo.cs
@@ -4,6 +4,7 @@ using System;
 using System.Collections.Generic;
 using System.IO;
 using System.Linq;
+using System.Threading;
 using EpicGames.Core;
 using UnrealBuildBase;
 
@@ -63,7 +64,7 @@ namespace UnrealBuildTool
 #pragma warning restore CA1034
 
 		static IReadOnlyDictionary<string, ConfigDataDrivenPlatformInfo>? PlatformInfos = null;
-		static readonly object LockObject = new object();
+		static readonly Lock LockObject = new();
 
 		/// <summary>
 		/// Return all data driven infos found
diff --git a/Engine/Source/Programs/UnrealBuildTool/System/ProgressWriter.cs b/Engine/Source/Programs/UnrealBuildTool/System/ProgressWriter.cs
--- a/Engine/Source/Programs/UnrealBuildTool/System/ProgressWriter.cs
+++ b/Engine/Source/Programs/UnrealBuildTool/System/ProgressWriter.cs
@@ -1,6 +1,7 @@
 // Copyright Epic Games, Inc. All Rights Reserved.
 
 using System;
+using System.Threading;
 using EpicGames.Core;
 using Microsoft.Extensions.Logging;
 
@@ -36,7 +37,7 @@ namespace UnrealBuildTool
 		/// </summary>
 		int CurrentProgressValue = -1;
 
-		private readonly object LockObject = new();
+		private readonly Lock LockObject = new();
 
 		/// <summary>
 		/// Constructor
diff --git a/Engine/Source/Programs/UnrealBuildTool/System/ScopedTimer.cs b/Engine/Source/Programs/UnrealBuildTool/System/ScopedTimer.cs
--- a/Engine/Source/Programs/UnrealBuildTool/System/ScopedTimer.cs
+++ b/Engine/Source/Programs/UnrealBuildTool/System/ScopedTimer.cs
@@ -2,6 +2,7 @@
 
 using System;
 using System.Text;
+using System.Threading;
 using Microsoft.Extensions.Logging;
 
 namespace UnrealBuildTool
@@ -19,7 +20,7 @@ namespace UnrealBuildTool
 		// TODO: This isn't true nesting support, as it assumes that all new ScopedTimers are within the previous scope.
 		// This doesn't account for totally unrelated timers, so this should probably use some context-local storage to deal with this.
 		static int Indent = 0;
-		static readonly object IndentLock = new object();
+		static readonly Lock IndentLock = new();
 
 		/// <summary>
 		/// Constructor
@@ -27,7 +28,7 @@ namespace UnrealBuildTool
 		/// <param name="Name">Name of the block being measured</param>
 		/// <param name="Logger">Logger for output</param>
 		/// <param name="InLevel">Verbosity for output messages</param>
-		/// <param name="bIncreaseIndent">Whether gobal indent should be increased or not; set to false when running a scope in parallel. Message will still be printed indented relative to parent scope.</param>
+		/// <param name="bIncreaseIndent">Whether global indent should be increased or not; set to false when running a scope in parallel. Message will still be printed indented relative to parent scope.</param>
 		public ScopedTimer(string Name, ILogger Logger, LogLevel InLevel = LogLevel.Debug, bool bIncreaseIndent = true)
 		{
 			this.Name = Name;
diff --git a/Engine/Source/Programs/UnrealBuildTool/System/SourceFileWorkingSet.cs b/Engine/Source/Programs/UnrealBuildTool/System/SourceFileWorkingSet.cs
--- a/Engine/Source/Programs/UnrealBuildTool/System/SourceFileWorkingSet.cs
+++ b/Engine/Source/Programs/UnrealBuildTool/System/SourceFileWorkingSet.cs
@@ -6,6 +6,7 @@ using System.Diagnostics;
 using System.Diagnostics.CodeAnalysis;
 using System.IO;
 using System.Linq;
+using System.Threading;
 using EpicGames.Core;
 using Microsoft.Extensions.Logging;
 using UnrealBuildBase;
@@ -92,7 +93,7 @@ namespace UnrealBuildTool
 	{
 		readonly DirectoryReference RootDir;
 		Process? BackgroundProcess;
-		readonly object BackgroundProcessDisposeLock = new();
+		readonly Lock BackgroundProcessDisposeLock = new();
 		readonly Lazy<bool> BackgroundProcessDone;
 		readonly HashSet<FileReference> Files;
 		readonly List<DirectoryReference> Directories;
diff --git a/Engine/Source/Programs/UnrealBuildTool/System/TargetMakefile.cs b/Engine/Source/Programs/UnrealBuildTool/System/TargetMakefile.cs
--- a/Engine/Source/Programs/UnrealBuildTool/System/TargetMakefile.cs
+++ b/Engine/Source/Programs/UnrealBuildTool/System/TargetMakefile.cs
@@ -5,6 +5,7 @@ using System.Collections.Generic;
 using System.IO;
 using System.Linq;
 using System.Reflection;
+using System.Threading;
 using System.Threading.Tasks;
 using EpicGames.Core;
 using EpicGames.UHT.Utils;
@@ -987,7 +988,7 @@ namespace UnrealBuildTool
 				});
 
 				bool isValid2 = true;
-				object isValidLock2 = new();
+				Lock isValidLock2 = new();
 				string reasonNotLoaded2 = "";
 
 				// First collect file timestamps by caching directories. This is much faster than letting each file do a GetFileAttribute
diff --git a/Engine/Source/Programs/UnrealBuildTool/System/UniqueStringRegistry.cs b/Engine/Source/Programs/UnrealBuildTool/System/UniqueStringRegistry.cs
--- a/Engine/Source/Programs/UnrealBuildTool/System/UniqueStringRegistry.cs
+++ b/Engine/Source/Programs/UnrealBuildTool/System/UniqueStringRegistry.cs
@@ -3,6 +3,7 @@
 using System;
 using System.Collections.Generic;
 using System.Linq;
+using System.Threading;
 
 namespace UnrealBuildTool
 {
@@ -12,7 +13,7 @@ namespace UnrealBuildTool
 	public class UniqueStringRegistry
 	{
 		// protect the InstanceMap
-		private readonly object LockObject = new object();
+		private readonly Lock LockObject = new();
 
 		// holds a mapping of string name to single instance
 		private Dictionary<string, int> StringToInstanceMap = new Dictionary<string, int>(StringComparer.OrdinalIgnoreCase);
