Horde: Fix buffer overread crash in AgentManagedProcess compute channel

AppendToBuffer already advances _bufferLength, but GetString added lineLen
again — reading past valid data. Any line >256 bytes spanning two receive
chunks triggers ArgumentOutOfRangeException (256 is the allocation slack).

diff --git a/Engine/Source/Programs/Shared/EpicGames.Horde/Compute/AgentManagedProcess.cs b/Engine/Source/Programs/Shared/EpicGames.Horde/Compute/AgentManagedProcess.cs
--- a/Engine/Source/Programs/Shared/EpicGames.Horde/Compute/AgentManagedProcess.cs
+++ b/Engine/Source/Programs/Shared/EpicGames.Horde/Compute/AgentManagedProcess.cs
@@ -113,7 +113,7 @@ namespace EpicGames.Horde.Compute
 				else
 				{
 					AppendToBuffer(data.Slice(0, lineLen));
-					str = Encoding.UTF8.GetString(_buffer.AsSpan(0, _bufferLength + lineLen));
+					str = Encoding.UTF8.GetString(_buffer.AsSpan(0, _bufferLength));
 					_bufferLength = 0;
 				}
 
