Updating the UHT implementation of VNI modules and verse fields.

#rnx

#Synced-CL 52506245

diff --git a/Engine/Source/Programs/Shared/EpicGames.UHT/Exporters/CodeGen/UhtHeaderCodeGeneratorCppFile.cs b/Engine/Source/Programs/Shared/EpicGames.UHT/Exporters/CodeGen/UhtHeaderCodeGeneratorCppFile.cs
--- a/Engine/Source/Programs/Shared/EpicGames.UHT/Exporters/CodeGen/UhtHeaderCodeGeneratorCppFile.cs
+++ b/Engine/Source/Programs/Shared/EpicGames.UHT/Exporters/CodeGen/UhtHeaderCodeGeneratorCppFile.cs
@@ -332,18 +332,6 @@ namespace EpicGames.UHT.Exporters.CodeGen
 					#define UHT_STRUCT_BASE(INIT) UE::CodeGen::ConstInit::TCompiledInObjectPtr<const FStructBaseChain>(UE::Private::AsStructBaseChain(INIT))
 					""");
 
-				// We need to emit the internal UClass needed for VMODULEs
-				foreach (UhtField field in HeaderFile.References.ExportTypes)
-				{
-					if (field is UhtClass classObj)
-					{
-						if (!classObj.ClassFlags.HasAnyFlags(EClassFlags.Intrinsic))
-						{
-							AppendClassDeclaration(builder, classObj);
-						}
-					}
-				}
-
 				Dictionary<UhtPackage, UhtRegistrations> packageRegistrations = [];
 				foreach (UhtField field in HeaderFile.References.ExportTypes)
 				{
@@ -2632,23 +2620,6 @@ namespace EpicGames.UHT.Exporters.CodeGen
 			return builder;
 		}
 
-		private StringBuilder AppendClassDeclaration(StringBuilder builder, UhtClass classObj)
-		{
-
-			// If we are a VMODULE, we need to emit the declaration
-			if (classObj.ClassType == UhtClassType.VModule)
-			{
-				classObj.Namespace.AppendMultipleLines(builder, (builder) =>
-				{
-					builder.Append($"class {classObj.SourceName} : public UObject\r\n");
-					builder.Append("{\r\n");
-					builder.Append('\t').AppendMacroName(this, classObj, GeneratedBodyMacroSuffix).Append("\r\n");
-					builder.Append("};\r\n");
-				});
-			}
-			return builder;
-		}
-
 		private StringBuilder AppendClass(StringBuilder builder, UhtClass classObj, UhtUsedDefineScopes<UhtFunction> functions)
 		{
 			// Add the auto getters/setters
diff --git a/Engine/Source/Programs/Shared/EpicGames.UHT/Parsers/UhtClassParser.cs b/Engine/Source/Programs/Shared/EpicGames.UHT/Parsers/UhtClassParser.cs
--- a/Engine/Source/Programs/Shared/EpicGames.UHT/Parsers/UhtClassParser.cs
+++ b/Engine/Source/Programs/Shared/EpicGames.UHT/Parsers/UhtClassParser.cs
@@ -25,6 +25,14 @@ namespace EpicGames.UHT.Parsers
 			return ParseUClass(topScope, ref token);
 		}
 
+		[UhtKeyword(Extends = UhtTableNames.Global)]
+		[SuppressMessage("CodeQuality", "IDE0051:Remove unused private members", Justification = "Attribute accessed method")]
+		[SuppressMessage("Style", "IDE0060:Remove unused parameter", Justification = "Attribute accessed method")]
+		private static UhtParseResult VMODULEKeyword(UhtParsingScope topScope, UhtParsingScope actionScope, ref UhtToken token)
+		{
+			return ParseUClass(topScope, ref token);
+		}
+
 		[UhtKeyword(Extends = UhtTableNames.Class)]
 		[UhtKeyword(Extends = UhtTableNames.Class, Keyword = "GENERATED_BODY")]
 		[SuppressMessage("CodeQuality", "IDE0051:Remove unused private members", Justification = "Attribute accessed method")]
@@ -57,6 +65,11 @@ namespace EpicGames.UHT.Parsers
 		private static UhtParseResult ParseUClass(UhtParsingScope parentScope, ref UhtToken token)
 		{
 			UhtClass classObj = new(parentScope.HeaderFile, parentScope.HeaderParser.GetNamespace(), parentScope.ScopeType, token.InputLine);
+			if (token.IsIdentifier("VMODULE"))
+			{
+				classObj.ClassType = UhtClassType.VModule;
+			}
+
 			{
 				using UhtParsingScope topScope = new(parentScope, classObj, parentScope.Session.GetKeywordTable(UhtTableNames.Class), UhtAccessSpecifier.Private);
 				const string ScopeName = "class";
@@ -69,6 +82,10 @@ namespace EpicGames.UHT.Parsers
 					specifiers.ParseSpecifiers();
 					classObj.PrologLineNumber = topScope.TokenReader.InputLine;
 					classObj.ClassFlags |= EClassFlags.Native;
+					if (classObj.ClassType == UhtClassType.VModule)
+					{
+						classObj.ClassFlags |= EClassFlags.Hidden;
+					}
 
 					topScope.AddFormattedCommentsAsTooltipMetaData();
 
diff --git a/Engine/Source/Programs/Shared/EpicGames.UHT/Parsers/UhtNamespaceParser.cs b/Engine/Source/Programs/Shared/EpicGames.UHT/Parsers/UhtNamespaceParser.cs
--- a/Engine/Source/Programs/Shared/EpicGames.UHT/Parsers/UhtNamespaceParser.cs
+++ b/Engine/Source/Programs/Shared/EpicGames.UHT/Parsers/UhtNamespaceParser.cs
@@ -24,14 +24,6 @@ namespace EpicGames.UHT.Parsers
 			return topScope.Module.Module.AllowUETypesInNamespaces ? ParseNamespace(topScope) : UhtParseResult.Unhandled;
 		}
 
-		[UhtKeyword(Extends = UhtTableNames.Global)]
-		[SuppressMessage("CodeQuality", "IDE0051:Remove unused private members", Justification = "Attribute accessed method")]
-		[SuppressMessage("Style", "IDE0060:Remove unused parameter", Justification = "Attribute accessed method")]
-		private static UhtParseResult VMODULEKeyword(UhtParsingScope topScope, UhtParsingScope actionScope, ref UhtToken token)
-		{
-			return ParseVModule(topScope, ref token);
-		}
-
 		private static int ParseNamespaceIdentifiers(UhtParsingScope topScope)
 		{
 			// Parse the namespace identifiers
@@ -87,75 +79,5 @@ namespace EpicGames.UHT.Parsers
 			topScope.HeaderParser.PopNamespaces(namespaceCount);
 			return UhtParseResult.Handled;
 		}
-
-		private static UhtParseResult ParseVModule(UhtParsingScope parentScope, ref UhtToken token)
-		{
-			// This function is equivalent to FSolClassGenerator::Prepare and they should be kept in sync
-
-			UhtClass classObj = new(parentScope.HeaderFile, parentScope.HeaderParser.GetNamespace(), parentScope.ScopeType, token.InputLine);
-			classObj.ClassType = UhtClassType.VModule;
-
-			using UhtParsingScope topScope = new(parentScope, classObj, parentScope.Session.GetKeywordTable(UhtTableNames.VModule), UhtAccessSpecifier.Public);
-			const string ScopeName = "VMODULE";
-			using UhtMessageContext tokenContext = new(ScopeName);
-
-			// Parse the specifiers
-			UhtSpecifierContext specifierContext = new(topScope, topScope.TokenReader, classObj.MetaData);
-			UhtSpecifierParser specifiers = UhtSpecifierParser.GetThreadInstance(specifierContext, ScopeName, parentScope.Session.GetSpecifierTable(UhtTableNames.Class));
-			specifiers.ParseSpecifiers();
-			classObj.GeneratedBodyLineNumber = classObj.LineNumber;
-			classObj.GeneratedBodyAccessSpecifier = UhtAccessSpecifier.Public;
-			classObj.PrologLineNumber = topScope.TokenReader.InputLine;
-
-			// Hide module classes from the editor pickers
-			classObj.ClassFlags |= EClassFlags.Native | EClassFlags.Hidden;
-			classObj.SuperIdentifier = [new UhtToken(UhtTokenType.Identifier, token.UngetPos, token.UngetLine, token.InputStartPos, token.InputLine, "UObject", false)];
-
-			topScope.AddFormattedCommentsAsTooltipMetaData();
-
-			topScope.TokenReader.Require("namespace");
-
-			// Handle any attributes
-			topScope.TokenReader.OptionalAttributes(false);
-
-			// Parse the namespace identifiers
-			int namespaceCount = ParseNamespaceIdentifiers(topScope);
-
-			// Remap the namespace
-			classObj.Namespace = topScope.HeaderParser.GetNamespace();
-
-			// Use the namespace name as the source name
-			string moduleName = classObj.Namespace.SourceName;
-			classObj.SourceName = $"VMODULE_{classObj.Module.Module.Name}_{classObj.Namespace.SourceName}";
-			classObj.ClassFlags |= EClassFlags.EditInlineNew | EClassFlags.HasInstancedReference;
-
-			specifiers.ParseDeferred();
-
-			{
-				using BorrowStringBuilder borrowBuilder = new(StringBuilderCache.Small);
-				borrowBuilder.StringBuilder.AppendVerseUEVNIPackageName(classObj);
-				classObj.Outer = classObj.Module.CreatePackage(borrowBuilder.StringBuilder.ToString());
-			}
-
-			UhtCompilerDirective compilerDirective = topScope.HeaderParser.GetCurrentCompositeCompilerDirective();
-			classObj.DefineScope |= compilerDirective.GetDefaultDefineScopes();
-			classObj.Outer?.AddChild(classObj);
-
-			topScope.HeaderParser.ParseStatements('{', '}', true);
-
-			// Mark all functions as static
-			foreach (UhtType child in classObj.Children)
-			{
-				if (child is UhtFunction functionObj)
-				{
-					functionObj.FunctionFlags |= EFunctionFlags.Static;
-					functionObj.FunctionExportFlags |= UhtFunctionExportFlags.CppStatic;
-				}
-			}
-
-			// Remove the namespaces from the active namespace
-			topScope.HeaderParser.PopNamespaces(namespaceCount);
-			return UhtParseResult.Handled;
-		}
 	}
 }
diff --git a/Engine/Source/Programs/Shared/EpicGames.UHT/Parsers/UhtPropertyParser.cs b/Engine/Source/Programs/Shared/EpicGames.UHT/Parsers/UhtPropertyParser.cs
--- a/Engine/Source/Programs/Shared/EpicGames.UHT/Parsers/UhtPropertyParser.cs
+++ b/Engine/Source/Programs/Shared/EpicGames.UHT/Parsers/UhtPropertyParser.cs
@@ -216,7 +216,8 @@ namespace EpicGames.UHT.Parsers
 		public static bool HasInitializer(this UhtPropertyMacroType propertyMacroType) =>
 			propertyMacroType switch
 			{
-				UhtPropertyMacroType.LayoutFieldInitialized => true,
+				UhtPropertyMacroType.LayoutFieldInitialized or
+				UhtPropertyMacroType.VerseField => true,
 				_ => false,
 			};
 
@@ -1000,8 +1001,7 @@ namespace EpicGames.UHT.Parsers
 			// Adjust property settings for property macros.
 			if (propertyMacroType.IsVerseField())
 			{
-				newProperty.PropertyExportFlags |= UhtPropertyExportFlags.VerseOnly;
-				newProperty.SourceName = UhtNames.VerseNotNativePrefix + newProperty.SourceName;
+				newProperty.PropertyExportFlags |= UhtPropertyExportFlags.VerseField;
 			}
 
 			// Adjust the name for verse
@@ -1070,6 +1070,7 @@ namespace EpicGames.UHT.Parsers
 		#region Keywords
 		[UhtKeyword(Extends = UhtTableNames.Class)]
 		[UhtKeyword(Extends = UhtTableNames.ScriptStruct)]
+		[UhtKeyword(Extends = UhtTableNames.VModule)]
 		[SuppressMessage("CodeQuality", "IDE0051:Remove unused private members", Justification = "Attribute accessed method")]
 		[SuppressMessage("Style", "IDE0060:Remove unused parameter", Justification = "Attribute accessed method")]
 		private static UhtParseResult UPROPERTYKeyword(UhtParsingScope topScope, UhtParsingScope actionScope, ref UhtToken token)
diff --git a/Engine/Source/Programs/Shared/EpicGames.UHT/Specifiers/UhtFieldSpecifiers.cs b/Engine/Source/Programs/Shared/EpicGames.UHT/Specifiers/UhtFieldSpecifiers.cs
--- a/Engine/Source/Programs/Shared/EpicGames.UHT/Specifiers/UhtFieldSpecifiers.cs
+++ b/Engine/Source/Programs/Shared/EpicGames.UHT/Specifiers/UhtFieldSpecifiers.cs
@@ -57,15 +57,17 @@ namespace EpicGames.UHT.Parsers
 
 			// Generate the engine name.  Verse forces a 'E' prefix on enum names...
 			string engineName = fieldObj is UhtEnum ? $"E{fieldObj.VerseName!}" : fieldObj.VerseName!;
-			if (String.IsNullOrEmpty(fieldObj.VerseModule))
+			if (!String.IsNullOrEmpty(fieldObj.VerseModule))
 			{
-				fieldObj.EngineName = engineName;
+				engineName = $"{fieldObj.VerseModule}_{engineName}";
 			}
-			else
+			if (fieldObj is UhtClass classObj && classObj.ClassExportFlags.HasAnyFlags(UhtClassExportFlags.VerseIsRootModule))
 			{
-				fieldObj.EngineName = $"{fieldObj.VerseModule}_{engineName}";
+				engineName += "_Root";
 			}
+			fieldObj.EngineName = engineName;
 
+			// Create the outer object for the field
 			using BorrowStringBuilder borrowBuilder = new(StringBuilderCache.Small);
 			borrowBuilder.StringBuilder.AppendVerseUEVNIPackageName(fieldObj);
 			fieldObj.Outer = fieldObj.Module.CreatePackage(borrowBuilder.StringBuilder.ToString());
@@ -171,6 +173,11 @@ namespace EpicGames.UHT.Parsers
 				{
 					return true;
 				}
+				else if (key.Equals("is_root", StringComparison.OrdinalIgnoreCase))
+				{
+					classObj.ClassExportFlags |= UhtClassExportFlags.VerseIsRootModule;
+					return true;
+				}
 				return false;
 			});
 		}
diff --git a/Engine/Source/Programs/Shared/EpicGames.UHT/Types/UhtClass.cs b/Engine/Source/Programs/Shared/EpicGames.UHT/Types/UhtClass.cs
--- a/Engine/Source/Programs/Shared/EpicGames.UHT/Types/UhtClass.cs
+++ b/Engine/Source/Programs/Shared/EpicGames.UHT/Types/UhtClass.cs
@@ -102,6 +102,11 @@ namespace EpicGames.UHT.Types
 		/// Class is marked as final
 		/// </summary>
 		Final = 1 << 13,
+
+		/// <summary>
+		/// Class is a verse root module.
+		/// </summary>
+		VerseIsRootModule = 1 << 14,
 	}
 
 	/// <summary>
@@ -976,6 +981,18 @@ namespace EpicGames.UHT.Types
 
 						case UhtClassType.VModule:
 							ClassWithin = superClass != null ? superClass.ClassWithin : Session.UObject;
+
+							foreach (UhtType child in Children)
+							{
+								if (child is UhtFunction function)
+								{
+									if (!function.FunctionFlags.HasAnyFlags(EFunctionFlags.Static) ||
+										!function.FunctionExportFlags.HasAnyFlags(UhtFunctionExportFlags.CppStatic))
+									{
+										this.LogError($"Verse modules class's '{SourceName}' function '{function.SourceName}' must be declared as static");
+									}
+								}
+							}
 							break;
 
 						default:
diff --git a/Engine/Source/Programs/Shared/EpicGames.UHT/Types/UhtProperty.cs b/Engine/Source/Programs/Shared/EpicGames.UHT/Types/UhtProperty.cs
--- a/Engine/Source/Programs/Shared/EpicGames.UHT/Types/UhtProperty.cs
+++ b/Engine/Source/Programs/Shared/EpicGames.UHT/Types/UhtProperty.cs
@@ -124,9 +124,9 @@ namespace EpicGames.UHT.Types
 		VerseDefaultValue = 1 << 19,
 
 		/// <summary>
-		/// If true, this property is not native
+		/// This property is wrapped by UE_VERSE_FIELD
 		/// </summary>
-		VerseOnly = 1 << 20,
+		VerseField = 1 << 20,
 	};
 
 	/// <summary>
@@ -1033,7 +1033,7 @@ namespace EpicGames.UHT.Types
 		/// <summary>
 		/// If true, then with the new VM, wrap the type in a rest value
 		/// </summary>
-		public virtual bool CodeGenWrapInRestValue => PropertyExportFlags.HasFlag(UhtPropertyExportFlags.VerseOnly);
+		public virtual bool CodeGenWrapInRestValue => PropertyExportFlags.HasFlag(UhtPropertyExportFlags.VerseField);
 
 		/// <summary>
 		/// Name of the code generation structure used when doing params code gen
diff --git a/Engine/Source/Runtime/CoreUObject/Public/UObject/ObjectMacros.h b/Engine/Source/Runtime/CoreUObject/Public/UObject/ObjectMacros.h
--- a/Engine/Source/Runtime/CoreUObject/Public/UObject/ObjectMacros.h
+++ b/Engine/Source/Runtime/CoreUObject/Public/UObject/ObjectMacros.h
@@ -822,9 +822,9 @@ struct FReferencerInformationList
 
 // Macro for defining verse only properties that are the given type with BPVM and another with the new VM
 #if WITH_VERSE_VM
-#define UE_VERSE_FIELD(T, Name) Verse::VRestValue __VerseOnly_##Name {0}
+#define UE_VERSE_FIELD(T, Name, Initializer) Verse::VRestValue Name {0}
 #else
-#define UE_VERSE_FIELD(T, Name) UE_REMOVE_OPTIONAL_PARENS(T) __VerseOnly_##Name;
+#define UE_VERSE_FIELD(T, Name, Initializer) UE_REMOVE_OPTIONAL_PARENS(T) Name UE_REMOVE_OPTIONAL_PARENS(Initializer)
 #endif
 
 // These are used for syntax highlighting and to allow autocomplete hints
