EpicGames.Perforce: Add support for p4 sizes -C (expected file size in the workspace).

#14586

#Synced-CL 52673681

diff --git a/Engine/Source/Programs/Shared/EpicGames.Perforce/PerforceConnection.cs b/Engine/Source/Programs/Shared/EpicGames.Perforce/PerforceConnection.cs
--- a/Engine/Source/Programs/Shared/EpicGames.Perforce/PerforceConnection.cs
+++ b/Engine/Source/Programs/Shared/EpicGames.Perforce/PerforceConnection.cs
@@ -3410,17 +3410,25 @@ namespace EpicGames.Perforce
 		/// <param name="fileSpecs">List of file specifications to query</param>
 		/// <param name="cancellationToken">Token used to cancel the operation</param>
 		/// <returns>List of response objects</returns>
+		/// <see href="https://help.perforce.com/helix-core/server-apps/cmdref/current/Content/CmdRef/p4_sizes.html"/>
 		public static async Task<PerforceResponseList<SizesRecord>> TrySizesAsync(this IPerforceConnection connection, SizesOptions options, int maxLines, FileSpecList fileSpecs, CancellationToken cancellationToken = default)
 		{
 			List<string> arguments = [];
-			if ((options & SizesOptions.AllRevisions) != 0)
+			// Put the sub-mode arguments first, per the p4 docs.
+			// nb. Not all options are valid together, but we let p4 validate this rather than duplicating the logic.
+			if ((options & SizesOptions.FileSizeAsExpectedInTheWorkspace) != 0)
 			{
-				arguments.Add("-a");
+				arguments.Add("-C");
 			}
 			if ((options & SizesOptions.LimitToArchiveDepots) != 0)
 			{
 				arguments.Add("-A");
 			}
+
+			if ((options & SizesOptions.AllRevisions) != 0)
+			{
+				arguments.Add("-a");
+			}
 			if ((options & SizesOptions.CalculateSum) != 0)
 			{
 				arguments.Add("-s");
diff --git a/Engine/Source/Programs/Shared/EpicGames.Perforce/Records/SizeOptions.cs b/Engine/Source/Programs/Shared/EpicGames.Perforce/Records/SizeOptions.cs
--- a/Engine/Source/Programs/Shared/EpicGames.Perforce/Records/SizeOptions.cs
+++ b/Engine/Source/Programs/Shared/EpicGames.Perforce/Records/SizeOptions.cs
@@ -7,6 +7,7 @@ namespace EpicGames.Perforce
 	/// <summary>
 	/// Options for the sizes command
 	/// </summary>
+	/// <see href="https://help.perforce.com/helix-core/server-apps/cmdref/current/Content/CmdRef/p4_sizes.html"/>
 	[Flags]
 	public enum SizesOptions
 	{
@@ -39,5 +40,11 @@ namespace EpicGames.Perforce
 		/// When calculating size information, exclude lazy copies.
 		/// </summary>
 		ExcludeLazyCopies = 16,
+
+		/// <summary>
+		/// Reports the file size as expected in the workspace, defaulting to the client's environment, accounting for line-endings,
+		/// charset conversion of unicode/utf8/utf16 filetypes (including the BOM for utf8 type files) and keyword expansion for +k type files.
+		/// </summary>
+		FileSizeAsExpectedInTheWorkspace = 32,
 	}
 }
