EpicGames.Perforce: Add the paths to the exception message in GetClientRelativePath when file is not under workspace root.

#14585

#Synced-CL 52673681

diff --git a/Engine/Source/Programs/Shared/EpicGames.Perforce/PerforceUtils.cs b/Engine/Source/Programs/Shared/EpicGames.Perforce/PerforceUtils.cs
--- a/Engine/Source/Programs/Shared/EpicGames.Perforce/PerforceUtils.cs
+++ b/Engine/Source/Programs/Shared/EpicGames.Perforce/PerforceUtils.cs
@@ -214,7 +214,7 @@ namespace EpicGames.Perforce
 		{
 			if (!workspaceFile.IsUnderDirectory(workspaceRoot))
 			{
-				throw new ArgumentException("File is not under workspace root", nameof(workspaceFile));
+				throw new ArgumentException($"File '{workspaceFile}' is not under workspace root '{workspaceRoot}'", nameof(workspaceFile));
 			}
 
 			return workspaceFile.MakeRelativeTo(workspaceRoot).Replace(Path.DirectorySeparatorChar, '/');
