Horde: Add agent/lease query and tool deployment methods to HordeHttpClient

diff --git a/Engine/Source/Programs/Shared/EpicGames.Horde/HordeHttpClient.cs b/Engine/Source/Programs/Shared/EpicGames.Horde/HordeHttpClient.cs
--- a/Engine/Source/Programs/Shared/EpicGames.Horde/HordeHttpClient.cs
+++ b/Engine/Source/Programs/Shared/EpicGames.Horde/HordeHttpClient.cs
@@ -14,6 +14,9 @@ using System.Text.RegularExpressions;
 using System.Threading;
 using System.Threading.Tasks;
 using EpicGames.Core;
+using EpicGames.Horde.Agents;
+using EpicGames.Horde.Agents.Pools;
+using EpicGames.Horde.Agents.Sessions;
 using EpicGames.Horde.Agents.Telemetry;
 using EpicGames.Horde.Artifacts;
 using EpicGames.Horde.Commits;
@@ -124,6 +127,85 @@ namespace EpicGames.Horde
 
 		#endregion
 
+		#region Agents
+
+		/// <summary>
+		/// Gets a list of agents matching the specified criteria
+		/// </summary>
+		/// <param name="poolId">Optional pool to filter by</param>
+		/// <param name="includeDeleted">Whether to include deleted agents</param>
+		/// <param name="cancellationToken">Cancellation token for the operation</param>
+		public Task<List<GetAgentResponse>> GetAgentsAsync(PoolId? poolId = null, bool includeDeleted = false, CancellationToken cancellationToken = default)
+		{
+			QueryStringBuilder queryParams = new();
+			if (poolId != null)
+			{
+				queryParams.Add("poolId", poolId.Value.ToString());
+			}
+			queryParams.Add("includeDeleted", includeDeleted ? "true" : "false");
+			return GetAsync<List<GetAgentResponse>>(_httpClient, $"api/v1/agents?{queryParams}", cancellationToken);
+		}
+
+		/// <summary>
+		/// Gets information about a specific agent
+		/// </summary>
+		/// <param name="agentId">The agent identifier</param>
+		/// <param name="cancellationToken">Cancellation token for the operation</param>
+		public Task<GetAgentResponse> GetAgentAsync(AgentId agentId, CancellationToken cancellationToken = default)
+		{
+			return GetAsync<GetAgentResponse>(_httpClient, $"api/v1/agents/{agentId}", cancellationToken);
+		}
+
+		#endregion
+
+		#region Leases
+
+		/// <summary>
+		/// Gets a list of leases matching the specified criteria
+		/// </summary>
+		/// <param name="agentId">Optional agent to filter by</param>
+		/// <param name="sessionId">Optional session to filter by</param>
+		/// <param name="startTime">Start of the time window</param>
+		/// <param name="finishTime">End of the time window</param>
+		/// <param name="minFinishTime">Minimum finish time filter</param>
+		/// <param name="maxFinishTime">Maximum finish time filter</param>
+		/// <param name="index">Starting index for pagination</param>
+		/// <param name="count">Maximum number of results</param>
+		/// <param name="cancellationToken">Cancellation token for the operation</param>
+		public Task<List<GetAgentLeaseResponse>> GetLeasesAsync(AgentId? agentId = null, SessionId? sessionId = null, DateTime? startTime = null, DateTime? finishTime = null, DateTime? minFinishTime = null, DateTime? maxFinishTime = null, int index = 0, int count = 1000, CancellationToken cancellationToken = default)
+		{
+			QueryStringBuilder queryParams = new();
+			if (agentId != null)
+			{
+				queryParams.Add("agentId", agentId.Value.ToString());
+			}
+			if (sessionId != null)
+			{
+				queryParams.Add("sessionId", sessionId.Value.ToString());
+			}
+			if (startTime != null)
+			{
+				queryParams.Add("startTime", startTime.Value.ToString("o"));
+			}
+			if (finishTime != null)
+			{
+				queryParams.Add("finishTime", finishTime.Value.ToString("o"));
+			}
+			if (minFinishTime != null)
+			{
+				queryParams.Add("minFinishTime", minFinishTime.Value.ToString("o"));
+			}
+			if (maxFinishTime != null)
+			{
+				queryParams.Add("maxFinishTime", maxFinishTime.Value.ToString("o"));
+			}
+			queryParams.Add("index", index.ToString());
+			queryParams.Add("count", count.ToString());
+			return GetAsync<List<GetAgentLeaseResponse>>(_httpClient, $"api/v1/leases?{queryParams}", cancellationToken);
+		}
+
+		#endregion
+
 		#region Artifacts
 
 		/// <summary>
@@ -564,6 +646,61 @@ namespace EpicGames.Horde
 			return response.Id;
 		}
 
+		/// <summary>
+		/// Updates a tool deployment (e.g. changing its state)
+		/// </summary>
+		/// <param name="id">Id for the tool</param>
+		/// <param name="deploymentId">Id of the deployment to update</param>
+		/// <param name="request">The update request</param>
+		/// <param name="cancellationToken">Cancellation token for the operation</param>
+		public Task UpdateToolDeploymentAsync(ToolId id, ToolDeploymentId deploymentId, UpdateDeploymentRequest request, CancellationToken cancellationToken = default)
+		{
+			return PatchAsync(_httpClient, $"api/v1/tools/{id}/deployments/{deploymentId}", request, cancellationToken);
+		}
+
+		/// <summary>
+		/// Creates a new tool deployment by uploading a zip file via multipart form data (v1 API)
+		/// </summary>
+		/// <param name="id">Id for the tool</param>
+		/// <param name="version">Version string for the new deployment</param>
+		/// <param name="content">Stream containing the zip file</param>
+		/// <param name="fileName">Name of the zip file</param>
+		/// <param name="duration">Duration over which to deploy the tool</param>
+		/// <param name="createPaused">Whether to create the deployment in a paused state</param>
+		/// <param name="cancellationToken">Cancellation token for the operation</param>
+		public async Task<ToolDeploymentId> CreateToolDeploymentFromStreamAsync(ToolId id, string version, Stream content, string fileName, TimeSpan? duration = null, bool createPaused = false, CancellationToken cancellationToken = default)
+		{
+#pragma warning disable IDE0028 // Collection initialization can be simplified - conditional Add calls prevent collection expression
+			using MultipartFormDataContent formContent = new();
+#pragma warning restore IDE0028
+
+#pragma warning disable CA2000 // Disposed by MultipartFormDataContent
+			formContent.Add(new StreamContent(content), "file", fileName);
+			formContent.Add(new StringContent(version), "Version");
+			if (duration != null)
+			{
+				formContent.Add(new StringContent(duration.Value.TotalSeconds.ToString()), "Duration");
+			}
+			if (createPaused)
+			{
+				formContent.Add(new StringContent("true"), "CreatePaused");
+			}
+#pragma warning restore CA2000
+
+			using HttpResponseMessage response = await _httpClient.PostAsync($"api/v1/tools/{id}/deployments", formContent, cancellationToken);
+			if (!response.IsSuccessStatusCode)
+			{
+				string body = await response.Content.ReadAsStringAsync(cancellationToken);
+				throw new HttpRequestException(
+					$"Failed to create tool deployment for '{id}' (status: {(int)response.StatusCode} {response.StatusCode}, body: {body})",
+					inner: null,
+					statusCode: response.StatusCode);
+			}
+
+			CreateToolDeploymentResponse? result = await response.Content.ReadFromJsonAsync<CreateToolDeploymentResponse>(JsonSerializerOptions, cancellationToken);
+			return result?.Id ?? throw new InvalidOperationException("Expected non-null response from tool deployment upload");
+		}
+
 		#endregion
 
 		#region Jobs
diff --git a/Engine/Source/Programs/Shared/EpicGames.Horde/HordeHttpRequest.cs b/Engine/Source/Programs/Shared/EpicGames.Horde/HordeHttpRequest.cs
--- a/Engine/Source/Programs/Shared/EpicGames.Horde/HordeHttpRequest.cs
+++ b/Engine/Source/Programs/Shared/EpicGames.Horde/HordeHttpRequest.cs
@@ -151,5 +151,19 @@ namespace EpicGames.Horde
 				return responseValue ?? throw new InvalidCastException($"Expected non-null response from PUT to {relativePath}");
 			}
 		}
+
+		/// <summary>
+		/// Patches a resource at an HTTP endpoint with a JSON object
+		/// </summary>
+		/// <typeparam name="TRequest">The object type to send</typeparam>
+		/// <param name="httpClient">Http client instance</param>
+		/// <param name="relativePath">The url to patch</param>
+		/// <param name="request">The object to send</param>
+		/// <param name="cancellationToken">Cancels the request</param>
+		public static async Task PatchAsync<TRequest>(HttpClient httpClient, string relativePath, TRequest request, CancellationToken cancellationToken = default)
+		{
+			using HttpResponseMessage response = await httpClient.PatchAsJsonAsync(relativePath, request, s_jsonSerializerOptions, cancellationToken);
+			response.EnsureSuccessStatusCode();
+		}
 	}
 }
